Deprecate GetOrConstructReference functions to discourage node developers from ignoring default literals.

diff --git a/Engine/Plugins/Runtime/AudioModulation/Source/AudioModulation/Private/MetaSoundNodes/GetModulatorValueNode.cpp b/Engine/Plugins/Runtime/AudioModulation/Source/AudioModulation/Private/MetaSoundNodes/GetModulatorValueNode.cpp
--- a/Engine/Plugins/Runtime/AudioModulation/Source/AudioModulation/Private/MetaSoundNodes/GetModulatorValueNode.cpp
+++ b/Engine/Plugins/Runtime/AudioModulation/Source/AudioModulation/Private/MetaSoundNodes/GetModulatorValueNode.cpp
@@ -69,7 +69,7 @@ namespace AudioModulation
 
 			if (InParams.Environment.Contains<Audio::FDeviceId>(Metasound::Frontend::SourceInterface::Environment::DeviceID))
 			{
-				FSoundModulatorAssetReadRef ModulatorReadRef = InputData.GetOrConstructDataReadReference<FSoundModulatorAsset>("Modulator");
+				FSoundModulatorAssetReadRef ModulatorReadRef = InputData.GetOrCreateDefaultDataReadReference<FSoundModulatorAsset>("Modulator", InParams.OperatorSettings);
 				FBoolReadRef NormalizedReadRef = InputData.GetOrCreateDefaultDataReadReference<bool>("Normalized", InParams.OperatorSettings);
 
 				return MakeUnique<FGetModulatorValueNodeOperator>(InParams, ModulatorReadRef, NormalizedReadRef);
diff --git a/Engine/Plugins/Runtime/AudioModulation/Source/AudioModulation/Private/MetaSoundNodes/MixModulatorsNode.cpp b/Engine/Plugins/Runtime/AudioModulation/Source/AudioModulation/Private/MetaSoundNodes/MixModulatorsNode.cpp
--- a/Engine/Plugins/Runtime/AudioModulation/Source/AudioModulation/Private/MetaSoundNodes/MixModulatorsNode.cpp
+++ b/Engine/Plugins/Runtime/AudioModulation/Source/AudioModulation/Private/MetaSoundNodes/MixModulatorsNode.cpp
@@ -71,9 +71,9 @@ namespace AudioModulation
 				
 			if (InParams.Environment.Contains<Audio::FDeviceId>(Metasound::Frontend::SourceInterface::Environment::DeviceID))
 			{
-				FSoundModulatorAssetReadRef Modulator1ReadRef = InputData.GetOrConstructDataReadReference<FSoundModulatorAsset>("In1");
-				FSoundModulatorAssetReadRef Modulator2ReadRef = InputData.GetOrConstructDataReadReference<FSoundModulatorAsset>("In2");
-				FSoundModulationParameterAssetReadRef ParameterReadRef = InputData.GetOrConstructDataReadReference<FSoundModulationParameterAsset>("MixParameter");
+				FSoundModulatorAssetReadRef Modulator1ReadRef = InputData.GetOrCreateDefaultDataReadReference<FSoundModulatorAsset>("In1", InParams.OperatorSettings);
+				FSoundModulatorAssetReadRef Modulator2ReadRef = InputData.GetOrCreateDefaultDataReadReference<FSoundModulatorAsset>("In2", InParams.OperatorSettings);
+				FSoundModulationParameterAssetReadRef ParameterReadRef = InputData.GetOrCreateDefaultDataReadReference<FSoundModulationParameterAsset>("MixParameter", InParams.OperatorSettings);
 				FBoolReadRef NormalizedReadRef = InputData.GetOrCreateDefaultDataReadReference<bool>("Normalized", InParams.OperatorSettings);
 
 				return MakeUnique<FMixModulatorsNodeOperator>(InParams, Modulator1ReadRef, Modulator2ReadRef, ParameterReadRef, NormalizedReadRef);
diff --git a/Engine/Plugins/Runtime/Harmonix/Source/HarmonixMetasound/Private/HarmonixMetasound/Nodes/DelayNode.cpp b/Engine/Plugins/Runtime/Harmonix/Source/HarmonixMetasound/Private/HarmonixMetasound/Nodes/DelayNode.cpp
--- a/Engine/Plugins/Runtime/Harmonix/Source/HarmonixMetasound/Private/HarmonixMetasound/Nodes/DelayNode.cpp
+++ b/Engine/Plugins/Runtime/Harmonix/Source/HarmonixMetasound/Private/HarmonixMetasound/Nodes/DelayNode.cpp
@@ -199,8 +199,8 @@ namespace HarmonixMetasound::DelayNode
 			
 			FInputs InputReferences
 			{
-				InputData.GetOrConstructDataReadReference<FAudioBuffer>(Inputs::AudioLeftName, InParams.OperatorSettings),
-				InputData.GetOrConstructDataReadReference<FAudioBuffer>(Inputs::AudioRightName, InParams.OperatorSettings),
+				InputData.GetOrCreateDefaultDataReadReference<FAudioBuffer>(Inputs::AudioLeftName, InParams.OperatorSettings),
+				InputData.GetOrCreateDefaultDataReadReference<FAudioBuffer>(Inputs::AudioRightName, InParams.OperatorSettings),
 				{},
 				InputData.GetOrCreateDefaultDataReadReference<FEnumTimeSyncOption>(
 					Inputs::DelayTimeTypeName,
diff --git a/Engine/Plugins/Runtime/Harmonix/Source/HarmonixMetasound/Private/HarmonixMetasound/Nodes/DistortionNode.cpp b/Engine/Plugins/Runtime/Harmonix/Source/HarmonixMetasound/Private/HarmonixMetasound/Nodes/DistortionNode.cpp
--- a/Engine/Plugins/Runtime/Harmonix/Source/HarmonixMetasound/Private/HarmonixMetasound/Nodes/DistortionNode.cpp
+++ b/Engine/Plugins/Runtime/Harmonix/Source/HarmonixMetasound/Private/HarmonixMetasound/Nodes/DistortionNode.cpp
@@ -252,26 +252,26 @@ namespace HarmonixMetasound
 			
 			const FInputVertexInterfaceData& InputData = InParams.InputData;
 
-			FAudioBufferReadRef InAudio = InputData.GetOrConstructDataReadReference<Metasound::FAudioBuffer>(
+			FAudioBufferReadRef InAudio = InputData.GetOrCreateDefaultDataReadReference<Metasound::FAudioBuffer>(
 				METASOUND_GET_PARAM_NAME(Inputs::AudioIn), InParams.OperatorSettings);
 			FFloatReadRef InInputGain = InputData.GetOrCreateDefaultDataReadReference<float>(METASOUND_GET_PARAM_NAME(Inputs::InputGain), InParams.OperatorSettings);
 			FFloatReadRef InOutputGain = InputData.GetOrCreateDefaultDataReadReference<float>(METASOUND_GET_PARAM_NAME(Inputs::OutputGain), InParams.OperatorSettings);
 			FFloatReadRef InWetGain = InputData.GetOrCreateDefaultDataReadReference<float>(METASOUND_GET_PARAM_NAME(Inputs::WetGain), InParams.OperatorSettings);
 			FFloatReadRef InDryGain = InputData.GetOrCreateDefaultDataReadReference<float>(METASOUND_GET_PARAM_NAME(Inputs::DryGain), InParams.OperatorSettings);
 			FFloatReadRef InDCAdjust = InputData.GetOrCreateDefaultDataReadReference<float>(METASOUND_GET_PARAM_NAME(Inputs::DCAdjust), InParams.OperatorSettings);
-			FDistortionTypeReadRef InDistortionType = InputData.GetOrConstructDataReadReference<FEnumDistortionType>(
-				METASOUND_GET_PARAM_NAME(Inputs::DistortionType));
+			FDistortionTypeReadRef InDistortionType = InputData.GetOrCreateDefaultDataReadReference<FEnumDistortionType>(
+				METASOUND_GET_PARAM_NAME(Inputs::DistortionType), InParams.OperatorSettings);
 			FBoolReadRef InOversample = InputData.GetOrCreateDefaultDataReadReference<bool>(METASOUND_GET_PARAM_NAME(Inputs::Oversample), InParams.OperatorSettings);
 
 			auto MakeFilterSettings = [&InputData, &InParams](int Index)-> FFilterSettings
 			{			
 				return FFilterSettings(
 					InputData.GetOrCreateDefaultDataReadReference<bool>(METASOUND_GET_PARAM_NAME_WITH_INDEX(Inputs::FilterEnabled, Index), InParams.OperatorSettings),
-					InputData.GetOrConstructDataReadReference<FEnumHmxBiquadFilterType>(METASOUND_GET_PARAM_NAME_WITH_INDEX(Inputs::FilterType, Index)),
+					InputData.GetOrCreateDefaultDataReadReference<FEnumHmxBiquadFilterType>(METASOUND_GET_PARAM_NAME_WITH_INDEX(Inputs::FilterType, Index), InParams.OperatorSettings),
 					InputData.GetOrCreateDefaultDataReadReference<float>(METASOUND_GET_PARAM_NAME_WITH_INDEX(Inputs::Freq, Index), InParams.OperatorSettings),
 					InputData.GetOrCreateDefaultDataReadReference<float>(METASOUND_GET_PARAM_NAME_WITH_INDEX(Inputs::Q, Index), InParams.OperatorSettings),
 					InputData.GetOrCreateDefaultDataReadReference<bool>(METASOUND_GET_PARAM_NAME_WITH_INDEX(Inputs::PreClip, Index), InParams.OperatorSettings),
-					InputData.GetOrConstructDataReadReference<FEnumDistortionFilterPasses>(METASOUND_GET_PARAM_NAME_WITH_INDEX(Inputs::Passes, Index))
+					InputData.GetOrCreateDefaultDataReadReference<FEnumDistortionFilterPasses>(METASOUND_GET_PARAM_NAME_WITH_INDEX(Inputs::Passes, Index), InParams.OperatorSettings)
 				);
 			};
 				
@@ -429,4 +429,4 @@ namespace HarmonixMetasound
 	METASOUND_REGISTER_NODE(FDistortionNode);
 }
 
-#undef LOCTEXT_NAMESPACE
\ No newline at end of file
+#undef LOCTEXT_NAMESPACE
diff --git a/Engine/Plugins/Runtime/Harmonix/Source/HarmonixMetasound/Private/HarmonixMetasound/Nodes/DjFilterNode.cpp b/Engine/Plugins/Runtime/Harmonix/Source/HarmonixMetasound/Private/HarmonixMetasound/Nodes/DjFilterNode.cpp
--- a/Engine/Plugins/Runtime/Harmonix/Source/HarmonixMetasound/Private/HarmonixMetasound/Nodes/DjFilterNode.cpp
+++ b/Engine/Plugins/Runtime/Harmonix/Source/HarmonixMetasound/Private/HarmonixMetasound/Nodes/DjFilterNode.cpp
@@ -124,7 +124,7 @@ namespace HarmonixMetasound::Nodes::DjFilter
 			
 			FInputs Inputs
 			{
-				InputData.GetOrConstructDataReadReference<Metasound::FAudioBuffer>(
+				InputData.GetOrCreateDefaultDataReadReference<Metasound::FAudioBuffer>(
 					Inputs::AudioMonoName,
 					OperatorSettings),
 				InputData.GetOrCreateDefaultDataReadReference<float>(
diff --git a/Engine/Plugins/Runtime/Harmonix/Source/HarmonixMetasound/Private/HarmonixMetasound/Nodes/FFTAnalyzerNode.cpp b/Engine/Plugins/Runtime/Harmonix/Source/HarmonixMetasound/Private/HarmonixMetasound/Nodes/FFTAnalyzerNode.cpp
--- a/Engine/Plugins/Runtime/Harmonix/Source/HarmonixMetasound/Private/HarmonixMetasound/Nodes/FFTAnalyzerNode.cpp
+++ b/Engine/Plugins/Runtime/Harmonix/Source/HarmonixMetasound/Private/HarmonixMetasound/Nodes/FFTAnalyzerNode.cpp
@@ -138,7 +138,7 @@ namespace HarmonixMetasound
 			{
 				InputData.GetOrCreateDefaultDataReadReference<bool>(
 					Inputs::EnableName, InParams.OperatorSettings),
-				InputData.GetOrConstructDataReadReference<FAudioBuffer>(
+				InputData.GetOrCreateDefaultDataReadReference<FAudioBuffer>(
 					METASOUND_GET_PARAM_NAME(Inputs::Audio), InParams.OperatorSettings),
 				InputData.GetOrCreateDefaultDataReadReference<int32>(
 					METASOUND_GET_PARAM_NAME(Inputs::FFTSize), InParams.OperatorSettings),
@@ -269,4 +269,4 @@ namespace HarmonixMetasound
 	METASOUND_REGISTER_NODE(FFFTAnalyzerNode);
 }
 
-#undef LOCTEXT_NAMESPACE
\ No newline at end of file
+#undef LOCTEXT_NAMESPACE
diff --git a/Engine/Plugins/Runtime/Harmonix/Source/HarmonixMetasound/Private/HarmonixMetasound/Nodes/FusionSamplerNode.cpp b/Engine/Plugins/Runtime/Harmonix/Source/HarmonixMetasound/Private/HarmonixMetasound/Nodes/FusionSamplerNode.cpp
--- a/Engine/Plugins/Runtime/Harmonix/Source/HarmonixMetasound/Private/HarmonixMetasound/Nodes/FusionSamplerNode.cpp
+++ b/Engine/Plugins/Runtime/Harmonix/Source/HarmonixMetasound/Private/HarmonixMetasound/Nodes/FusionSamplerNode.cpp
@@ -83,7 +83,7 @@ namespace HarmonixMetasound
 				: Params(InParams)
 				, InSettings(& InParams.OperatorSettings)
 				, InEnabled(InParams.InputData.GetOrCreateDefaultDataReadReference<bool>(METASOUND_GET_PARAM_NAME(CommonPinNames::Inputs::Enable), InParams.OperatorSettings))
-				, InMidiStream(InParams.InputData.GetOrConstructDataReadReference<FMidiStream>(METASOUND_GET_PARAM_NAME(CommonPinNames::Inputs::MidiStream)))
+				, InMidiStream(InParams.InputData.GetOrCreateDefaultDataReadReference<FMidiStream>(METASOUND_GET_PARAM_NAME(CommonPinNames::Inputs::MidiStream), InParams.OperatorSettings))
 				, InTrackNumber(InParams.InputData.GetOrCreateDefaultDataReadReference<int32>(METASOUND_GET_PARAM_NAME(FusionSamplerNodePinNames::MidiTrackNumber), InParams.OperatorSettings))
 				, InTransposition(InParams.InputData.GetOrCreateDefaultDataReadReference<int32>(METASOUND_GET_PARAM_NAME(CommonPinNames::Inputs::Transposition), InParams.OperatorSettings))
 				, InLfo0RateOverride(InParams.InputData.GetOrCreateDefaultDataReadReference<float>(METASOUND_GET_PARAM_NAME(FusionSamplerNodePinNames::Lfo0FrequencyOverride), InParams.OperatorSettings))
@@ -91,7 +91,7 @@ namespace HarmonixMetasound
 				, InLfo1RateOverride(InParams.InputData.GetOrCreateDefaultDataReadReference<float>(METASOUND_GET_PARAM_NAME(FusionSamplerNodePinNames::Lfo1FrequencyOverride), InParams.OperatorSettings))
 				, InLfo1DepthOverride(InParams.InputData.GetOrCreateDefaultDataReadReference<float>(METASOUND_GET_PARAM_NAME(FusionSamplerNodePinNames::Lfo1DepthOverride), InParams.OperatorSettings))
 				, InFineTuneCents(InParams.InputData.GetOrCreateDefaultDataReadReference<float>(METASOUND_GET_PARAM_NAME(FusionSamplerNodePinNames::FineTuneCents), InParams.OperatorSettings))
-				, InPatch(InParams.InputData.GetOrConstructDataReadReference<FFusionPatchAsset>(METASOUND_GET_PARAM_NAME(CommonPinNames::Inputs::SynthPatch)))
+				, InPatch(InParams.InputData.GetOrCreateDefaultDataReadReference<FFusionPatchAsset>(METASOUND_GET_PARAM_NAME(CommonPinNames::Inputs::SynthPatch), InParams.OperatorSettings))
 				, InClockSpeedToPitch(InParams.InputData.GetOrCreateDefaultDataReadReference<bool>(METASOUND_GET_PARAM_NAME(CommonPinNames::Inputs::ClockSpeedToPitch), InParams.OperatorSettings))
 				, InMTRenderingEnable(InParams.InputData.GetOrCreateDefaultValue<bool>(METASOUND_GET_PARAM_NAME(FusionSamplerNodePinNames::EnableMTFusion), InParams.OperatorSettings))
 				, InVoicePoolName(InParams.InputData.GetOrCreateDefaultValue<FString>(METASOUND_GET_PARAM_NAME(FusionSamplerNodePinNames::VoicePoolName), InParams.OperatorSettings))
@@ -726,4 +726,4 @@ namespace HarmonixMetasound
 	}
 }
 
-#undef LOCTEXT_NAMESPACE // "HarmonixMetaSound"
\ No newline at end of file
+#undef LOCTEXT_NAMESPACE // "HarmonixMetaSound"
diff --git a/Engine/Plugins/Runtime/Harmonix/Source/HarmonixMetasound/Private/HarmonixMetasound/Nodes/MetronomeNode.cpp b/Engine/Plugins/Runtime/Harmonix/Source/HarmonixMetasound/Private/HarmonixMetasound/Nodes/MetronomeNode.cpp
--- a/Engine/Plugins/Runtime/Harmonix/Source/HarmonixMetasound/Private/HarmonixMetasound/Nodes/MetronomeNode.cpp
+++ b/Engine/Plugins/Runtime/Harmonix/Source/HarmonixMetasound/Private/HarmonixMetasound/Nodes/MetronomeNode.cpp
@@ -130,7 +130,7 @@ namespace HarmonixMetasound::Nodes::MetronomeNode
 		const FInputVertexInterfaceData& InputData = InParams.InputData;
 		const FOperatorSettings& Settings = InParams.OperatorSettings;
 
-		FMusicTransportEventStreamReadRef InTransport = InputData.GetOrConstructDataReadReference<FMusicTransportEventStream>(METASOUND_GET_PARAM_NAME(Inputs::Transport), Settings);
+		FMusicTransportEventStreamReadRef InTransport = InputData.GetOrCreateDefaultDataReadReference<FMusicTransportEventStream>(METASOUND_GET_PARAM_NAME(Inputs::Transport), Settings);
 		bool InLoop = InputData.GetOrCreateDefaultValue<bool>(METASOUND_GET_PARAM_NAME(Inputs::Loop), Settings);
 		int32 InLoopLengthBars = InputData.GetOrCreateDefaultValue<int32>(METASOUND_GET_PARAM_NAME(Inputs::LoopLengthBars), Settings);
 		FInt32ReadRef InTimeSigNumerator   = InputData.GetOrCreateDefaultDataReadReference<int32>(METASOUND_GET_PARAM_NAME(Inputs::TimeSigNumerator), Settings);
diff --git a/Engine/Plugins/Runtime/Harmonix/Source/HarmonixMetasound/Private/HarmonixMetasound/Nodes/MidiChannelFilterNode.cpp b/Engine/Plugins/Runtime/Harmonix/Source/HarmonixMetasound/Private/HarmonixMetasound/Nodes/MidiChannelFilterNode.cpp
--- a/Engine/Plugins/Runtime/Harmonix/Source/HarmonixMetasound/Private/HarmonixMetasound/Nodes/MidiChannelFilterNode.cpp
+++ b/Engine/Plugins/Runtime/Harmonix/Source/HarmonixMetasound/Private/HarmonixMetasound/Nodes/MidiChannelFilterNode.cpp
@@ -105,7 +105,7 @@ namespace HarmonixMetasound::Nodes::MidiChannelFilter
 			FInputs Inputs
 			{
 				InputData.GetOrCreateDefaultDataReadReference<bool>(Inputs::EnableName, OperatorSettings),
-				InputData.GetOrConstructDataReadReference<FMidiStream>(Inputs::MidiStreamName),
+				InputData.GetOrCreateDefaultDataReadReference<FMidiStream>(Inputs::MidiStreamName, OperatorSettings),
 				InputData.GetOrCreateDefaultDataReadReference<int32>(Inputs::ChannelName, OperatorSettings)
 			};
 
diff --git a/Engine/Plugins/Runtime/Harmonix/Source/HarmonixMetasound/Private/HarmonixMetasound/Nodes/MidiClockOffsetNode.cpp b/Engine/Plugins/Runtime/Harmonix/Source/HarmonixMetasound/Private/HarmonixMetasound/Nodes/MidiClockOffsetNode.cpp
--- a/Engine/Plugins/Runtime/Harmonix/Source/HarmonixMetasound/Private/HarmonixMetasound/Nodes/MidiClockOffsetNode.cpp
+++ b/Engine/Plugins/Runtime/Harmonix/Source/HarmonixMetasound/Private/HarmonixMetasound/Nodes/MidiClockOffsetNode.cpp
@@ -76,9 +76,9 @@ namespace HarmonixMetasound::Nodes::MidiClockOffset
 			static const FVertexInterface Interface(
 				FInputVertexInterface(
 					TInputDataVertex<FMidiClock>(METASOUND_GET_PARAM_NAME_AND_METADATA(Inputs::MidiClock)),
-					TInputDataVertex<int32>(METASOUND_GET_PARAM_NAME_AND_METADATA(Inputs::OffsetBars)),
-					TInputDataVertex<float>(METASOUND_GET_PARAM_NAME_AND_METADATA(Inputs::OffsetBeats)),
-					TInputDataVertex<float>(METASOUND_GET_PARAM_NAME_AND_METADATA(Inputs::OffsetMs))
+					TInputDataVertex<int32>(METASOUND_GET_PARAM_NAME_AND_METADATA(Inputs::OffsetBars), 0),
+					TInputDataVertex<float>(METASOUND_GET_PARAM_NAME_AND_METADATA(Inputs::OffsetBeats), 0.0f),
+					TInputDataVertex<float>(METASOUND_GET_PARAM_NAME_AND_METADATA(Inputs::OffsetMs), 0.0f)
 				),
 				FOutputVertexInterface(
 					TOutputDataVertex<FMidiClock>(METASOUND_GET_PARAM_NAME_AND_METADATA(Outputs::MidiClock))
@@ -92,10 +92,10 @@ namespace HarmonixMetasound::Nodes::MidiClockOffset
 		{
 			const FOperatorSettings& Settings = InParams.OperatorSettings;
 			const FInputVertexInterfaceData& InputData = InParams.InputData;
-			FMidiClockReadRef InMidiClock = InputData.GetOrConstructDataReadReference<FMidiClock>(METASOUND_GET_PARAM_NAME(Inputs::MidiClock), Settings);
-			FInt32ReadRef InOffsetBars = InputData.GetOrConstructDataReadReference<int32>(METASOUND_GET_PARAM_NAME(Inputs::OffsetBars), 0);
-			FFloatReadRef InOffsetBeats = InputData.GetOrConstructDataReadReference<float>(METASOUND_GET_PARAM_NAME(Inputs::OffsetBeats), 0.0f);
-			FFloatReadRef InOffsetMs = InputData.GetOrConstructDataReadReference<float>(METASOUND_GET_PARAM_NAME(Inputs::OffsetMs), 0.0f);
+			FMidiClockReadRef InMidiClock = InputData.GetOrCreateDefaultDataReadReference<FMidiClock>(METASOUND_GET_PARAM_NAME(Inputs::MidiClock), Settings);
+			FInt32ReadRef InOffsetBars = InputData.GetOrCreateDefaultDataReadReference<int32>(METASOUND_GET_PARAM_NAME(Inputs::OffsetBars), Settings);
+			FFloatReadRef InOffsetBeats = InputData.GetOrCreateDefaultDataReadReference<float>(METASOUND_GET_PARAM_NAME(Inputs::OffsetBeats), Settings);
+			FFloatReadRef InOffsetMs = InputData.GetOrCreateDefaultDataReadReference<float>(METASOUND_GET_PARAM_NAME(Inputs::OffsetMs), Settings);
 			return MakeUnique<FMidiClockOffsetOperator>(InParams.OperatorSettings, InMidiClock, InOffsetBars, InOffsetBeats, InOffsetMs);	
 		}
 
@@ -276,4 +276,4 @@ namespace HarmonixMetasound::Nodes::MidiClockOffset
 	METASOUND_REGISTER_NODE(FMidiClockOffsetNode)
 }
 
-#undef LOCTEXT_NAMESPACE // "HarmonixMetaSound"
\ No newline at end of file
+#undef LOCTEXT_NAMESPACE // "HarmonixMetaSound"
diff --git a/Engine/Plugins/Runtime/Harmonix/Source/HarmonixMetasound/Private/HarmonixMetasound/Nodes/MidiClockStateNode.cpp b/Engine/Plugins/Runtime/Harmonix/Source/HarmonixMetasound/Private/HarmonixMetasound/Nodes/MidiClockStateNode.cpp
--- a/Engine/Plugins/Runtime/Harmonix/Source/HarmonixMetasound/Private/HarmonixMetasound/Nodes/MidiClockStateNode.cpp
+++ b/Engine/Plugins/Runtime/Harmonix/Source/HarmonixMetasound/Private/HarmonixMetasound/Nodes/MidiClockStateNode.cpp
@@ -126,7 +126,7 @@ namespace HarmonixMetasound
 		const FMidiClockStateNode& TheNode = static_cast<const FMidiClockStateNode&>(InParams.Node);
 
 		const FInputVertexInterfaceData& InputData = InParams.InputData;
-		FMidiClockReadRef InMidiClock = InputData.GetOrConstructDataReadReference<FMidiClock>(METASOUND_GET_PARAM_NAME(Inputs::MidiClock), InParams.OperatorSettings);
+		FMidiClockReadRef InMidiClock = InputData.GetOrCreateDefaultDataReadReference<FMidiClock>(METASOUND_GET_PARAM_NAME(Inputs::MidiClock), InParams.OperatorSettings);
 		FBoolReadRef InEnabled        = InputData.GetOrCreateDefaultDataReadReference<bool>(METASOUND_GET_PARAM_NAME(Inputs::Enable), InParams.OperatorSettings);
 
 		return MakeUnique<FMidiClockStateOperator>(InParams, InEnabled, InMidiClock);
diff --git a/Engine/Plugins/Runtime/Harmonix/Source/HarmonixMetasound/Private/HarmonixMetasound/Nodes/MidiClockSubdivisionTriggerNode.cpp b/Engine/Plugins/Runtime/Harmonix/Source/HarmonixMetasound/Private/HarmonixMetasound/Nodes/MidiClockSubdivisionTriggerNode.cpp
--- a/Engine/Plugins/Runtime/Harmonix/Source/HarmonixMetasound/Private/HarmonixMetasound/Nodes/MidiClockSubdivisionTriggerNode.cpp
+++ b/Engine/Plugins/Runtime/Harmonix/Source/HarmonixMetasound/Private/HarmonixMetasound/Nodes/MidiClockSubdivisionTriggerNode.cpp
@@ -140,7 +140,7 @@ namespace HarmonixMetasound::Nodes::MidiClockSubdivisionTriggerNode
 		const FMidiClockSubdivisionTriggerNode& TheNode = static_cast<const FMidiClockSubdivisionTriggerNode&>(InParams.Node);
 
 		const FInputVertexInterfaceData& InputData = InParams.InputData;
-		FMidiClockReadRef InMidiClock = InputData.GetOrConstructDataReadReference<FMidiClock>(METASOUND_GET_PARAM_NAME(Inputs::MidiClock), InParams.OperatorSettings);
+		FMidiClockReadRef InMidiClock = InputData.GetOrCreateDefaultDataReadReference<FMidiClock>(METASOUND_GET_PARAM_NAME(Inputs::MidiClock), InParams.OperatorSettings);
 		FBoolReadRef InEnabled        = InputData.GetOrCreateDefaultDataReadReference<bool>(METASOUND_GET_PARAM_NAME(Inputs::Enable), InParams.OperatorSettings);
 		FEnumMidiClockSubdivisionQuantizationReadRef InGridSizeUnits = InputData.GetOrCreateDefaultDataReadReference<FEnumMidiClockSubdivisionQuantizationType>(METASOUND_GET_PARAM_NAME(Inputs::GridSizeUnits), InParams.OperatorSettings);
 		FInt32ReadRef InGridSizeMult = InputData.GetOrCreateDefaultDataReadReference<int32>(METASOUND_GET_PARAM_NAME(Inputs::GridSizeMult), InParams.OperatorSettings);
diff --git a/Engine/Plugins/Runtime/Harmonix/Source/HarmonixMetasound/Private/HarmonixMetasound/Nodes/MidiClockTimestampTriggerNode.cpp b/Engine/Plugins/Runtime/Harmonix/Source/HarmonixMetasound/Private/HarmonixMetasound/Nodes/MidiClockTimestampTriggerNode.cpp
--- a/Engine/Plugins/Runtime/Harmonix/Source/HarmonixMetasound/Private/HarmonixMetasound/Nodes/MidiClockTimestampTriggerNode.cpp
+++ b/Engine/Plugins/Runtime/Harmonix/Source/HarmonixMetasound/Private/HarmonixMetasound/Nodes/MidiClockTimestampTriggerNode.cpp
@@ -129,9 +129,9 @@ namespace HarmonixMetasound
 		const FMidiClockTimestampTriggerNode& TheNode = static_cast<const FMidiClockTimestampTriggerNode&>(InParams.Node);
 
 		const FInputVertexInterfaceData& InputData = InParams.InputData;
-		FMidiClockReadRef InMidiClock      = InputData.GetOrConstructDataReadReference<FMidiClock>(METASOUND_GET_PARAM_NAME(Inputs::MidiClock), InParams.OperatorSettings);
+		FMidiClockReadRef InMidiClock      = InputData.GetOrCreateDefaultDataReadReference<FMidiClock>(METASOUND_GET_PARAM_NAME(Inputs::MidiClock), InParams.OperatorSettings);
 		FBoolReadRef  InEnabled            = InputData.GetOrCreateDefaultDataReadReference<bool>(METASOUND_GET_PARAM_NAME(Inputs::Enable), InParams.OperatorSettings);
-		FMusicTimestampReadRef InTimestamp = InputData.GetOrConstructDataReadReference<FMusicTimestamp>(METASOUND_GET_PARAM_NAME(Inputs::Timestamp), 1, 1.0f);
+		FMusicTimestampReadRef InTimestamp = InputData.GetOrCreateDefaultDataReadReference<FMusicTimestamp>(METASOUND_GET_PARAM_NAME(Inputs::Timestamp), InParams.OperatorSettings);
 		FEnumMidiClockSubdivisionQuantizationReadRef InQuantizeUnits = InputData.GetOrCreateDefaultDataReadReference<FEnumMidiClockSubdivisionQuantizationType>(METASOUND_GET_PARAM_NAME(Inputs::QuantizationUnit), InParams.OperatorSettings);
 		FBoolReadRef  InTriggerDuringSeek  = InputData.GetOrCreateDefaultDataReadReference<bool>(METASOUND_GET_PARAM_NAME(TriggerDuringSeek), InParams.OperatorSettings);
 
diff --git a/Engine/Plugins/Runtime/Harmonix/Source/HarmonixMetasound/Private/HarmonixMetasound/Nodes/MidiNoteFilterNode.cpp b/Engine/Plugins/Runtime/Harmonix/Source/HarmonixMetasound/Private/HarmonixMetasound/Nodes/MidiNoteFilterNode.cpp
--- a/Engine/Plugins/Runtime/Harmonix/Source/HarmonixMetasound/Private/HarmonixMetasound/Nodes/MidiNoteFilterNode.cpp
+++ b/Engine/Plugins/Runtime/Harmonix/Source/HarmonixMetasound/Private/HarmonixMetasound/Nodes/MidiNoteFilterNode.cpp
@@ -119,7 +119,7 @@ namespace HarmonixMetasound::Nodes::MidiNoteFilter
 			FInputs Inputs
 			{
 				InputData.GetOrCreateDefaultDataReadReference<bool>(Inputs::EnableName, OperatorSettings),
-				InputData.GetOrConstructDataReadReference<FMidiStream>(Inputs::MidiStreamName),
+				InputData.GetOrCreateDefaultDataReadReference<FMidiStream>(Inputs::MidiStreamName, OperatorSettings),
 				InputData.GetOrCreateDefaultDataReadReference<int32>(Inputs::MinNoteNumberName, OperatorSettings),
 				InputData.GetOrCreateDefaultDataReadReference<int32>(Inputs::MaxNoteNumberName, OperatorSettings),
 				InputData.GetOrCreateDefaultDataReadReference<int32>(Inputs::MinVelocityName, OperatorSettings),
diff --git a/Engine/Plugins/Runtime/Harmonix/Source/HarmonixMetasound/Private/HarmonixMetasound/Nodes/MidiNoteTriggerNode.cpp b/Engine/Plugins/Runtime/Harmonix/Source/HarmonixMetasound/Private/HarmonixMetasound/Nodes/MidiNoteTriggerNode.cpp
--- a/Engine/Plugins/Runtime/Harmonix/Source/HarmonixMetasound/Private/HarmonixMetasound/Nodes/MidiNoteTriggerNode.cpp
+++ b/Engine/Plugins/Runtime/Harmonix/Source/HarmonixMetasound/Private/HarmonixMetasound/Nodes/MidiNoteTriggerNode.cpp
@@ -96,7 +96,7 @@ namespace HarmonixMetasound::Nodes::MidiNoteTriggerNode
 
 			const FInputVertexInterfaceData& InputData = InParams.InputData;
 			FBoolReadRef InEnabled          = InputData.GetOrCreateDefaultDataReadReference<bool>(METASOUND_GET_PARAM_NAME(Inputs::Enable), InParams.OperatorSettings);
-			FMidiStreamReadRef InMidiStream = InputData.GetOrConstructDataReadReference<FMidiStream>(METASOUND_GET_PARAM_NAME(Inputs::MidiStream));
+			FMidiStreamReadRef InMidiStream = InputData.GetOrCreateDefaultDataReadReference<FMidiStream>(METASOUND_GET_PARAM_NAME(Inputs::MidiStream), InParams.OperatorSettings);
 
 			return MakeUnique<FMidiNoteTriggerOperator_V1>(InParams, MoveTemp(InEnabled), MoveTemp(InMidiStream));
 		}
diff --git a/Engine/Plugins/Runtime/Harmonix/Source/HarmonixMetasound/Private/HarmonixMetasound/Nodes/MidiPlayerNode.cpp b/Engine/Plugins/Runtime/Harmonix/Source/HarmonixMetasound/Private/HarmonixMetasound/Nodes/MidiPlayerNode.cpp
--- a/Engine/Plugins/Runtime/Harmonix/Source/HarmonixMetasound/Private/HarmonixMetasound/Nodes/MidiPlayerNode.cpp
+++ b/Engine/Plugins/Runtime/Harmonix/Source/HarmonixMetasound/Private/HarmonixMetasound/Nodes/MidiPlayerNode.cpp
@@ -227,16 +227,16 @@ namespace HarmonixMetasound::Nodes::MidiPlayerNode
 
 		const FOperatorSettings& Settings = InParams.OperatorSettings;
 
-		FMidiAssetReadRef InMidiAsset = InputData.GetOrConstructDataReadReference<FMidiAsset>(METASOUND_GET_PARAM_NAME(Inputs::MidiFileAsset));
-		FMusicTransportEventStreamReadRef InTransport  = InputData.GetOrConstructDataReadReference<FMusicTransportEventStream>(METASOUND_GET_PARAM_NAME(Inputs::Transport), Settings);
-		FBoolReadRef InLoop = InputData.GetOrConstructDataReadReference<bool>(METASOUND_GET_PARAM_NAME(Inputs::Loop), false);
-		FFloatReadRef InSpeed = InputData.GetOrConstructDataReadReference<float>(METASOUND_GET_PARAM_NAME(Inputs::Speed), 1.0f);
+		FMidiAssetReadRef InMidiAsset = InputData.GetOrCreateDefaultDataReadReference<FMidiAsset>(METASOUND_GET_PARAM_NAME(Inputs::MidiFileAsset), Settings);
+		FMusicTransportEventStreamReadRef InTransport  = InputData.GetOrCreateDefaultDataReadReference<FMusicTransportEventStream>(METASOUND_GET_PARAM_NAME(Inputs::Transport), Settings);
+		FBoolReadRef InLoop = InputData.GetOrCreateDefaultDataReadReference<bool>(METASOUND_GET_PARAM_NAME(Inputs::Loop), Settings);
+		FFloatReadRef InSpeed = InputData.GetOrCreateDefaultDataReadReference<float>(METASOUND_GET_PARAM_NAME(Inputs::Speed), Settings);
 		int32 InPrerollBars = InputData.GetOrCreateDefaultValue<int32>(METASOUND_GET_PARAM_NAME(Inputs::PrerollBars), Settings);
 		bool bInKillVoicesOnSeek = InputData.GetOrCreateDefaultValue<bool>(METASOUND_GET_PARAM_NAME(Inputs::KillVoicesOnSeek), Settings);
 		bool bInKillVoicesOnMidiChange = InputData.GetOrCreateDefaultValue<bool>(METASOUND_GET_PARAM_NAME(Inputs::KillVoicesOnMidiChange), Settings);
 		if (InputData.IsVertexBound(METASOUND_GET_PARAM_NAME(Inputs::MidiClock)))
 		{
-			FMidiClockReadRef InMidiClock = InputData.GetOrConstructDataReadReference<FMidiClock>(METASOUND_GET_PARAM_NAME(Inputs::MidiClock), Settings);
+			FMidiClockReadRef InMidiClock = InputData.GetOrCreateDefaultDataReadReference<FMidiClock>(METASOUND_GET_PARAM_NAME(Inputs::MidiClock), Settings);
 			return MakeUnique<FExternallyClockedMidiPlayerOperator>(InParams.OperatorSettings, InMidiAsset, InTransport, InMidiClock, InLoop, InSpeed, InPrerollBars, bInKillVoicesOnSeek, bInKillVoicesOnMidiChange);
 		}
 
diff --git a/Engine/Plugins/Runtime/Harmonix/Source/HarmonixMetasound/Private/HarmonixMetasound/Nodes/MidiPulseGeneratorNode.cpp b/Engine/Plugins/Runtime/Harmonix/Source/HarmonixMetasound/Private/HarmonixMetasound/Nodes/MidiPulseGeneratorNode.cpp
--- a/Engine/Plugins/Runtime/Harmonix/Source/HarmonixMetasound/Private/HarmonixMetasound/Nodes/MidiPulseGeneratorNode.cpp
+++ b/Engine/Plugins/Runtime/Harmonix/Source/HarmonixMetasound/Private/HarmonixMetasound/Nodes/MidiPulseGeneratorNode.cpp
@@ -129,7 +129,7 @@ namespace HarmonixMetasound::Nodes::MidiPulseGeneratorNode
 		{
 			FInputs Inputs
 			{
-				InParams.InputData.GetOrConstructDataReadReference<FMidiClock>(Inputs::MidiClockName, InParams.OperatorSettings),
+				InParams.InputData.GetOrCreateDefaultDataReadReference<FMidiClock>(Inputs::MidiClockName, InParams.OperatorSettings),
 				InParams.InputData.GetOrCreateDefaultDataReadReference<Metasound::FEnumMidiClockSubdivisionQuantizationType>(
 					Inputs::IntervalName, InParams.OperatorSettings),
 				InParams.InputData.GetOrCreateDefaultDataReadReference<int32>(Inputs::IntervalMultiplierName, InParams.OperatorSettings),
diff --git a/Engine/Plugins/Runtime/Harmonix/Source/HarmonixMetasound/Private/HarmonixMetasound/Nodes/MidiStreamLoggerNode.cpp b/Engine/Plugins/Runtime/Harmonix/Source/HarmonixMetasound/Private/HarmonixMetasound/Nodes/MidiStreamLoggerNode.cpp
--- a/Engine/Plugins/Runtime/Harmonix/Source/HarmonixMetasound/Private/HarmonixMetasound/Nodes/MidiStreamLoggerNode.cpp
+++ b/Engine/Plugins/Runtime/Harmonix/Source/HarmonixMetasound/Private/HarmonixMetasound/Nodes/MidiStreamLoggerNode.cpp
@@ -95,7 +95,7 @@ namespace HarmonixMetasound
 
 		const FInputVertexInterfaceData& InputData = InParams.InputData;
 		FBoolReadRef InEnabled = InputData.GetOrCreateDefaultDataReadReference<bool>(METASOUND_GET_PARAM_NAME(Inputs::Enable), InParams.OperatorSettings);
-		FMidiStreamReadRef InMidiStream = InputData.GetOrConstructDataReadReference<FMidiStream>(METASOUND_GET_PARAM_NAME(Inputs::MidiStream));
+		FMidiStreamReadRef InMidiStream = InputData.GetOrCreateDefaultDataReadReference<FMidiStream>(METASOUND_GET_PARAM_NAME(Inputs::MidiStream), InParams.OperatorSettings);
 
 		return MakeUnique<FMidiStreamLoggerOperator>(InParams, InEnabled, InMidiStream);
 	}
diff --git a/Engine/Plugins/Runtime/Harmonix/Source/HarmonixMetasound/Private/HarmonixMetasound/Nodes/MidiStreamMergeNode.cpp b/Engine/Plugins/Runtime/Harmonix/Source/HarmonixMetasound/Private/HarmonixMetasound/Nodes/MidiStreamMergeNode.cpp
--- a/Engine/Plugins/Runtime/Harmonix/Source/HarmonixMetasound/Private/HarmonixMetasound/Nodes/MidiStreamMergeNode.cpp
+++ b/Engine/Plugins/Runtime/Harmonix/Source/HarmonixMetasound/Private/HarmonixMetasound/Nodes/MidiStreamMergeNode.cpp
@@ -101,8 +101,8 @@ namespace HarmonixMetasound
 		const FMidiStreamMergeNode& LoggerNode = static_cast<const FMidiStreamMergeNode&>(InParams.Node);
 
 		const FInputVertexInterfaceData& InputData = InParams.InputData;
-		FMidiStreamReadRef InMidiStreamA = InputData.GetOrConstructDataReadReference<FMidiStream>(METASOUND_GET_PARAM_NAME(InputMidiStreamA));
-		FMidiStreamReadRef InMidiStreamB = InputData.GetOrConstructDataReadReference<FMidiStream>(METASOUND_GET_PARAM_NAME(InputMidiStreamB));
+		FMidiStreamReadRef InMidiStreamA = InputData.GetOrCreateDefaultDataReadReference<FMidiStream>(METASOUND_GET_PARAM_NAME(InputMidiStreamA), InParams.OperatorSettings);
+		FMidiStreamReadRef InMidiStreamB = InputData.GetOrCreateDefaultDataReadReference<FMidiStream>(METASOUND_GET_PARAM_NAME(InputMidiStreamB), InParams.OperatorSettings);
 
 		return MakeUnique<FMidiStreamMergeOperator>(InParams, InMidiStreamA, InMidiStreamB);
 	}
diff --git a/Engine/Plugins/Runtime/Harmonix/Source/HarmonixMetasound/Private/HarmonixMetasound/Nodes/MidiStreamSelectNode.cpp b/Engine/Plugins/Runtime/Harmonix/Source/HarmonixMetasound/Private/HarmonixMetasound/Nodes/MidiStreamSelectNode.cpp
--- a/Engine/Plugins/Runtime/Harmonix/Source/HarmonixMetasound/Private/HarmonixMetasound/Nodes/MidiStreamSelectNode.cpp
+++ b/Engine/Plugins/Runtime/Harmonix/Source/HarmonixMetasound/Private/HarmonixMetasound/Nodes/MidiStreamSelectNode.cpp
@@ -122,11 +122,11 @@ namespace HarmonixMetasound
 		const FMidiStreamSelectNode& LoggerNode = static_cast<const FMidiStreamSelectNode&>(InParams.Node);
 
 		const FInputVertexInterfaceData& InputData = InParams.InputData;
-		FMidiStreamReadRef InMidiStreamA = InputData.GetOrConstructDataReadReference<FMidiStream>(METASOUND_GET_PARAM_NAME(InputMidiStreamA));
-		FMidiStreamReadRef InMidiStreamB = InputData.GetOrConstructDataReadReference<FMidiStream>(METASOUND_GET_PARAM_NAME(InputMidiStreamB));
-		FInt32ReadRef InStreamIndex = InputData.GetOrConstructDataReadReference<int32>(METASOUND_GET_PARAM_NAME(InputStreamIndex));
-		FBoolReadRef InImmediateNoteOff = InputData.GetOrConstructDataReadReference<bool>(METASOUND_GET_PARAM_NAME(InputImmediateNoteOff));
-		FBoolReadRef InCopyInactiveNoteOffs = InputData.GetOrConstructDataReadReference<bool>(METASOUND_GET_PARAM_NAME(InputCopyInactiveNoteOffs));
+		FMidiStreamReadRef InMidiStreamA = InputData.GetOrCreateDefaultDataReadReference<FMidiStream>(METASOUND_GET_PARAM_NAME(InputMidiStreamA), InParams.OperatorSettings);
+		FMidiStreamReadRef InMidiStreamB = InputData.GetOrCreateDefaultDataReadReference<FMidiStream>(METASOUND_GET_PARAM_NAME(InputMidiStreamB), InParams.OperatorSettings);
+		FInt32ReadRef InStreamIndex = InputData.GetOrCreateDefaultDataReadReference<int32>(METASOUND_GET_PARAM_NAME(InputStreamIndex), InParams.OperatorSettings);
+		FBoolReadRef InImmediateNoteOff = InputData.GetOrCreateDefaultDataReadReference<bool>(METASOUND_GET_PARAM_NAME(InputImmediateNoteOff), InParams.OperatorSettings);
+		FBoolReadRef InCopyInactiveNoteOffs = InputData.GetOrCreateDefaultDataReadReference<bool>(METASOUND_GET_PARAM_NAME(InputCopyInactiveNoteOffs), InParams.OperatorSettings);
 
 		return MakeUnique<FMidiStreamSelectOperator>(InParams, InMidiStreamA, InMidiStreamB, InStreamIndex, InImmediateNoteOff, InCopyInactiveNoteOffs);
 	}
@@ -258,4 +258,4 @@ namespace HarmonixMetasound
 	}
 }
 
-#undef LOCTEXT_NAMESPACE // "HarmonixMetaSound"
\ No newline at end of file
+#undef LOCTEXT_NAMESPACE // "HarmonixMetaSound"
diff --git a/Engine/Plugins/Runtime/Harmonix/Source/HarmonixMetasound/Private/HarmonixMetasound/Nodes/MidiStreamTrackFilterNode.cpp b/Engine/Plugins/Runtime/Harmonix/Source/HarmonixMetasound/Private/HarmonixMetasound/Nodes/MidiStreamTrackFilterNode.cpp
--- a/Engine/Plugins/Runtime/Harmonix/Source/HarmonixMetasound/Private/HarmonixMetasound/Nodes/MidiStreamTrackFilterNode.cpp
+++ b/Engine/Plugins/Runtime/Harmonix/Source/HarmonixMetasound/Private/HarmonixMetasound/Nodes/MidiStreamTrackFilterNode.cpp
@@ -114,7 +114,7 @@ namespace HarmonixMetasound::Nodes::MidiTrackFilter
 			FInputs Inputs
 			{
 				InputData.GetOrCreateDefaultDataReadReference<bool>(Inputs::EnableName, InParams.OperatorSettings),
-				InputData.GetOrConstructDataReadReference<FMidiStream>(Inputs::MidiStreamName),
+				InputData.GetOrCreateDefaultDataReadReference<FMidiStream>(Inputs::MidiStreamName, InParams.OperatorSettings),
 				InputData.GetOrCreateDefaultDataReadReference<int32>(Inputs::MinTrackIndexName, InParams.OperatorSettings),
 				InputData.GetOrCreateDefaultDataReadReference<int32>(Inputs::MaxTrackIndexName, InParams.OperatorSettings),
 				InputData.GetOrCreateDefaultDataReadReference<bool>(Inputs::IncludeConductorTrackName, InParams.OperatorSettings)
diff --git a/Engine/Plugins/Runtime/Harmonix/Source/HarmonixMetasound/Private/HarmonixMetasound/Nodes/MidiStreamTransposerNode.cpp b/Engine/Plugins/Runtime/Harmonix/Source/HarmonixMetasound/Private/HarmonixMetasound/Nodes/MidiStreamTransposerNode.cpp
--- a/Engine/Plugins/Runtime/Harmonix/Source/HarmonixMetasound/Private/HarmonixMetasound/Nodes/MidiStreamTransposerNode.cpp
+++ b/Engine/Plugins/Runtime/Harmonix/Source/HarmonixMetasound/Private/HarmonixMetasound/Nodes/MidiStreamTransposerNode.cpp
@@ -103,7 +103,7 @@ namespace HarmonixMetasound::Nodes::MidiNoteTranspose
 			FInputs Inputs
 			{
 				InputData.GetOrCreateDefaultDataReadReference<bool>(METASOUND_GET_PARAM_NAME(Inputs::Enable), InParams.OperatorSettings),
-				InputData.GetOrConstructDataReadReference<FMidiStream>(METASOUND_GET_PARAM_NAME(Inputs::MidiStream)),
+				InputData.GetOrCreateDefaultDataReadReference<FMidiStream>(METASOUND_GET_PARAM_NAME(Inputs::MidiStream), InParams.OperatorSettings),
 				InputData.GetOrCreateDefaultDataReadReference<int32>(METASOUND_GET_PARAM_NAME(Inputs::Transposition), InParams.OperatorSettings)
 			};
 
diff --git a/Engine/Plugins/Runtime/Harmonix/Source/HarmonixMetasound/Private/HarmonixMetasound/Nodes/MidiStreamWriterNode.cpp b/Engine/Plugins/Runtime/Harmonix/Source/HarmonixMetasound/Private/HarmonixMetasound/Nodes/MidiStreamWriterNode.cpp
--- a/Engine/Plugins/Runtime/Harmonix/Source/HarmonixMetasound/Private/HarmonixMetasound/Nodes/MidiStreamWriterNode.cpp
+++ b/Engine/Plugins/Runtime/Harmonix/Source/HarmonixMetasound/Private/HarmonixMetasound/Nodes/MidiStreamWriterNode.cpp
@@ -90,7 +90,7 @@ namespace HarmonixMetasound::Nodes::MidiStreamWriter
 
 			FBoolReadRef InEnabled = InputData.GetOrCreateDefaultDataReadReference<bool>(METASOUND_GET_PARAM_NAME(Inputs::Enable), InParams.OperatorSettings);
 			FStringReadRef InFilenamePrefix = InputData.GetOrCreateDefaultDataReadReference<FString>(METASOUND_GET_PARAM_NAME(Inputs::FilenamePrefix), InParams.OperatorSettings);
-			FMidiStreamReadRef InMidiStream = InputData.GetOrConstructDataReadReference<FMidiStream>(METASOUND_GET_PARAM_NAME(Inputs::MidiStream));
+			FMidiStreamReadRef InMidiStream = InputData.GetOrCreateDefaultDataReadReference<FMidiStream>(METASOUND_GET_PARAM_NAME(Inputs::MidiStream), InParams.OperatorSettings);
 			
 			return MakeUnique<FMidiStreamWriterOperator>(InParams, InEnabled, InFilenamePrefix, InMidiStream);
 		}
diff --git a/Engine/Plugins/Runtime/Harmonix/Source/HarmonixMetasound/Private/HarmonixMetasound/Nodes/MidiTextTriggerNode.cpp b/Engine/Plugins/Runtime/Harmonix/Source/HarmonixMetasound/Private/HarmonixMetasound/Nodes/MidiTextTriggerNode.cpp
--- a/Engine/Plugins/Runtime/Harmonix/Source/HarmonixMetasound/Private/HarmonixMetasound/Nodes/MidiTextTriggerNode.cpp
+++ b/Engine/Plugins/Runtime/Harmonix/Source/HarmonixMetasound/Private/HarmonixMetasound/Nodes/MidiTextTriggerNode.cpp
@@ -76,7 +76,7 @@ namespace HarmonixMetasound
 
 			const FInputVertexInterfaceData& InputData = InParams.InputData;
 			FBoolReadRef InEnabled          = InputData.GetOrCreateDefaultDataReadReference<bool>(METASOUND_GET_PARAM_NAME(Inputs::Enable), InParams.OperatorSettings);
-			FMidiStreamReadRef InMidiStream = InputData.GetOrConstructDataReadReference<FMidiStream>(METASOUND_GET_PARAM_NAME(Inputs::MidiStream));
+			FMidiStreamReadRef InMidiStream = InputData.GetOrCreateDefaultDataReadReference<FMidiStream>(METASOUND_GET_PARAM_NAME(Inputs::MidiStream), InParams.OperatorSettings);
 			FStringReadRef InText           = InputData.GetOrCreateDefaultDataReadReference<FString>(METASOUND_GET_PARAM_NAME(TextInput), InParams.OperatorSettings);
 
 			return MakeUnique<FMidiTextTriggerOperator_V1>(InParams, InEnabled, InMidiStream, InText);
diff --git a/Engine/Plugins/Runtime/Harmonix/Source/HarmonixMetasound/Private/HarmonixMetasound/Nodes/MultibandAnalyzerNode.cpp b/Engine/Plugins/Runtime/Harmonix/Source/HarmonixMetasound/Private/HarmonixMetasound/Nodes/MultibandAnalyzerNode.cpp
--- a/Engine/Plugins/Runtime/Harmonix/Source/HarmonixMetasound/Private/HarmonixMetasound/Nodes/MultibandAnalyzerNode.cpp
+++ b/Engine/Plugins/Runtime/Harmonix/Source/HarmonixMetasound/Private/HarmonixMetasound/Nodes/MultibandAnalyzerNode.cpp
@@ -252,7 +252,7 @@ namespace HarmonixMetasound::Nodes::MultibandAnalyzer
 			FInputs Inputs
 			{
 				InParams.InputData.GetOrCreateDefaultDataReadReference<bool>(Inputs::EnableName, InParams.OperatorSettings),
-				InParams.InputData.GetOrConstructDataReadReference<Metasound::FAudioBuffer>(Inputs::AudioMonoName),
+				InParams.InputData.GetOrCreateDefaultDataReadReference<Metasound::FAudioBuffer>(Inputs::AudioMonoName, InParams.OperatorSettings),
 				InParams.InputData.GetOrCreateDefaultValue<TArray<float>>(Inputs::CrossoverFrequenciesName, InParams.OperatorSettings),
 				InParams.InputData.GetOrCreateDefaultValue<bool>(Inputs::ApplySmoothingName, InParams.OperatorSettings),
 				InParams.InputData.GetOrCreateDefaultValue<Metasound::FTime>(Inputs::AttackTimeName, InParams.OperatorSettings),
diff --git a/Engine/Plugins/Runtime/Harmonix/Source/HarmonixMetasound/Private/HarmonixMetasound/Nodes/PeakNode.cpp b/Engine/Plugins/Runtime/Harmonix/Source/HarmonixMetasound/Private/HarmonixMetasound/Nodes/PeakNode.cpp
--- a/Engine/Plugins/Runtime/Harmonix/Source/HarmonixMetasound/Private/HarmonixMetasound/Nodes/PeakNode.cpp
+++ b/Engine/Plugins/Runtime/Harmonix/Source/HarmonixMetasound/Private/HarmonixMetasound/Nodes/PeakNode.cpp
@@ -102,7 +102,7 @@ namespace HarmonixMetasound::Nodes::Peak
 		{
 			FInputs Inputs
 			{
-				InParams.InputData.GetOrConstructDataReadReference<Metasound::FAudioBuffer>(
+				InParams.InputData.GetOrCreateDefaultDataReadReference<Metasound::FAudioBuffer>(
 					Inputs::AudioMonoName,
 					InParams.OperatorSettings),
 				InParams.InputData.GetOrCreateDefaultDataReadReference<bool>(
diff --git a/Engine/Plugins/Runtime/Harmonix/Source/HarmonixMetasound/Private/HarmonixMetasound/Nodes/StepSequencePlayerNode.cpp b/Engine/Plugins/Runtime/Harmonix/Source/HarmonixMetasound/Private/HarmonixMetasound/Nodes/StepSequencePlayerNode.cpp
--- a/Engine/Plugins/Runtime/Harmonix/Source/HarmonixMetasound/Private/HarmonixMetasound/Nodes/StepSequencePlayerNode.cpp
+++ b/Engine/Plugins/Runtime/Harmonix/Source/HarmonixMetasound/Private/HarmonixMetasound/Nodes/StepSequencePlayerNode.cpp
@@ -223,11 +223,11 @@ namespace HarmonixMetasound::Nodes::StepSequencePlayer
 
 		const FInputVertexInterfaceData& InputData = InParams.InputData;
 
-		FMidiStepSequenceAssetReadRef InSequenceAsset = InputData.GetOrConstructDataReadReference<FMidiStepSequenceAsset>(METASOUND_GET_PARAM_NAME(Inputs::SequenceAsset));
-		FMusicTransportEventStreamReadRef InTransport = InputData.GetOrConstructDataReadReference<FMusicTransportEventStream>(METASOUND_GET_PARAM_NAME(Inputs::Transport), InParams.OperatorSettings);
+		FMidiStepSequenceAssetReadRef InSequenceAsset = InputData.GetOrCreateDefaultDataReadReference<FMidiStepSequenceAsset>(METASOUND_GET_PARAM_NAME(Inputs::SequenceAsset), InParams.OperatorSettings);
+		FMusicTransportEventStreamReadRef InTransport = InputData.GetOrCreateDefaultDataReadReference<FMusicTransportEventStream>(METASOUND_GET_PARAM_NAME(Inputs::Transport), InParams.OperatorSettings);
 		FFloatReadRef InSpeed = InputData.GetOrCreateDefaultDataReadReference<float>(METASOUND_GET_PARAM_NAME(Inputs::Speed), InParams.OperatorSettings);
 		FFloatReadRef InVelocityMult = InputData.GetOrCreateDefaultDataReadReference<float>(METASOUND_GET_PARAM_NAME(Inputs::VelocityMultiplier), InParams.OperatorSettings);
-		FMidiClockReadRef InMidiClock = InputData.GetOrConstructDataReadReference<FMidiClock>(METASOUND_GET_PARAM_NAME(Inputs::MidiClock), InParams.OperatorSettings);
+		FMidiClockReadRef InMidiClock = InputData.GetOrCreateDefaultDataReadReference<FMidiClock>(METASOUND_GET_PARAM_NAME(Inputs::MidiClock), InParams.OperatorSettings);
 		FFloatReadRef InMaxColumns = InputData.GetOrCreateDefaultDataReadReference<float>(METASOUND_GET_PARAM_NAME(Inputs::MaxColumns), InParams.OperatorSettings);
 		FFloatReadRef InAdditionalOctaves = InputData.GetOrCreateDefaultDataReadReference<float>(METASOUND_GET_PARAM_NAME(Inputs::AdditionalOctaves), InParams.OperatorSettings);
 		FFloatReadRef InStepSize = InputData.GetOrCreateDefaultDataReadReference<float>(METASOUND_GET_PARAM_NAME(Inputs::StepSizeQuarterNotes), InParams.OperatorSettings);
diff --git a/Engine/Plugins/Runtime/Harmonix/Source/HarmonixMetasound/Private/HarmonixMetasound/Nodes/TransportToTriggerNode.cpp b/Engine/Plugins/Runtime/Harmonix/Source/HarmonixMetasound/Private/HarmonixMetasound/Nodes/TransportToTriggerNode.cpp
--- a/Engine/Plugins/Runtime/Harmonix/Source/HarmonixMetasound/Private/HarmonixMetasound/Nodes/TransportToTriggerNode.cpp
+++ b/Engine/Plugins/Runtime/Harmonix/Source/HarmonixMetasound/Private/HarmonixMetasound/Nodes/TransportToTriggerNode.cpp
@@ -99,7 +99,7 @@ namespace HarmonixMetasound
 			using namespace CommonPinNames;
 
 			const FInputVertexInterfaceData& InputData = InParams.InputData;
-			FMusicTransportEventStreamReadRef InTransport = InputData.GetOrConstructDataReadReference<FMusicTransportEventStream>(METASOUND_GET_PARAM_NAME(Inputs::Transport), InParams.OperatorSettings);
+			FMusicTransportEventStreamReadRef InTransport = InputData.GetOrCreateDefaultDataReadReference<FMusicTransportEventStream>(METASOUND_GET_PARAM_NAME(Inputs::Transport), InParams.OperatorSettings);
 			
 			return MakeUnique<FTransportToTriggerOperator>(InParams, InTransport);
 		}
diff --git a/Engine/Plugins/Runtime/Harmonix/Source/HarmonixMetasound/Private/HarmonixMetasound/Nodes/TransportWavePlayerControllerNode.cpp b/Engine/Plugins/Runtime/Harmonix/Source/HarmonixMetasound/Private/HarmonixMetasound/Nodes/TransportWavePlayerControllerNode.cpp
--- a/Engine/Plugins/Runtime/Harmonix/Source/HarmonixMetasound/Private/HarmonixMetasound/Nodes/TransportWavePlayerControllerNode.cpp
+++ b/Engine/Plugins/Runtime/Harmonix/Source/HarmonixMetasound/Private/HarmonixMetasound/Nodes/TransportWavePlayerControllerNode.cpp
@@ -108,8 +108,8 @@ namespace HarmonixMetasound
 			using namespace CommonPinNames;
 
 			const FInputVertexInterfaceData& InputData = InParams.InputData;
-			FMusicTransportEventStreamReadRef InTransport = InputData.GetOrConstructDataReadReference<FMusicTransportEventStream>(METASOUND_GET_PARAM_NAME(Inputs::Transport), InParams.OperatorSettings);
-			FMidiClockReadRef InMidiClock = InputData.GetOrConstructDataReadReference<FMidiClock>(METASOUND_GET_PARAM_NAME(Inputs::MidiClock), InParams.OperatorSettings);
+			FMusicTransportEventStreamReadRef InTransport = InputData.GetOrCreateDefaultDataReadReference<FMusicTransportEventStream>(METASOUND_GET_PARAM_NAME(Inputs::Transport), InParams.OperatorSettings);
+			FMidiClockReadRef InMidiClock = InputData.GetOrCreateDefaultDataReadReference<FMidiClock>(METASOUND_GET_PARAM_NAME(Inputs::MidiClock), InParams.OperatorSettings);
 			
 			return MakeUnique<FTransportWavePlayerControllerOperator>(InParams.OperatorSettings, InTransport, InMidiClock);
 		}
diff --git a/Engine/Plugins/Runtime/Harmonix/Source/HarmonixMetasound/Private/HarmonixMetasound/Nodes/TriggerToTransportNode.cpp b/Engine/Plugins/Runtime/Harmonix/Source/HarmonixMetasound/Private/HarmonixMetasound/Nodes/TriggerToTransportNode.cpp
--- a/Engine/Plugins/Runtime/Harmonix/Source/HarmonixMetasound/Private/HarmonixMetasound/Nodes/TriggerToTransportNode.cpp
+++ b/Engine/Plugins/Runtime/Harmonix/Source/HarmonixMetasound/Private/HarmonixMetasound/Nodes/TriggerToTransportNode.cpp
@@ -136,14 +136,14 @@ namespace HarmonixMetasound::Nodes::TriggerToTransportNode
 	TUniquePtr<IOperator> FTriggerToTransportOperator::CreateOperator(const FBuildOperatorParams& InParams, FBuildResults& OutResults)
 	{
 		const FInputVertexInterfaceData& InputData = InParams.InputData;
-		FTriggerReadRef InTriggerPrepare          = InputData.GetOrConstructDataReadReference<FTrigger>(METASOUND_GET_PARAM_NAME(Inputs::TransportPrepare), InParams.OperatorSettings);
-		FTriggerReadRef InTriggerPlay             = InputData.GetOrConstructDataReadReference<FTrigger>(METASOUND_GET_PARAM_NAME(Inputs::TransportPlay), InParams.OperatorSettings);
-		FTriggerReadRef InTriggerPause            = InputData.GetOrConstructDataReadReference<FTrigger>(METASOUND_GET_PARAM_NAME(Inputs::TransportPause), InParams.OperatorSettings);
-		FTriggerReadRef InTriggerContinue         = InputData.GetOrConstructDataReadReference<FTrigger>(METASOUND_GET_PARAM_NAME(Inputs::TransportContinue), InParams.OperatorSettings);
-		FTriggerReadRef InTriggerStop             = InputData.GetOrConstructDataReadReference<FTrigger>(METASOUND_GET_PARAM_NAME(Inputs::TransportStop), InParams.OperatorSettings);
-		FTriggerReadRef InTriggerKill             = InputData.GetOrConstructDataReadReference<FTrigger>(METASOUND_GET_PARAM_NAME(Inputs::TransportKill), InParams.OperatorSettings);
-		FTriggerReadRef InTriggerSeek             = InputData.GetOrConstructDataReadReference<FTrigger>(METASOUND_GET_PARAM_NAME(Inputs::TriggerSeek), InParams.OperatorSettings);
-		FMusicSeekTargetReadRef InSeekDestination = InputData.GetOrConstructDataReadReference<FMusicSeekTarget>(METASOUND_GET_PARAM_NAME(Inputs::SeekDestination));
+		FTriggerReadRef InTriggerPrepare          = InputData.GetOrCreateDefaultDataReadReference<FTrigger>(METASOUND_GET_PARAM_NAME(Inputs::TransportPrepare), InParams.OperatorSettings);
+		FTriggerReadRef InTriggerPlay             = InputData.GetOrCreateDefaultDataReadReference<FTrigger>(METASOUND_GET_PARAM_NAME(Inputs::TransportPlay), InParams.OperatorSettings);
+		FTriggerReadRef InTriggerPause            = InputData.GetOrCreateDefaultDataReadReference<FTrigger>(METASOUND_GET_PARAM_NAME(Inputs::TransportPause), InParams.OperatorSettings);
+		FTriggerReadRef InTriggerContinue         = InputData.GetOrCreateDefaultDataReadReference<FTrigger>(METASOUND_GET_PARAM_NAME(Inputs::TransportContinue), InParams.OperatorSettings);
+		FTriggerReadRef InTriggerStop             = InputData.GetOrCreateDefaultDataReadReference<FTrigger>(METASOUND_GET_PARAM_NAME(Inputs::TransportStop), InParams.OperatorSettings);
+		FTriggerReadRef InTriggerKill             = InputData.GetOrCreateDefaultDataReadReference<FTrigger>(METASOUND_GET_PARAM_NAME(Inputs::TransportKill), InParams.OperatorSettings);
+		FTriggerReadRef InTriggerSeek             = InputData.GetOrCreateDefaultDataReadReference<FTrigger>(METASOUND_GET_PARAM_NAME(Inputs::TriggerSeek), InParams.OperatorSettings);
+		FMusicSeekTargetReadRef InSeekDestination = InputData.GetOrCreateDefaultDataReadReference<FMusicSeekTarget>(METASOUND_GET_PARAM_NAME(Inputs::SeekDestination), InParams.OperatorSettings);
 		return MakeUnique<FTriggerToTransportOperator>(InParams,
 			InTriggerPrepare,
 			InTriggerPlay,
diff --git a/Engine/Plugins/Runtime/Metasound/Source/MetasoundEngine/Private/MetasoundAudioBusReaderNode.cpp b/Engine/Plugins/Runtime/Metasound/Source/MetasoundEngine/Private/MetasoundAudioBusReaderNode.cpp
--- a/Engine/Plugins/Runtime/Metasound/Source/MetasoundEngine/Private/MetasoundAudioBusReaderNode.cpp
+++ b/Engine/Plugins/Runtime/Metasound/Source/MetasoundEngine/Private/MetasoundAudioBusReaderNode.cpp
@@ -118,7 +118,7 @@ namespace Metasound
 			
 			if (bHasEnvironmentVars)
 			{
-				FAudioBusAssetReadRef AudioBusIn = InputData.GetOrConstructDataReadReference<FAudioBusAsset>(METASOUND_GET_PARAM_NAME(InParamAudioBusInput));
+				FAudioBusAssetReadRef AudioBusIn = InputData.GetOrCreateDefaultDataReadReference<FAudioBusAsset>(METASOUND_GET_PARAM_NAME(InParamAudioBusInput), InParams.OperatorSettings);
 				return MakeUnique<TAudioBusReaderOperator<NumChannels>>(InParams, AudioBusIn);
 			}
 			else
diff --git a/Engine/Plugins/Runtime/Metasound/Source/MetasoundEngine/Private/MetasoundAudioBusWriterNode.cpp b/Engine/Plugins/Runtime/Metasound/Source/MetasoundEngine/Private/MetasoundAudioBusWriterNode.cpp
--- a/Engine/Plugins/Runtime/Metasound/Source/MetasoundEngine/Private/MetasoundAudioBusWriterNode.cpp
+++ b/Engine/Plugins/Runtime/Metasound/Source/MetasoundEngine/Private/MetasoundAudioBusWriterNode.cpp
@@ -103,12 +103,12 @@ namespace Metasound
 
 			if (bHasEnvironmentVars)
 			{
-				FAudioBusAssetReadRef AudioBusIn = InputData.GetOrConstructDataReadReference<FAudioBusAsset>(METASOUND_GET_PARAM_NAME(Inputs::AudioBus));
+				FAudioBusAssetReadRef AudioBusIn = InputData.GetOrCreateDefaultDataReadReference<FAudioBusAsset>(METASOUND_GET_PARAM_NAME(Inputs::AudioBus), InParams.OperatorSettings);
 
 				TArray<FAudioBufferReadRef> AudioInputs;
 				for (int32 ChannelIndex = 0; ChannelIndex < NumChannels; ++ChannelIndex)
 				{
-					AudioInputs.Add(InputData.GetOrConstructDataReadReference<FAudioBuffer>(METASOUND_GET_PARAM_NAME_WITH_INDEX(Inputs::Audio, ChannelIndex), InParams.OperatorSettings));
+					AudioInputs.Add(InputData.GetOrCreateDefaultDataReadReference<FAudioBuffer>(METASOUND_GET_PARAM_NAME_WITH_INDEX(Inputs::Audio, ChannelIndex), InParams.OperatorSettings));
 				}
 
 				FString GraphName;
diff --git a/Engine/Plugins/Runtime/Metasound/Source/MetasoundEngine/Private/MetasoundWaveInfoNode.cpp b/Engine/Plugins/Runtime/Metasound/Source/MetasoundEngine/Private/MetasoundWaveInfoNode.cpp
--- a/Engine/Plugins/Runtime/Metasound/Source/MetasoundEngine/Private/MetasoundWaveInfoNode.cpp
+++ b/Engine/Plugins/Runtime/Metasound/Source/MetasoundEngine/Private/MetasoundWaveInfoNode.cpp
@@ -121,7 +121,7 @@ namespace Metasound
 		
 		const FInputVertexInterfaceData& InputData = InParams.InputData;
 		// inputs
-		FWaveAssetReadRef WaveAssetIn = InputData.GetOrConstructDataReadReference<FWaveAsset>(METASOUND_GET_PARAM_NAME(ParamWaveAsset));
+		FWaveAssetReadRef WaveAssetIn = InputData.GetOrCreateDefaultDataReadReference<FWaveAsset>(METASOUND_GET_PARAM_NAME(ParamWaveAsset), InParams.OperatorSettings);
 
 		return MakeUnique < FWaveInfoNodeOperator >(InParams.OperatorSettings, WaveAssetIn);
 	}
diff --git a/Engine/Plugins/Runtime/Metasound/Source/MetasoundEngine/Private/MetasoundWavePlayerNode.cpp b/Engine/Plugins/Runtime/Metasound/Source/MetasoundEngine/Private/MetasoundWavePlayerNode.cpp
--- a/Engine/Plugins/Runtime/Metasound/Source/MetasoundEngine/Private/MetasoundWavePlayerNode.cpp
+++ b/Engine/Plugins/Runtime/Metasound/Source/MetasoundEngine/Private/MetasoundWavePlayerNode.cpp
@@ -1184,9 +1184,9 @@ namespace Metasound
 			{
 				InParams.OperatorSettings,
 				OutputAudioVertices,
-				Inputs.GetOrConstructDataReadReference<FTrigger>(METASOUND_GET_PARAM_NAME(InputTriggerPlay), InParams.OperatorSettings),
-				Inputs.GetOrConstructDataReadReference<FTrigger>(METASOUND_GET_PARAM_NAME(InputTriggerStop), InParams.OperatorSettings),
-				Inputs.GetOrConstructDataReadReference<FWaveAsset>(METASOUND_GET_PARAM_NAME(InputWaveAsset)),
+				Inputs.GetOrCreateDefaultDataReadReference<FTrigger>(METASOUND_GET_PARAM_NAME(InputTriggerPlay), InParams.OperatorSettings),
+				Inputs.GetOrCreateDefaultDataReadReference<FTrigger>(METASOUND_GET_PARAM_NAME(InputTriggerStop), InParams.OperatorSettings),
+				Inputs.GetOrCreateDefaultDataReadReference<FWaveAsset>(METASOUND_GET_PARAM_NAME(InputWaveAsset), InParams.OperatorSettings),
 				Inputs.GetOrCreateDefaultDataReadReference<FTime>(METASOUND_GET_PARAM_NAME(InputStartTime), InParams.OperatorSettings),
 				Inputs.GetOrCreateDefaultDataReadReference<float>(METASOUND_GET_PARAM_NAME(InputPitchShift), InParams.OperatorSettings),
 				Inputs.GetOrCreateDefaultDataReadReference<bool>(METASOUND_GET_PARAM_NAME(InputLoop), InParams.OperatorSettings),
diff --git a/Engine/Plugins/Runtime/Metasound/Source/MetasoundEngine/Private/MetasoundWaveTableBankEvaluateNode.cpp b/Engine/Plugins/Runtime/Metasound/Source/MetasoundEngine/Private/MetasoundWaveTableBankEvaluateNode.cpp
--- a/Engine/Plugins/Runtime/Metasound/Source/MetasoundEngine/Private/MetasoundWaveTableBankEvaluateNode.cpp
+++ b/Engine/Plugins/Runtime/Metasound/Source/MetasoundEngine/Private/MetasoundWaveTableBankEvaluateNode.cpp
@@ -88,7 +88,7 @@ namespace Metasound
 
 			const FInputVertexInterfaceData& InputData = InParams.InputData;
 
-			FWaveTableBankAssetReadRef InWaveTableReadRef = InputData.GetOrConstructDataReadReference<FWaveTableBankAsset>("WaveTableBank");
+			FWaveTableBankAssetReadRef InWaveTableReadRef = InputData.GetOrCreateDefaultDataReadReference<FWaveTableBankAsset>("WaveTableBank", InParams.OperatorSettings);
 			FFloatReadRef InInputReadRef = InputData.GetOrCreateDefaultDataReadReference<float>("Input", InParams.OperatorSettings);
 			FFloatReadRef InIndexReadRef = InputData.GetOrCreateDefaultDataReadReference<float>("Index", InParams.OperatorSettings);
 			FEnumWaveTableInterpModeReadRef InInterpReadRef = InputData.GetOrCreateDefaultDataReadReference<FEnumWaveTableInterpolationMode>("Interpolation", InParams.OperatorSettings);
diff --git a/Engine/Plugins/Runtime/Metasound/Source/MetasoundEngine/Private/MetasoundWaveTableEnvelopeNode.cpp b/Engine/Plugins/Runtime/Metasound/Source/MetasoundEngine/Private/MetasoundWaveTableEnvelopeNode.cpp
--- a/Engine/Plugins/Runtime/Metasound/Source/MetasoundEngine/Private/MetasoundWaveTableEnvelopeNode.cpp
+++ b/Engine/Plugins/Runtime/Metasound/Source/MetasoundEngine/Private/MetasoundWaveTableEnvelopeNode.cpp
@@ -102,7 +102,7 @@ namespace Metasound
 			
 			const FInputVertexInterfaceData& InputData = InParams.InputData;
 
-			FWaveTableReadRef InWaveTableReadRef = InputData.GetOrConstructDataReadReference<FWaveTable>("WaveTable");
+			FWaveTableReadRef InWaveTableReadRef = InputData.GetOrCreateDefaultDataReadReference<FWaveTable>("WaveTable", InParams.OperatorSettings);
 			FTriggerReadRef InPlayReadRef = InputData.GetOrCreateDefaultDataReadReference<FTrigger>("Play", InParams.OperatorSettings);
 			FTriggerReadRef InStopReadRef = InputData.GetOrCreateDefaultDataReadReference<FTrigger>("Stop", InParams.OperatorSettings);
 			FTriggerReadRef InPauseReadRef = InputData.GetOrCreateDefaultDataReadReference<FTrigger>("Pause", InParams.OperatorSettings);
diff --git a/Engine/Plugins/Runtime/Metasound/Source/MetasoundEngine/Private/MetasoundWaveTableEvaluateNode.cpp b/Engine/Plugins/Runtime/Metasound/Source/MetasoundEngine/Private/MetasoundWaveTableEvaluateNode.cpp
--- a/Engine/Plugins/Runtime/Metasound/Source/MetasoundEngine/Private/MetasoundWaveTableEvaluateNode.cpp
+++ b/Engine/Plugins/Runtime/Metasound/Source/MetasoundEngine/Private/MetasoundWaveTableEvaluateNode.cpp
@@ -83,7 +83,7 @@ namespace Metasound
 
 			const FInputVertexInterfaceData& InputData = InParams.InputData;
 
-			FWaveTableReadRef InWaveTableReadRef = InputData.GetOrConstructDataReadReference<FWaveTable>("WaveTable");
+			FWaveTableReadRef InWaveTableReadRef = InputData.GetOrCreateDefaultDataReadReference<FWaveTable>("WaveTable", InParams.OperatorSettings);
 			FFloatReadRef InInputReadRef = InputData.GetOrCreateDefaultDataReadReference<float>("Input", InParams.OperatorSettings);
 			FEnumWaveTableInterpModeReadRef InInterpReadRef = InputData.GetOrCreateDefaultDataReadReference<FEnumWaveTableInterpolationMode>("Interpolation", InParams.OperatorSettings);
 
diff --git a/Engine/Plugins/Runtime/Metasound/Source/MetasoundEngine/Private/MetasoundWaveTableGetNode.cpp b/Engine/Plugins/Runtime/Metasound/Source/MetasoundEngine/Private/MetasoundWaveTableGetNode.cpp
--- a/Engine/Plugins/Runtime/Metasound/Source/MetasoundEngine/Private/MetasoundWaveTableGetNode.cpp
+++ b/Engine/Plugins/Runtime/Metasound/Source/MetasoundEngine/Private/MetasoundWaveTableGetNode.cpp
@@ -75,7 +75,7 @@ namespace Metasound
 			
 			const FInputVertexInterfaceData& InputData = InParams.InputData;
 
-			FWaveTableBankAssetReadRef WaveTableBankReadRef = InputData.GetOrConstructDataReadReference<FWaveTableBankAsset>("WaveTableBank");
+			FWaveTableBankAssetReadRef WaveTableBankReadRef = InputData.GetOrCreateDefaultDataReadReference<FWaveTableBankAsset>("WaveTableBank", InParams.OperatorSettings);
 			FFloatReadRef TableIndexReadRef = InputData.GetOrCreateDefaultDataReadReference<float>("TableIndex", InParams.OperatorSettings);
 
 			return MakeUnique<FMetasoundWaveTableGetNodeOperator>(InParams, WaveTableBankReadRef, TableIndexReadRef);
diff --git a/Engine/Plugins/Runtime/Metasound/Source/MetasoundEngine/Private/MetasoundWaveTableOscillatorNode.cpp b/Engine/Plugins/Runtime/Metasound/Source/MetasoundEngine/Private/MetasoundWaveTableOscillatorNode.cpp
--- a/Engine/Plugins/Runtime/Metasound/Source/MetasoundEngine/Private/MetasoundWaveTableOscillatorNode.cpp
+++ b/Engine/Plugins/Runtime/Metasound/Source/MetasoundEngine/Private/MetasoundWaveTableOscillatorNode.cpp
@@ -85,7 +85,7 @@ namespace Metasound
 			
 			const FInputVertexInterfaceData& InputData = InParams.InputData;
 
-			FWaveTableReadRef InWaveTableReadRef = InputData.GetOrConstructDataReadReference<FWaveTable>("WaveTable");
+			FWaveTableReadRef InWaveTableReadRef = InputData.GetOrCreateDefaultDataReadReference<FWaveTable>("WaveTable", InParams.OperatorSettings);
 			FTriggerReadRef InPlayReadRef = InputData.GetOrCreateDefaultDataReadReference<FTrigger>("Play", InParams.OperatorSettings);
 			FTriggerReadRef InStopReadRef = InputData.GetOrCreateDefaultDataReadReference<FTrigger>("Stop", InParams.OperatorSettings);
 			FTriggerReadRef InSyncReadRef = InputData.GetOrCreateDefaultDataReadReference<FTrigger>("Sync", InParams.OperatorSettings);
diff --git a/Engine/Plugins/Runtime/Metasound/Source/MetasoundEngine/Private/MetasoundWaveTablePlayerNode.cpp b/Engine/Plugins/Runtime/Metasound/Source/MetasoundEngine/Private/MetasoundWaveTablePlayerNode.cpp
--- a/Engine/Plugins/Runtime/Metasound/Source/MetasoundEngine/Private/MetasoundWaveTablePlayerNode.cpp
+++ b/Engine/Plugins/Runtime/Metasound/Source/MetasoundEngine/Private/MetasoundWaveTablePlayerNode.cpp
@@ -111,8 +111,8 @@ namespace Metasound
 			
 			const FInputVertexInterfaceData& InputData = InParams.InputData;
 
-			FWaveTableBankAssetReadRef InWaveTableBankReadRef = InputData.GetOrConstructDataReadReference<FWaveTableBankAsset>("WaveTableBank");
-			FInt32ReadRef InIndexReadRef = InputData.GetOrConstructDataReadReference<int32>("Index");
+			FWaveTableBankAssetReadRef InWaveTableBankReadRef = InputData.GetOrCreateDefaultDataReadReference<FWaveTableBankAsset>("WaveTableBank", InParams.OperatorSettings);
+			FInt32ReadRef InIndexReadRef = InputData.GetOrCreateDefaultDataReadReference<int32>("Index", InParams.OperatorSettings);
 			FTriggerReadRef InPlayReadRef = InputData.GetOrCreateDefaultDataReadReference<FTrigger>("Play", InParams.OperatorSettings);
 			FTriggerReadRef InStopReadRef = InputData.GetOrCreateDefaultDataReadReference<FTrigger>("Stop", InParams.OperatorSettings);
 			FTriggerReadRef InSyncReadRef = InputData.GetOrCreateDefaultDataReadReference<FTrigger>("Sync", InParams.OperatorSettings);
diff --git a/Engine/Plugins/Runtime/Metasound/Source/MetasoundFrontend/Public/MetasoundAutoConverterNode.h b/Engine/Plugins/Runtime/Metasound/Source/MetasoundFrontend/Public/MetasoundAutoConverterNode.h
--- a/Engine/Plugins/Runtime/Metasound/Source/MetasoundFrontend/Public/MetasoundAutoConverterNode.h
+++ b/Engine/Plugins/Runtime/Metasound/Source/MetasoundFrontend/Public/MetasoundAutoConverterNode.h
@@ -202,32 +202,10 @@ namespace Metasound
 
 				virtual TUniquePtr<IOperator> CreateOperator(const FBuildOperatorParams& InParams, FBuildResults& OutResults) override
 				{
-					TDataWriteReference<ToDataType> WriteReference = TDataWriteReferenceFactory<ToDataType>::CreateAny(InParams.OperatorSettings);
+					TDataWriteReference<ToDataType> WriteReference = TDataWriteReferenceFactory<ToDataType>::CreateExplicitArgs(InParams.OperatorSettings);
+					TDataReadReference<FromDataType> ReadReference = InParams.InputData.GetOrCreateDefaultDataReadReference<FromDataType>(GetInputName(), InParams.OperatorSettings);
 
-					const FVertexName& InputName = GetInputName();
-					const bool bContainsRef = InParams.InputData.IsVertexBound(InputName);
-					if (bContainsRef)
-					{
-						TDataReadReference<FromDataType> ReadReference = InParams.InputData.GetDataReadReference<FromDataType>(InputName);
-						return MakeUnique<FConverterOperator>(ReadReference, WriteReference);
-					}
-
-					if constexpr (TIsParsable<FromDataType>::Value)
-					{
-						TDataReadReference<FromDataType> ReadReference = TDataReadReferenceFactory<FromDataType>::CreateAny(InParams.OperatorSettings);
-						return MakeUnique<FConverterOperator>(ReadReference, WriteReference);
-					}
-					else
-					{
-						// Converter node requires parsable reference if input not connected. Report as an error.
-						if (ensure(InParams.Node.GetVertexInterface().ContainsInputVertex(InputName)))
-						{
-							FInputDataDestination Dest(InParams.Node, InParams.Node.GetVertexInterface().GetInputVertex(GetInputName()));
-							AddBuildError<FMissingInputDataReferenceError>(OutResults.Errors, Dest);
-						}
-	
-						return TUniquePtr<IOperator>(nullptr);
-					}
+					return MakeUnique<FConverterOperator>(MoveTemp(ReadReference), MoveTemp(WriteReference));
 				}
 		};
 
diff --git a/Engine/Plugins/Runtime/Metasound/Source/MetasoundFrontend/Public/MetasoundReceiveNode.h b/Engine/Plugins/Runtime/Metasound/Source/MetasoundFrontend/Public/MetasoundReceiveNode.h
--- a/Engine/Plugins/Runtime/Metasound/Source/MetasoundFrontend/Public/MetasoundReceiveNode.h
+++ b/Engine/Plugins/Runtime/Metasound/Source/MetasoundFrontend/Public/MetasoundReceiveNode.h
@@ -181,7 +181,7 @@ namespace Metasound
 				return MakeUnique<TReceiverOperator>(
 					DefaultReadRef,
 					TDataWriteReferenceFactory<TDataType>::CreateAny(InParams.OperatorSettings, *DefaultReadRef),
-					InParams.InputData.GetOrConstructDataReadReference<FSendAddress>(METASOUND_GET_PARAM_NAME(AddressInput)),
+					InParams.InputData.GetOrCreateDefaultDataReadReference<FSendAddress>(METASOUND_GET_PARAM_NAME(AddressInput), InParams.OperatorSettings),
 					InParams.OperatorSettings
 					);
 			}
diff --git a/Engine/Plugins/Runtime/Metasound/Source/MetasoundFrontend/Public/MetasoundSendNode.h b/Engine/Plugins/Runtime/Metasound/Source/MetasoundFrontend/Public/MetasoundSendNode.h
--- a/Engine/Plugins/Runtime/Metasound/Source/MetasoundFrontend/Public/MetasoundSendNode.h
+++ b/Engine/Plugins/Runtime/Metasound/Source/MetasoundFrontend/Public/MetasoundSendNode.h
@@ -141,8 +141,9 @@ namespace Metasound
 
 				if (InParams.InputData.IsVertexBound(GetSendInputName()))
 				{
-					return MakeUnique<SendNodePrivate::TSendOperator<TDataType>>(InParams.InputData.GetDataReadReference<TDataType>(GetSendInputName()),
-						InParams.InputData.GetOrConstructDataReadReference<FSendAddress>(METASOUND_GET_PARAM_NAME(AddressInput)),
+					return MakeUnique<SendNodePrivate::TSendOperator<TDataType>>(
+						InParams.InputData.GetOrCreateDefaultDataReadReference<TDataType>(GetSendInputName(), InParams.OperatorSettings),
+						InParams.InputData.GetOrCreateDefaultDataReadReference<FSendAddress>(METASOUND_GET_PARAM_NAME(AddressInput), InParams.OperatorSettings),
 						InParams.OperatorSettings
 					);
 				}
diff --git a/Engine/Plugins/Runtime/Metasound/Source/MetasoundGenerator/Private/MetasoundGeneratorBuilder.cpp b/Engine/Plugins/Runtime/Metasound/Source/MetasoundGenerator/Private/MetasoundGeneratorBuilder.cpp
--- a/Engine/Plugins/Runtime/Metasound/Source/MetasoundGenerator/Private/MetasoundGeneratorBuilder.cpp
+++ b/Engine/Plugins/Runtime/Metasound/Source/MetasoundGenerator/Private/MetasoundGeneratorBuilder.cpp
@@ -319,15 +319,16 @@ namespace Metasound
 			VertexData.GetInputs() = InGraphOperatorAndInputs.Inputs;
 
 			// Get inputs
-			FTriggerWriteRef PlayTrigger = VertexData.GetInputs().GetOrConstructDataWriteReference<FTrigger>(SourceInterface::Inputs::OnPlay, InOperatorSettings, false);
+			FTriggerWriteRef PlayTrigger = VertexData.GetInputs().GetOrCreateDefaultDataWriteReference<FTrigger>(SourceInterface::Inputs::OnPlay, InOperatorSettings);
 
 			// Get outputs
 			TArray<FAudioBufferReadRef> OutputBuffers = FindOutputAudioBuffers(InInitParams.AudioOutputNames, VertexData, InOperatorSettings, InInitParams.MetaSoundName);
 			FTriggerReadRef FinishTrigger = TDataReadReferenceFactory<FTrigger>::CreateExplicitArgs(InOperatorSettings, false);
 
-			if (InInitParams.Graph->GetVertexInterface().GetOutputInterface().Contains(SourceOneShotInterface::Outputs::OnFinished))
+			if (const FAnyDataReference* FinishTriggerRef = VertexData.GetOutputs().FindDataReference(SourceOneShotInterface::Outputs::OnFinished))
 			{
-				FinishTrigger = VertexData.GetOutputs().GetOrConstructDataReadReference<FTrigger>(SourceOneShotInterface::Outputs::OnFinished, InOperatorSettings, false);
+
+				FinishTrigger = FinishTriggerRef->GetDataReadReference<FTrigger>();
 			}
 
 			
diff --git a/Engine/Plugins/Runtime/Metasound/Source/MetasoundGraphCore/Public/MetasoundVertexData.h b/Engine/Plugins/Runtime/Metasound/Source/MetasoundGraphCore/Public/MetasoundVertexData.h
--- a/Engine/Plugins/Runtime/Metasound/Source/MetasoundGraphCore/Public/MetasoundVertexData.h
+++ b/Engine/Plugins/Runtime/Metasound/Source/MetasoundGraphCore/Public/MetasoundVertexData.h
@@ -831,6 +831,7 @@ namespace Metasound
 		 * arguments.
 		 */
 		template<typename DataType, typename ... ConstructorArgTypes>
+		UE_DEPRECATED(5.6, "Please GetOrCreateDefaultDataReadReference instead.")
 		TDataReadReference<DataType> GetOrConstructDataReadReference(const FVertexName& InVertexName, ConstructorArgTypes&&... ConstructorArgs) const
 		{
 			if (const FInputBinding* Binding = Find(InVertexName))
@@ -878,6 +879,7 @@ namespace Metasound
 		 * arguments.
 		 */
 		template<typename DataType, typename ... ConstructorArgTypes>
+		UE_DEPRECATED(5.6, "Please GetOrCreateDefaultDataWriteReference instead.")
 		TDataWriteReference<DataType> GetOrConstructDataWriteReference(const FVertexName& InVertexName, ConstructorArgTypes&&... ConstructorArgs) const
 		{
 			if (const FInputBinding* Binding = Find(InVertexName))
diff --git a/Engine/Plugins/Runtime/Metasound/Source/MetasoundStandardNodes/Private/MetasoundADEnvelopeNode.cpp b/Engine/Plugins/Runtime/Metasound/Source/MetasoundStandardNodes/Private/MetasoundADEnvelopeNode.cpp
--- a/Engine/Plugins/Runtime/Metasound/Source/MetasoundStandardNodes/Private/MetasoundADEnvelopeNode.cpp
+++ b/Engine/Plugins/Runtime/Metasound/Source/MetasoundStandardNodes/Private/MetasoundADEnvelopeNode.cpp
@@ -291,7 +291,7 @@ namespace Metasound
 
 			const FInputVertexInterfaceData& InputData = InParams.InputData;
 
-			FTriggerReadRef TriggerIn = InputData.GetOrConstructDataReadReference<FTrigger>(METASOUND_GET_PARAM_NAME(InputTrigger), InParams.OperatorSettings);
+			FTriggerReadRef TriggerIn = InputData.GetOrCreateDefaultDataReadReference<FTrigger>(METASOUND_GET_PARAM_NAME(InputTrigger), InParams.OperatorSettings);
 			FTimeReadRef AttackTime = InputData.GetOrCreateDefaultDataReadReference<FTime>(METASOUND_GET_PARAM_NAME(InputAttackTime), InParams.OperatorSettings);
 			FTimeReadRef DecayTime = InputData.GetOrCreateDefaultDataReadReference<FTime>(METASOUND_GET_PARAM_NAME(InputDecayTime), InParams.OperatorSettings);
 			FFloatReadRef AttackCurveFactor = InputData.GetOrCreateDefaultDataReadReference<float>(METASOUND_GET_PARAM_NAME(InputAttackCurve), InParams.OperatorSettings);
diff --git a/Engine/Plugins/Runtime/Metasound/Source/MetasoundStandardNodes/Private/MetasoundADSREnvelopeNode.cpp b/Engine/Plugins/Runtime/Metasound/Source/MetasoundStandardNodes/Private/MetasoundADSREnvelopeNode.cpp
--- a/Engine/Plugins/Runtime/Metasound/Source/MetasoundStandardNodes/Private/MetasoundADSREnvelopeNode.cpp
+++ b/Engine/Plugins/Runtime/Metasound/Source/MetasoundStandardNodes/Private/MetasoundADSREnvelopeNode.cpp
@@ -587,8 +587,8 @@ namespace Metasound
 			FOperatorArgs Args 
 			{
 				InParams,
-				InputData.GetOrConstructDataReadReference<FTrigger>(METASOUND_GET_PARAM_NAME(InputAttackTrigger), InParams.OperatorSettings),
-				InputData.GetOrConstructDataReadReference<FTrigger>(METASOUND_GET_PARAM_NAME(InputReleaseTrigger), InParams.OperatorSettings),
+				InputData.GetOrCreateDefaultDataReadReference<FTrigger>(METASOUND_GET_PARAM_NAME(InputAttackTrigger), InParams.OperatorSettings),
+				InputData.GetOrCreateDefaultDataReadReference<FTrigger>(METASOUND_GET_PARAM_NAME(InputReleaseTrigger), InParams.OperatorSettings),
 				InputData.GetOrCreateDefaultDataReadReference<FTime>(METASOUND_GET_PARAM_NAME(InputAttackTime), InParams.OperatorSettings),
 				InputData.GetOrCreateDefaultDataReadReference<FTime>(METASOUND_GET_PARAM_NAME(InputDecayTime), InParams.OperatorSettings),
 				InputData.GetOrCreateDefaultDataReadReference<float>(METASOUND_GET_PARAM_NAME(InputSustainLevel), InParams.OperatorSettings),
diff --git a/Engine/Plugins/Runtime/Metasound/Source/MetasoundStandardNodes/Private/MetasoundAbsoluteValueNodes.cpp b/Engine/Plugins/Runtime/Metasound/Source/MetasoundStandardNodes/Private/MetasoundAbsoluteValueNodes.cpp
--- a/Engine/Plugins/Runtime/Metasound/Source/MetasoundStandardNodes/Private/MetasoundAbsoluteValueNodes.cpp
+++ b/Engine/Plugins/Runtime/Metasound/Source/MetasoundStandardNodes/Private/MetasoundAbsoluteValueNodes.cpp
@@ -118,7 +118,7 @@ namespace Metasound
 			{
 				using namespace AbsVertexNames;
 				const FInputVertexInterfaceData& InputData = InParams.InputData;
-				return InputData.GetOrConstructDataReadReference<FAudioBuffer>(METASOUND_GET_PARAM_NAME(InputValue), InParams.OperatorSettings);
+				return InputData.GetOrCreateDefaultDataReadReference<FAudioBuffer>(METASOUND_GET_PARAM_NAME(InputValue), InParams.OperatorSettings);
 			}
 		};
 	}
diff --git a/Engine/Plugins/Runtime/Metasound/Source/MetasoundStandardNodes/Private/MetasoundAdditiveSynthNode.cpp b/Engine/Plugins/Runtime/Metasound/Source/MetasoundStandardNodes/Private/MetasoundAdditiveSynthNode.cpp
--- a/Engine/Plugins/Runtime/Metasound/Source/MetasoundStandardNodes/Private/MetasoundAdditiveSynthNode.cpp
+++ b/Engine/Plugins/Runtime/Metasound/Source/MetasoundStandardNodes/Private/MetasoundAdditiveSynthNode.cpp
@@ -89,10 +89,10 @@ namespace Metasound
 			const FInputVertexInterfaceData& InputData = InParams.InputData;
 
 			FFloatReadRef InputBaseFrequency = InputData.GetOrCreateDefaultDataReadReference<float>(METASOUND_GET_PARAM_NAME(BaseFrequency), InParams.OperatorSettings);
-			FInputFloatArrayReadRef InputHarmonicMultipliers = InputData.GetOrConstructDataReadReference<FInputFloatArrayType>(METASOUND_GET_PARAM_NAME(HarmonicMultipliers));
-			FInputFloatArrayReadRef InputAmplitudes = InputData.GetOrConstructDataReadReference<FInputFloatArrayType>(METASOUND_GET_PARAM_NAME(Amplitudes));
-			FInputFloatArrayReadRef InputPhases = InputData.GetOrConstructDataReadReference<FInputFloatArrayType>(METASOUND_GET_PARAM_NAME(Phases));
-			FInputFloatArrayReadRef InputPanAmounts = InputData.GetOrConstructDataReadReference<FInputFloatArrayType>(METASOUND_GET_PARAM_NAME(PanAmounts));
+			FInputFloatArrayReadRef InputHarmonicMultipliers = InputData.GetOrCreateDefaultDataReadReference<FInputFloatArrayType>(METASOUND_GET_PARAM_NAME(HarmonicMultipliers), InParams.OperatorSettings);
+			FInputFloatArrayReadRef InputAmplitudes = InputData.GetOrCreateDefaultDataReadReference<FInputFloatArrayType>(METASOUND_GET_PARAM_NAME(Amplitudes), InParams.OperatorSettings);
+			FInputFloatArrayReadRef InputPhases = InputData.GetOrCreateDefaultDataReadReference<FInputFloatArrayType>(METASOUND_GET_PARAM_NAME(Phases), InParams.OperatorSettings);
+			FInputFloatArrayReadRef InputPanAmounts = InputData.GetOrCreateDefaultDataReadReference<FInputFloatArrayType>(METASOUND_GET_PARAM_NAME(PanAmounts), InParams.OperatorSettings);
 
 			return MakeUnique<FAdditiveSynthOperator>(InParams.OperatorSettings, InputBaseFrequency, InputHarmonicMultipliers, InputAmplitudes, InputPhases, InputPanAmounts);
 		}
diff --git a/Engine/Plugins/Runtime/Metasound/Source/MetasoundStandardNodes/Private/MetasoundBandsplitterNode.cpp b/Engine/Plugins/Runtime/Metasound/Source/MetasoundStandardNodes/Private/MetasoundBandsplitterNode.cpp
--- a/Engine/Plugins/Runtime/Metasound/Source/MetasoundStandardNodes/Private/MetasoundBandsplitterNode.cpp
+++ b/Engine/Plugins/Runtime/Metasound/Source/MetasoundStandardNodes/Private/MetasoundBandsplitterNode.cpp
@@ -236,7 +236,7 @@ namespace Metasound
 
 			for (uint32 Chan = 0; Chan < NumChannels; Chan++)
 			{
-				InputBuffers.Add(InputData.GetOrConstructDataReadReference<FAudioBuffer>(AudioInputNames[Chan], InParams.OperatorSettings));
+				InputBuffers.Add(InputData.GetOrCreateDefaultDataReadReference<FAudioBuffer>(AudioInputNames[Chan], InParams.OperatorSettings));
 			}
 
 			for (uint32 Band = 0; Band < NumBands - 1; Band++)
diff --git a/Engine/Plugins/Runtime/Metasound/Source/MetasoundStandardNodes/Private/MetasoundBasicFilters.cpp b/Engine/Plugins/Runtime/Metasound/Source/MetasoundStandardNodes/Private/MetasoundBasicFilters.cpp
--- a/Engine/Plugins/Runtime/Metasound/Source/MetasoundStandardNodes/Private/MetasoundBasicFilters.cpp
+++ b/Engine/Plugins/Runtime/Metasound/Source/MetasoundStandardNodes/Private/MetasoundBasicFilters.cpp
@@ -203,7 +203,7 @@ namespace Metasound
 		const FInputVertexInterfaceData& InputData = InParams.InputData;
 
 		// inputs
-		FAudioBufferReadRef AudioIn = InputData.GetOrConstructDataReadReference<FAudioBuffer>(METASOUND_GET_PARAM_NAME(ParamAudioInput), InParams.OperatorSettings);
+		FAudioBufferReadRef AudioIn = InputData.GetOrCreateDefaultDataReadReference<FAudioBuffer>(METASOUND_GET_PARAM_NAME(ParamAudioInput), InParams.OperatorSettings);
 		FFloatReadRef FrequencyIn = InputData.GetOrCreateDefaultDataReadReference<float>(METASOUND_GET_PARAM_NAME(ParamCutoffFrequency), InParams.OperatorSettings);
 		FFloatReadRef ResonanceIn = InputData.GetOrCreateDefaultDataReadReference<float>(METASOUND_GET_PARAM_NAME(ParamResonance), InParams.OperatorSettings);
 
@@ -392,7 +392,7 @@ namespace Metasound
 		const FInputVertexInterfaceData& InputData = InParams.InputData;
 
 		// inputs
-		FAudioBufferReadRef AudioIn = InputData.GetOrConstructDataReadReference<FAudioBuffer>(METASOUND_GET_PARAM_NAME(ParamAudioInput), InParams.OperatorSettings);
+		FAudioBufferReadRef AudioIn = InputData.GetOrCreateDefaultDataReadReference<FAudioBuffer>(METASOUND_GET_PARAM_NAME(ParamAudioInput), InParams.OperatorSettings);
 		FFloatReadRef FrequencyIn = InputData.GetOrCreateDefaultDataReadReference<float>(METASOUND_GET_PARAM_NAME(ParamCutoffFrequency), InParams.OperatorSettings);
 		FFloatReadRef ResonanceIn = InputData.GetOrCreateDefaultDataReadReference<float>(METASOUND_GET_PARAM_NAME(ParamResonance), InParams.OperatorSettings);
 		FFloatReadRef PassBandGainCompensationIn = InputData.GetOrCreateDefaultDataReadReference<float>(METASOUND_GET_PARAM_NAME(ParamBandStopControl), InParams.OperatorSettings);
@@ -566,7 +566,7 @@ namespace Metasound
 		const FInputVertexInterfaceData& InputData = InParams.InputData;
 
 		// inputs
-		FAudioBufferReadRef AudioIn = InputData.GetOrConstructDataReadReference<FAudioBuffer>(METASOUND_GET_PARAM_NAME(ParamAudioInput), InParams.OperatorSettings);
+		FAudioBufferReadRef AudioIn = InputData.GetOrCreateDefaultDataReadReference<FAudioBuffer>(METASOUND_GET_PARAM_NAME(ParamAudioInput), InParams.OperatorSettings);
 		FFloatReadRef FrequencyIn = InputData.GetOrCreateDefaultDataReadReference<float>(METASOUND_GET_PARAM_NAME(ParamCutoffFrequency), InParams.OperatorSettings);
 
 		return MakeUnique<FOnePoleLowPassFilterOperator>(
@@ -701,7 +701,7 @@ namespace Metasound
 		const FInputVertexInterfaceData& InputData = InParams.InputData;
 
 		// inputs
-		FAudioBufferReadRef AudioIn = InputData.GetOrConstructDataReadReference<FAudioBuffer>(METASOUND_GET_PARAM_NAME(ParamAudioInput), InParams.OperatorSettings);
+		FAudioBufferReadRef AudioIn = InputData.GetOrCreateDefaultDataReadReference<FAudioBuffer>(METASOUND_GET_PARAM_NAME(ParamAudioInput), InParams.OperatorSettings);
 		FFloatReadRef FrequencyIn = InputData.GetOrCreateDefaultDataReadReference<float>(METASOUND_GET_PARAM_NAME(ParamCutoffFrequency), InParams.OperatorSettings);
 
 		return MakeUnique<FOnePoleHighPassFilterOperator>(
@@ -873,7 +873,7 @@ namespace Metasound
 		const FInputVertexInterfaceData& InputData = InParams.InputData;
 
 		// inputs
-		FAudioBufferReadRef AudioIn = InputData.GetOrConstructDataReadReference<FAudioBuffer>(METASOUND_GET_PARAM_NAME(ParamAudioInput), InParams.OperatorSettings);
+		FAudioBufferReadRef AudioIn = InputData.GetOrCreateDefaultDataReadReference<FAudioBuffer>(METASOUND_GET_PARAM_NAME(ParamAudioInput), InParams.OperatorSettings);
 		FFloatReadRef FrequencyIn = InputData.GetOrCreateDefaultDataReadReference<float>(METASOUND_GET_PARAM_NAME(ParamCutoffFrequency), InParams.OperatorSettings);
 		FFloatReadRef BandwidthIn = InputData.GetOrCreateDefaultDataReadReference<float>(METASOUND_GET_PARAM_NAME(ParamBandwidth), InParams.OperatorSettings);
 		FFloatReadRef FilterGainDbIn = InputData.GetOrCreateDefaultDataReadReference<float>(METASOUND_GET_PARAM_NAME(ParamGainDb), InParams.OperatorSettings);
diff --git a/Engine/Plugins/Runtime/Metasound/Source/MetasoundStandardNodes/Private/MetasoundBitcrusherNode.cpp b/Engine/Plugins/Runtime/Metasound/Source/MetasoundStandardNodes/Private/MetasoundBitcrusherNode.cpp
--- a/Engine/Plugins/Runtime/Metasound/Source/MetasoundStandardNodes/Private/MetasoundBitcrusherNode.cpp
+++ b/Engine/Plugins/Runtime/Metasound/Source/MetasoundStandardNodes/Private/MetasoundBitcrusherNode.cpp
@@ -114,7 +114,7 @@ namespace Metasound
 			using namespace BitcrusherVertexNames;
 			const FInputVertexInterfaceData& InputData = InParams.InputData;
 
-			FAudioBufferReadRef AudioIn = InputData.GetOrConstructDataReadReference<FAudioBuffer>(METASOUND_GET_PARAM_NAME(InputAudio), InParams.OperatorSettings);
+			FAudioBufferReadRef AudioIn = InputData.GetOrCreateDefaultDataReadReference<FAudioBuffer>(METASOUND_GET_PARAM_NAME(InputAudio), InParams.OperatorSettings);
 			FFloatReadRef SampleRateIn = InputData.GetOrCreateDefaultDataReadReference<float>(METASOUND_GET_PARAM_NAME(InputSampleRate), InParams.OperatorSettings);
 			FFloatReadRef BitDepthIn = InputData.GetOrCreateDefaultDataReadReference<float>(METASOUND_GET_PARAM_NAME(InputBitDepth), InParams.OperatorSettings);
 
diff --git a/Engine/Plugins/Runtime/Metasound/Source/MetasoundStandardNodes/Private/MetasoundClampNode.cpp b/Engine/Plugins/Runtime/Metasound/Source/MetasoundStandardNodes/Private/MetasoundClampNode.cpp
--- a/Engine/Plugins/Runtime/Metasound/Source/MetasoundStandardNodes/Private/MetasoundClampNode.cpp
+++ b/Engine/Plugins/Runtime/Metasound/Source/MetasoundStandardNodes/Private/MetasoundClampNode.cpp
@@ -130,21 +130,21 @@ namespace Metasound
 			{
 				using namespace ClampVertexNames;
 				const FInputVertexInterfaceData& InputData = InParams.InputData;
-				return InputData.GetOrConstructDataReadReference<FAudioBuffer>(METASOUND_GET_PARAM_NAME(InputValue), InParams.OperatorSettings);
+				return InputData.GetOrCreateDefaultDataReadReference<FAudioBuffer>(METASOUND_GET_PARAM_NAME(InputValue), InParams.OperatorSettings);
 			}
 
 			static TDataReadReference<FAudioBuffer> CreateInMin(const FBuildOperatorParams& InParams)
 			{
 				using namespace ClampVertexNames;
 				const FInputVertexInterfaceData& InputData = InParams.InputData;
-				return InputData.GetOrConstructDataReadReference<FAudioBuffer>(METASOUND_GET_PARAM_NAME(InputMinValue), InParams.OperatorSettings);
+				return InputData.GetOrCreateDefaultDataReadReference<FAudioBuffer>(METASOUND_GET_PARAM_NAME(InputMinValue), InParams.OperatorSettings);
 			}
 
 			static TDataReadReference<FAudioBuffer> CreateInMax(const FBuildOperatorParams& InParams)
 			{
 				using namespace ClampVertexNames;
 				const FInputVertexInterfaceData& InputData = InParams.InputData;
-				return InputData.GetOrConstructDataReadReference<FAudioBuffer>(METASOUND_GET_PARAM_NAME(InputMaxValue), InParams.OperatorSettings);
+				return InputData.GetOrCreateDefaultDataReadReference<FAudioBuffer>(METASOUND_GET_PARAM_NAME(InputMaxValue), InParams.OperatorSettings);
 			}
 		};
 	}
diff --git a/Engine/Plugins/Runtime/Metasound/Source/MetasoundStandardNodes/Private/MetasoundCompressorNode.cpp b/Engine/Plugins/Runtime/Metasound/Source/MetasoundStandardNodes/Private/MetasoundCompressorNode.cpp
--- a/Engine/Plugins/Runtime/Metasound/Source/MetasoundStandardNodes/Private/MetasoundCompressorNode.cpp
+++ b/Engine/Plugins/Runtime/Metasound/Source/MetasoundStandardNodes/Private/MetasoundCompressorNode.cpp
@@ -214,14 +214,14 @@ namespace Metasound
 			const FInputVertexInterfaceData& InputData = InParams.InputData;
 
 			FBoolReadRef bIsBypassedIn = InputData.GetOrCreateDefaultDataReadReference<bool>(METASOUND_GET_PARAM_NAME(InputIsBypassed), InParams.OperatorSettings);
-			FAudioBufferReadRef AudioIn = InputData.GetOrConstructDataReadReference<FAudioBuffer>(METASOUND_GET_PARAM_NAME(InputAudio), InParams.OperatorSettings);
+			FAudioBufferReadRef AudioIn = InputData.GetOrCreateDefaultDataReadReference<FAudioBuffer>(METASOUND_GET_PARAM_NAME(InputAudio), InParams.OperatorSettings);
 			FFloatReadRef RatioIn = InputData.GetOrCreateDefaultDataReadReference<float>(METASOUND_GET_PARAM_NAME(InputRatio), InParams.OperatorSettings);
 			FFloatReadRef ThresholdDbIn = InputData.GetOrCreateDefaultDataReadReference<float>(METASOUND_GET_PARAM_NAME(InputThreshold), InParams.OperatorSettings);
 			FTimeReadRef AttackTimeIn = InputData.GetOrCreateDefaultDataReadReference<FTime>(METASOUND_GET_PARAM_NAME(InputAttackTime), InParams.OperatorSettings);
 			FTimeReadRef ReleaseTimeIn = InputData.GetOrCreateDefaultDataReadReference<FTime>(METASOUND_GET_PARAM_NAME(InputReleaseTime), InParams.OperatorSettings);
 			FTimeReadRef LookaheadTimeIn = InputData.GetOrCreateDefaultDataReadReference<FTime>(METASOUND_GET_PARAM_NAME(InputLookaheadTime), InParams.OperatorSettings);
 			FFloatReadRef KneeIn = InputData.GetOrCreateDefaultDataReadReference<float>(METASOUND_GET_PARAM_NAME(InputKnee), InParams.OperatorSettings);
-			FAudioBufferReadRef SidechainIn = InputData.GetOrConstructDataReadReference<FAudioBuffer>(METASOUND_GET_PARAM_NAME(InputSidechain), InParams.OperatorSettings);
+			FAudioBufferReadRef SidechainIn = InputData.GetOrCreateDefaultDataReadReference<FAudioBuffer>(METASOUND_GET_PARAM_NAME(InputSidechain), InParams.OperatorSettings);
 			FEnvelopePeakModeReadRef EnvelopeModeIn = InputData.GetOrCreateDefaultDataReadReference<FEnumEnvelopePeakMode>(METASOUND_GET_PARAM_NAME(InputEnvelopeMode), InParams.OperatorSettings);
 			FBoolReadRef bIsAnalogIn = InputData.GetOrCreateDefaultDataReadReference<bool>(METASOUND_GET_PARAM_NAME(InputIsAnalog), InParams.OperatorSettings);
 			FBoolReadRef bIsUpwardsIn = InputData.GetOrCreateDefaultDataReadReference<bool>(METASOUND_GET_PARAM_NAME(InputIsUpwards), InParams.OperatorSettings);
diff --git a/Engine/Plugins/Runtime/Metasound/Source/MetasoundStandardNodes/Private/MetasoundConversionNodes.cpp b/Engine/Plugins/Runtime/Metasound/Source/MetasoundStandardNodes/Private/MetasoundConversionNodes.cpp
--- a/Engine/Plugins/Runtime/Metasound/Source/MetasoundStandardNodes/Private/MetasoundConversionNodes.cpp
+++ b/Engine/Plugins/Runtime/Metasound/Source/MetasoundStandardNodes/Private/MetasoundConversionNodes.cpp
@@ -130,7 +130,7 @@ namespace Metasound
 			{
 				using namespace ConversionNodeVertexNames;
 				const FInputVertexInterfaceData& InputData = InParams.InputData;
-				return InputData.GetOrConstructDataReadReference<FAudioBuffer>(METASOUND_GET_PARAM_NAME(InputValue), InParams.OperatorSettings);
+				return InputData.GetOrCreateDefaultDataReadReference<FAudioBuffer>(METASOUND_GET_PARAM_NAME(InputValue), InParams.OperatorSettings);
 			}
 
 			static void GetConvertedValue(const FAudioBuffer& InValue, float& OutValue)
diff --git a/Engine/Plugins/Runtime/Metasound/Source/MetasoundStandardNodes/Private/MetasoundDebugLogNode.cpp b/Engine/Plugins/Runtime/Metasound/Source/MetasoundStandardNodes/Private/MetasoundDebugLogNode.cpp
--- a/Engine/Plugins/Runtime/Metasound/Source/MetasoundStandardNodes/Private/MetasoundDebugLogNode.cpp
+++ b/Engine/Plugins/Runtime/Metasound/Source/MetasoundStandardNodes/Private/MetasoundDebugLogNode.cpp
@@ -93,7 +93,7 @@ namespace Metasound
 
 				const FInputVertexInterfaceData& InputData = InParams.InputData;
 
-				FTriggerReadRef Trigger = InputData.GetOrConstructDataReadReference<FTrigger>(METASOUND_GET_PARAM_NAME(InputTrigger), InParams.OperatorSettings);
+				FTriggerReadRef Trigger = InputData.GetOrCreateDefaultDataReadReference<FTrigger>(METASOUND_GET_PARAM_NAME(InputTrigger), InParams.OperatorSettings);
 
 				TDataReadReference<FString> Label = InputData.GetOrCreateDefaultDataReadReference<FString>(METASOUND_GET_PARAM_NAME(InputLabel), InParams.OperatorSettings);
 				TDataReadReference<PrintLogType> ValueToLog = InputData.GetOrCreateDefaultDataReadReference<PrintLogType>(METASOUND_GET_PARAM_NAME(InputValueToLog), InParams.OperatorSettings);
diff --git a/Engine/Plugins/Runtime/Metasound/Source/MetasoundStandardNodes/Private/MetasoundDebugTriggerOnNanNode.cpp b/Engine/Plugins/Runtime/Metasound/Source/MetasoundStandardNodes/Private/MetasoundDebugTriggerOnNanNode.cpp
--- a/Engine/Plugins/Runtime/Metasound/Source/MetasoundStandardNodes/Private/MetasoundDebugTriggerOnNanNode.cpp
+++ b/Engine/Plugins/Runtime/Metasound/Source/MetasoundStandardNodes/Private/MetasoundDebugTriggerOnNanNode.cpp
@@ -110,7 +110,7 @@ namespace Metasound
 
 		const FInputVertexInterfaceData& InputData = InParams.InputData;
 
-		FAudioBufferReadRef AudioInput = InputData.GetOrConstructDataReadReference<FAudioBuffer>(METASOUND_GET_PARAM_NAME(AudioInput), InParams.OperatorSettings);
+		FAudioBufferReadRef AudioInput = InputData.GetOrCreateDefaultDataReadReference<FAudioBuffer>(METASOUND_GET_PARAM_NAME(AudioInput), InParams.OperatorSettings);
 		FBoolReadRef bTriggerOnce = InputData.GetOrCreateDefaultDataReadReference<bool>(METASOUND_GET_PARAM_NAME(TriggerOnce), InParams.OperatorSettings);
 
 		return MakeUnique<FTriggerOnNanOperator>(InParams.OperatorSettings, AudioInput, bTriggerOnce);
diff --git a/Engine/Plugins/Runtime/Metasound/Source/MetasoundStandardNodes/Private/MetasoundDelayNode.cpp b/Engine/Plugins/Runtime/Metasound/Source/MetasoundStandardNodes/Private/MetasoundDelayNode.cpp
--- a/Engine/Plugins/Runtime/Metasound/Source/MetasoundStandardNodes/Private/MetasoundDelayNode.cpp
+++ b/Engine/Plugins/Runtime/Metasound/Source/MetasoundStandardNodes/Private/MetasoundDelayNode.cpp
@@ -86,7 +86,7 @@ namespace Metasound
 			static TDataReadReference<FAudioBuffer> CreateInRef(const FBuildOperatorParams& InParams)
 			{
 				const FInputVertexInterfaceData& InputData = InParams.InputData;
-				return InputData.GetOrConstructDataReadReference<FAudioBuffer>(METASOUND_GET_PARAM_NAME(InParamDelayTime), InParams.OperatorSettings);
+				return InputData.GetOrCreateDefaultDataReadReference<FAudioBuffer>(METASOUND_GET_PARAM_NAME(InParamDelayTime), InParams.OperatorSettings);
 			}
 
 			static bool IsAudioType() { return true; }
@@ -154,13 +154,13 @@ namespace Metasound
 
 			const FInputVertexInterfaceData& InputData = InParams.InputData;
 
-			FAudioBufferReadRef AudioIn = InputData.GetOrConstructDataReadReference<FAudioBuffer>(METASOUND_GET_PARAM_NAME(InParamAudioInput), InParams.OperatorSettings);
+			FAudioBufferReadRef AudioIn = InputData.GetOrCreateDefaultDataReadReference<FAudioBuffer>(METASOUND_GET_PARAM_NAME(InParamAudioInput), InParams.OperatorSettings);
 			TDataReadReference<ValueType> DelayTime = TDelay<ValueType>::CreateInRef(InParams);
 			FFloatReadRef DryLevel = InputData.GetOrCreateDefaultDataReadReference<float>(METASOUND_GET_PARAM_NAME(InParamDryLevel), InParams.OperatorSettings);
 			FFloatReadRef WetLevel = InputData.GetOrCreateDefaultDataReadReference<float>(METASOUND_GET_PARAM_NAME(InParamWetLevel), InParams.OperatorSettings);
 			FFloatReadRef Feedback = InputData.GetOrCreateDefaultDataReadReference<float>(METASOUND_GET_PARAM_NAME(InParamFeedbackAmount), InParams.OperatorSettings);
 			FTime MaxDelayTime = InputData.GetOrCreateDefaultValue<FTime>(METASOUND_GET_PARAM_NAME(InParamMaxDelayTime), InParams.OperatorSettings);
-			FTriggerReadRef TriggerReset = InputData.GetOrConstructDataReadReference<FTrigger>(METASOUND_GET_PARAM_NAME(InputResetDelay), InParams.OperatorSettings);
+			FTriggerReadRef TriggerReset = InputData.GetOrCreateDefaultDataReadReference<FTrigger>(METASOUND_GET_PARAM_NAME(InputResetDelay), InParams.OperatorSettings);
 
 			return MakeUnique<TDelayNodeOperator<ValueType>>(InParams, AudioIn, DelayTime, DryLevel, WetLevel, Feedback, MaxDelayTime.GetSeconds(), TriggerReset);
 		}
diff --git a/Engine/Plugins/Runtime/Metasound/Source/MetasoundStandardNodes/Private/MetasoundDiffuserNode.cpp b/Engine/Plugins/Runtime/Metasound/Source/MetasoundStandardNodes/Private/MetasoundDiffuserNode.cpp
--- a/Engine/Plugins/Runtime/Metasound/Source/MetasoundStandardNodes/Private/MetasoundDiffuserNode.cpp
+++ b/Engine/Plugins/Runtime/Metasound/Source/MetasoundStandardNodes/Private/MetasoundDiffuserNode.cpp
@@ -81,7 +81,7 @@ namespace Metasound
 			
 			const FInputVertexInterfaceData& InputData = InParams.InputData;
 
-			FAudioBufferReadRef AudioIn = InputData.GetOrConstructDataReadReference<FAudioBuffer>(METASOUND_GET_PARAM_NAME(InputAudio), InParams.OperatorSettings);
+			FAudioBufferReadRef AudioIn = InputData.GetOrCreateDefaultDataReadReference<FAudioBuffer>(METASOUND_GET_PARAM_NAME(InputAudio), InParams.OperatorSettings);
 			FInt32ReadRef DepthIn = InputData.GetOrCreateDefaultDataReadReference<int32>(METASOUND_GET_PARAM_NAME(InputDiffusionDepth), InParams.OperatorSettings);
 			FFloatReadRef FeedbackIn = InputData.GetOrCreateDefaultDataReadReference<float>(METASOUND_GET_PARAM_NAME(InputFeedbackGain), InParams.OperatorSettings);
 
diff --git a/Engine/Plugins/Runtime/Metasound/Source/MetasoundStandardNodes/Private/MetasoundDopplerPitchShiftNode.cpp b/Engine/Plugins/Runtime/Metasound/Source/MetasoundStandardNodes/Private/MetasoundDopplerPitchShiftNode.cpp
--- a/Engine/Plugins/Runtime/Metasound/Source/MetasoundStandardNodes/Private/MetasoundDopplerPitchShiftNode.cpp
+++ b/Engine/Plugins/Runtime/Metasound/Source/MetasoundStandardNodes/Private/MetasoundDopplerPitchShiftNode.cpp
@@ -133,7 +133,7 @@ namespace Metasound
 
 		const FInputVertexInterfaceData& InputData = InParams.InputData;
 
-		FAudioBufferReadRef AudioIn = InputData.GetOrConstructDataReadReference<FAudioBuffer>(METASOUND_GET_PARAM_NAME(InParamAudioInput), InParams.OperatorSettings);
+		FAudioBufferReadRef AudioIn = InputData.GetOrCreateDefaultDataReadReference<FAudioBuffer>(METASOUND_GET_PARAM_NAME(InParamAudioInput), InParams.OperatorSettings);
 		FFloatReadRef PitchShift = InputData.GetOrCreateDefaultDataReadReference<float>(METASOUND_GET_PARAM_NAME(InParamPitchShift), InParams.OperatorSettings);
 		FFloatReadRef DelayLength = InputData.GetOrCreateDefaultDataReadReference<float>(METASOUND_GET_PARAM_NAME(InParamDelayLength), InParams.OperatorSettings);
 
diff --git a/Engine/Plugins/Runtime/Metasound/Source/MetasoundStandardNodes/Private/MetasoundDynamicFilterNode.cpp b/Engine/Plugins/Runtime/Metasound/Source/MetasoundStandardNodes/Private/MetasoundDynamicFilterNode.cpp
--- a/Engine/Plugins/Runtime/Metasound/Source/MetasoundStandardNodes/Private/MetasoundDynamicFilterNode.cpp
+++ b/Engine/Plugins/Runtime/Metasound/Source/MetasoundStandardNodes/Private/MetasoundDynamicFilterNode.cpp
@@ -191,8 +191,8 @@ namespace Metasound
 			
 			const FInputVertexInterfaceData& InputData = InParams.InputData;
 
-			FAudioBufferReadRef AudioIn = InputData.GetOrConstructDataReadReference<FAudioBuffer>(METASOUND_GET_PARAM_NAME(InputAudio), InParams.OperatorSettings);
-			FAudioBufferReadRef SidechainIn = InputData.GetOrConstructDataReadReference<FAudioBuffer>(METASOUND_GET_PARAM_NAME(InputSidechain), InParams.OperatorSettings);
+			FAudioBufferReadRef AudioIn = InputData.GetOrCreateDefaultDataReadReference<FAudioBuffer>(METASOUND_GET_PARAM_NAME(InputAudio), InParams.OperatorSettings);
+			FAudioBufferReadRef SidechainIn = InputData.GetOrCreateDefaultDataReadReference<FAudioBuffer>(METASOUND_GET_PARAM_NAME(InputSidechain), InParams.OperatorSettings);
 			FEnumEDynamicFilterReadRef FilterTypeIn = InputData.GetOrCreateDefaultDataReadReference<FEnumEDynamicFilterType>(METASOUND_GET_PARAM_NAME(InputFilterType), InParams.OperatorSettings);
 			FFloatReadRef FrequencyIn = InputData.GetOrCreateDefaultDataReadReference<float>(METASOUND_GET_PARAM_NAME(InputFrequency), InParams.OperatorSettings);
 			FFloatReadRef QIn = InputData.GetOrCreateDefaultDataReadReference<float>(METASOUND_GET_PARAM_NAME(InputQ), InParams.OperatorSettings);
diff --git a/Engine/Plugins/Runtime/Metasound/Source/MetasoundStandardNodes/Private/MetasoundEnvelopeFollowerNode.cpp b/Engine/Plugins/Runtime/Metasound/Source/MetasoundStandardNodes/Private/MetasoundEnvelopeFollowerNode.cpp
--- a/Engine/Plugins/Runtime/Metasound/Source/MetasoundStandardNodes/Private/MetasoundEnvelopeFollowerNode.cpp
+++ b/Engine/Plugins/Runtime/Metasound/Source/MetasoundStandardNodes/Private/MetasoundEnvelopeFollowerNode.cpp
@@ -254,7 +254,7 @@ namespace Metasound
 		const FInputVertexInterfaceData& InputData = InParams.InputData;
 
 		FBoolReadRef EnableIn = InputData.GetOrCreateDefaultDataReadReference<bool>(METASOUND_GET_PARAM_NAME(InParamEnable), InParams.OperatorSettings);
-		FAudioBufferReadRef AudioIn = InputData.GetOrConstructDataReadReference<FAudioBuffer>(METASOUND_GET_PARAM_NAME(InParamAudioInput), InParams.OperatorSettings);
+		FAudioBufferReadRef AudioIn = InputData.GetOrCreateDefaultDataReadReference<FAudioBuffer>(METASOUND_GET_PARAM_NAME(InParamAudioInput), InParams.OperatorSettings);
 		FTimeReadRef AttackTime = InputData.GetOrCreateDefaultDataReadReference<FTime>(METASOUND_GET_PARAM_NAME(InParamAttackTime), InParams.OperatorSettings);
 		FTimeReadRef ReleaseTime = InputData.GetOrCreateDefaultDataReadReference<FTime>(METASOUND_GET_PARAM_NAME(InParamReleaseTime), InParams.OperatorSettings);
 		FEnvelopePeakModeReadRef EnvelopeModeIn = InputData.GetOrCreateDefaultDataReadReference<FEnumEnvelopePeakMode>(METASOUND_GET_PARAM_NAME(InParamFollowMode), InParams.OperatorSettings);
diff --git a/Engine/Plugins/Runtime/Metasound/Source/MetasoundStandardNodes/Private/MetasoundFlangerNode.cpp b/Engine/Plugins/Runtime/Metasound/Source/MetasoundStandardNodes/Private/MetasoundFlangerNode.cpp
--- a/Engine/Plugins/Runtime/Metasound/Source/MetasoundStandardNodes/Private/MetasoundFlangerNode.cpp
+++ b/Engine/Plugins/Runtime/Metasound/Source/MetasoundStandardNodes/Private/MetasoundFlangerNode.cpp
@@ -84,7 +84,7 @@ namespace Metasound
 			
 			const FInputVertexInterfaceData& InputData = InParams.InputData;
 
-			FAudioBufferReadRef AudioInput = InputData.GetOrConstructDataReadReference<FAudioBuffer>(METASOUND_GET_PARAM_NAME(InputAudio), InParams.OperatorSettings);
+			FAudioBufferReadRef AudioInput = InputData.GetOrCreateDefaultDataReadReference<FAudioBuffer>(METASOUND_GET_PARAM_NAME(InputAudio), InParams.OperatorSettings);
 			FFloatReadRef ModulationRate = InputData.GetOrCreateDefaultDataReadReference<float>(METASOUND_GET_PARAM_NAME(InputModulationRate), InParams.OperatorSettings);
 			FFloatReadRef ModulationDepth = InputData.GetOrCreateDefaultDataReadReference<float>(METASOUND_GET_PARAM_NAME(InputModulationDepth), InParams.OperatorSettings);
 			FFloatReadRef CenterDelay = InputData.GetOrCreateDefaultDataReadReference<float>(METASOUND_GET_PARAM_NAME(InputCenterDelay), InParams.OperatorSettings);
diff --git a/Engine/Plugins/Runtime/Metasound/Source/MetasoundStandardNodes/Private/MetasoundGrainDelayNode.cpp b/Engine/Plugins/Runtime/Metasound/Source/MetasoundStandardNodes/Private/MetasoundGrainDelayNode.cpp
--- a/Engine/Plugins/Runtime/Metasound/Source/MetasoundStandardNodes/Private/MetasoundGrainDelayNode.cpp
+++ b/Engine/Plugins/Runtime/Metasound/Source/MetasoundStandardNodes/Private/MetasoundGrainDelayNode.cpp
@@ -202,8 +202,8 @@ namespace Metasound
 
 		const FInputVertexInterfaceData& InputData = InParams.InputData;
 		
-		FAudioBufferReadRef AudioInput = InputData.GetOrConstructDataReadReference<FAudioBuffer>(METASOUND_GET_PARAM_NAME(InAudio), InParams.OperatorSettings);
-		FTriggerReadRef GrainSpawnTrigger = InputData.GetOrConstructDataReadReference<FTrigger>(METASOUND_GET_PARAM_NAME(InGrainSpawnTrigger), InParams.OperatorSettings);
+		FAudioBufferReadRef AudioInput = InputData.GetOrCreateDefaultDataReadReference<FAudioBuffer>(METASOUND_GET_PARAM_NAME(InAudio), InParams.OperatorSettings);
+		FTriggerReadRef GrainSpawnTrigger = InputData.GetOrCreateDefaultDataReadReference<FTrigger>(METASOUND_GET_PARAM_NAME(InGrainSpawnTrigger), InParams.OperatorSettings);
 		FFloatReadRef GrainDelay = InputData.GetOrCreateDefaultDataReadReference<float>(METASOUND_GET_PARAM_NAME(InGrainDelay), InParams.OperatorSettings);
 		FFloatReadRef GrainDelayDelta = InputData.GetOrCreateDefaultDataReadReference<float>(METASOUND_GET_PARAM_NAME(InGrainDelayDelta), InParams.OperatorSettings);
 		FFloatReadRef GrainDuration = InputData.GetOrCreateDefaultDataReadReference<float>(METASOUND_GET_PARAM_NAME(InGrainDuration), InParams.OperatorSettings);
diff --git a/Engine/Plugins/Runtime/Metasound/Source/MetasoundStandardNodes/Private/MetasoundITDPannerNode.cpp b/Engine/Plugins/Runtime/Metasound/Source/MetasoundStandardNodes/Private/MetasoundITDPannerNode.cpp
--- a/Engine/Plugins/Runtime/Metasound/Source/MetasoundStandardNodes/Private/MetasoundITDPannerNode.cpp
+++ b/Engine/Plugins/Runtime/Metasound/Source/MetasoundStandardNodes/Private/MetasoundITDPannerNode.cpp
@@ -288,7 +288,7 @@ namespace Metasound
 
 		using namespace ITDPannerVertexNames;
 
-		FAudioBufferReadRef AudioIn = InputData.GetOrConstructDataReadReference<FAudioBuffer>(METASOUND_GET_PARAM_NAME(InputAudio), InParams.OperatorSettings);
+		FAudioBufferReadRef AudioIn = InputData.GetOrCreateDefaultDataReadReference<FAudioBuffer>(METASOUND_GET_PARAM_NAME(InputAudio), InParams.OperatorSettings);
 		FFloatReadRef PanningAngle = InputData.GetOrCreateDefaultDataReadReference<float>(METASOUND_GET_PARAM_NAME(InputPanAngle), InParams.OperatorSettings);
 		FFloatReadRef DistanceFactor = InputData.GetOrCreateDefaultDataReadReference<float>(METASOUND_GET_PARAM_NAME(InputDistanceFactor), InParams.OperatorSettings);
 		FFloatReadRef HeadWidth = InputData.GetOrCreateDefaultDataReadReference<float>(METASOUND_GET_PARAM_NAME(InputHeadWidth), InParams.OperatorSettings);
diff --git a/Engine/Plugins/Runtime/Metasound/Source/MetasoundStandardNodes/Private/MetasoundLimiterNode.cpp b/Engine/Plugins/Runtime/Metasound/Source/MetasoundStandardNodes/Private/MetasoundLimiterNode.cpp
--- a/Engine/Plugins/Runtime/Metasound/Source/MetasoundStandardNodes/Private/MetasoundLimiterNode.cpp
+++ b/Engine/Plugins/Runtime/Metasound/Source/MetasoundStandardNodes/Private/MetasoundLimiterNode.cpp
@@ -144,7 +144,7 @@ namespace Metasound
 			
 			const FInputVertexInterfaceData& InputData = InParams.InputData;
 
-			FAudioBufferReadRef AudioIn = InputData.GetOrConstructDataReadReference<FAudioBuffer>(METASOUND_GET_PARAM_NAME(InputAudio), InParams.OperatorSettings);
+			FAudioBufferReadRef AudioIn = InputData.GetOrCreateDefaultDataReadReference<FAudioBuffer>(METASOUND_GET_PARAM_NAME(InputAudio), InParams.OperatorSettings);
 			FFloatReadRef InGainDbIn = InputData.GetOrCreateDefaultDataReadReference<float>(METASOUND_GET_PARAM_NAME(InputInGainDb), InParams.OperatorSettings);
 			FFloatReadRef ThresholdDbIn = InputData.GetOrCreateDefaultDataReadReference<float>(METASOUND_GET_PARAM_NAME(InputThresholdDb), InParams.OperatorSettings);
 			FTimeReadRef ReleaseTimeIn = InputData.GetOrCreateDefaultDataReadReference<FTime>(METASOUND_GET_PARAM_NAME(InputReleaseTime), InParams.OperatorSettings);
diff --git a/Engine/Plugins/Runtime/Metasound/Source/MetasoundStandardNodes/Private/MetasoundMaxNode.cpp b/Engine/Plugins/Runtime/Metasound/Source/MetasoundStandardNodes/Private/MetasoundMaxNode.cpp
--- a/Engine/Plugins/Runtime/Metasound/Source/MetasoundStandardNodes/Private/MetasoundMaxNode.cpp
+++ b/Engine/Plugins/Runtime/Metasound/Source/MetasoundStandardNodes/Private/MetasoundMaxNode.cpp
@@ -118,14 +118,14 @@ namespace Metasound
 			{
 				using namespace MaxVertexNames;
 				const FInputVertexInterfaceData& InputData = InParams.InputData;
-				return InputData.GetOrConstructDataReadReference<FAudioBuffer>(METASOUND_GET_PARAM_NAME(InputAValue), InParams.OperatorSettings);
+				return InputData.GetOrCreateDefaultDataReadReference<FAudioBuffer>(METASOUND_GET_PARAM_NAME(InputAValue), InParams.OperatorSettings);
 			}
 
 			static TDataReadReference<FAudioBuffer> CreateInRefB(const FBuildOperatorParams& InParams)
 			{
 				using namespace MaxVertexNames;				
 				const FInputVertexInterfaceData& InputData = InParams.InputData;
-				return InputData.GetOrConstructDataReadReference<FAudioBuffer>(METASOUND_GET_PARAM_NAME(InputBValue), InParams.OperatorSettings);
+				return InputData.GetOrCreateDefaultDataReadReference<FAudioBuffer>(METASOUND_GET_PARAM_NAME(InputBValue), InParams.OperatorSettings);
 			}
 		};
 	}
diff --git a/Engine/Plugins/Runtime/Metasound/Source/MetasoundStandardNodes/Private/MetasoundMidSideNodes.cpp b/Engine/Plugins/Runtime/Metasound/Source/MetasoundStandardNodes/Private/MetasoundMidSideNodes.cpp
--- a/Engine/Plugins/Runtime/Metasound/Source/MetasoundStandardNodes/Private/MetasoundMidSideNodes.cpp
+++ b/Engine/Plugins/Runtime/Metasound/Source/MetasoundStandardNodes/Private/MetasoundMidSideNodes.cpp
@@ -128,8 +128,8 @@ namespace Metasound
 			
 			const FInputVertexInterfaceData& InputData = InParams.InputData;
 			
-			FAudioBufferReadRef LeftAudioIn = InputData.GetOrConstructDataReadReference<FAudioBuffer>(METASOUND_GET_PARAM_NAME(InputAudioLeft), InParams.OperatorSettings);
-			FAudioBufferReadRef RightAudioIn = InputData.GetOrConstructDataReadReference<FAudioBuffer>(METASOUND_GET_PARAM_NAME(InputAudioRight), InParams.OperatorSettings);
+			FAudioBufferReadRef LeftAudioIn = InputData.GetOrCreateDefaultDataReadReference<FAudioBuffer>(METASOUND_GET_PARAM_NAME(InputAudioLeft), InParams.OperatorSettings);
+			FAudioBufferReadRef RightAudioIn = InputData.GetOrCreateDefaultDataReadReference<FAudioBuffer>(METASOUND_GET_PARAM_NAME(InputAudioRight), InParams.OperatorSettings);
 			FFloatReadRef SpreadAmountIn = InputData.GetOrCreateDefaultDataReadReference<float>(METASOUND_GET_PARAM_NAME(InputSpreadAmount), InParams.OperatorSettings);
 			FBoolReadRef bEqualPowerIn = InputData.GetOrCreateDefaultDataReadReference<bool>(METASOUND_GET_PARAM_NAME(InputEqualPower), InParams.OperatorSettings);
 
@@ -359,8 +359,8 @@ namespace Metasound
 			
 			const FInputVertexInterfaceData& InputData = InParams.InputData;
 
-			FAudioBufferReadRef MidAudioIn = InputData.GetOrConstructDataReadReference<FAudioBuffer>(METASOUND_GET_PARAM_NAME(InputAudioMid), InParams.OperatorSettings);
-			FAudioBufferReadRef SideAudioIn = InputData.GetOrConstructDataReadReference<FAudioBuffer>(METASOUND_GET_PARAM_NAME(InputAudioSide), InParams.OperatorSettings);
+			FAudioBufferReadRef MidAudioIn = InputData.GetOrCreateDefaultDataReadReference<FAudioBuffer>(METASOUND_GET_PARAM_NAME(InputAudioMid), InParams.OperatorSettings);
+			FAudioBufferReadRef SideAudioIn = InputData.GetOrCreateDefaultDataReadReference<FAudioBuffer>(METASOUND_GET_PARAM_NAME(InputAudioSide), InParams.OperatorSettings);
 			FFloatReadRef SpreadAmountIn = InputData.GetOrCreateDefaultDataReadReference<float>(METASOUND_GET_PARAM_NAME(InputSpreadAmount), InParams.OperatorSettings);
 			FBoolReadRef bEqualPowerIn = InputData.GetOrCreateDefaultDataReadReference<bool>(METASOUND_GET_PARAM_NAME(InputEqualPower), InParams.OperatorSettings);
 
diff --git a/Engine/Plugins/Runtime/Metasound/Source/MetasoundStandardNodes/Private/MetasoundMidiNoteQuantizerNode.cpp b/Engine/Plugins/Runtime/Metasound/Source/MetasoundStandardNodes/Private/MetasoundMidiNoteQuantizerNode.cpp
--- a/Engine/Plugins/Runtime/Metasound/Source/MetasoundStandardNodes/Private/MetasoundMidiNoteQuantizerNode.cpp
+++ b/Engine/Plugins/Runtime/Metasound/Source/MetasoundStandardNodes/Private/MetasoundMidiNoteQuantizerNode.cpp
@@ -164,7 +164,7 @@ namespace Metasound
 		// inputs
 		FFloatReadRef MidiNoteIn = InputData.GetOrCreateDefaultDataReadReference<float>(METASOUND_GET_PARAM_NAME(ParamNoteIn), InParams.OperatorSettings);
 		FFloatReadRef RootNoteIn = InputData.GetOrCreateDefaultDataReadReference<float>(METASOUND_GET_PARAM_NAME(ParamRootNote), InParams.OperatorSettings);
-		FArrayScaleDegreeReadRef InScaleArray = InputData.GetOrConstructDataReadReference<ScaleDegreeArrayType>(METASOUND_GET_PARAM_NAME(ParamScaleDegrees));
+		FArrayScaleDegreeReadRef InScaleArray = InputData.GetOrCreateDefaultDataReadReference<ScaleDegreeArrayType>(METASOUND_GET_PARAM_NAME(ParamScaleDegrees), InParams.OperatorSettings);
 		FFloatReadRef ScaleRangeIn = InputData.GetOrCreateDefaultDataReadReference<float>(METASOUND_GET_PARAM_NAME(ParamScaleRange), InParams.OperatorSettings);
 
 		return MakeUnique <FMidiNoteQuantizerOperator>(
diff --git a/Engine/Plugins/Runtime/Metasound/Source/MetasoundStandardNodes/Private/MetasoundMinNode.cpp b/Engine/Plugins/Runtime/Metasound/Source/MetasoundStandardNodes/Private/MetasoundMinNode.cpp
--- a/Engine/Plugins/Runtime/Metasound/Source/MetasoundStandardNodes/Private/MetasoundMinNode.cpp
+++ b/Engine/Plugins/Runtime/Metasound/Source/MetasoundStandardNodes/Private/MetasoundMinNode.cpp
@@ -116,14 +116,14 @@ namespace Metasound
 			{
 				using namespace MinVertexNames;
 				const FInputVertexInterfaceData& InputData = InParams.InputData;
-				return InputData.GetOrConstructDataReadReference<FAudioBuffer>(METASOUND_GET_PARAM_NAME(InputAValue), InParams.OperatorSettings);
+				return InputData.GetOrCreateDefaultDataReadReference<FAudioBuffer>(METASOUND_GET_PARAM_NAME(InputAValue), InParams.OperatorSettings);
 			}
 
 			static TDataReadReference<FAudioBuffer> CreateInRefB(const FBuildOperatorParams& InParams)
 			{
 				using namespace MinVertexNames;
 				const FInputVertexInterfaceData& InputData = InParams.InputData;
-				return InputData.GetOrConstructDataReadReference<FAudioBuffer>(METASOUND_GET_PARAM_NAME(InputBValue), InParams.OperatorSettings);
+				return InputData.GetOrCreateDefaultDataReadReference<FAudioBuffer>(METASOUND_GET_PARAM_NAME(InputBValue), InParams.OperatorSettings);
 			}
 		};
 
diff --git a/Engine/Plugins/Runtime/Metasound/Source/MetasoundStandardNodes/Private/MetasoundMixerNode.cpp b/Engine/Plugins/Runtime/Metasound/Source/MetasoundStandardNodes/Private/MetasoundMixerNode.cpp
--- a/Engine/Plugins/Runtime/Metasound/Source/MetasoundStandardNodes/Private/MetasoundMixerNode.cpp
+++ b/Engine/Plugins/Runtime/Metasound/Source/MetasoundStandardNodes/Private/MetasoundMixerNode.cpp
@@ -156,7 +156,7 @@ namespace Metasound
 			{
 				for (uint32 Chan = 0; Chan < NumChannels; ++Chan)
 				{
-					InputBuffers.Add(InputData.GetOrConstructDataReadReference<FAudioBuffer>(GetAudioInputName(i, Chan), InParams.OperatorSettings));
+					InputBuffers.Add(InputData.GetOrCreateDefaultDataReadReference<FAudioBuffer>(GetAudioInputName(i, Chan), InParams.OperatorSettings));
 				}
 
 				InputGains.Add(InputData.GetOrCreateDefaultDataReadReference<float>(GainInputNames[i], InParams.OperatorSettings));
diff --git a/Engine/Plugins/Runtime/Metasound/Source/MetasoundStandardNodes/Private/MetasoundPlateReverbNode.cpp b/Engine/Plugins/Runtime/Metasound/Source/MetasoundStandardNodes/Private/MetasoundPlateReverbNode.cpp
--- a/Engine/Plugins/Runtime/Metasound/Source/MetasoundStandardNodes/Private/MetasoundPlateReverbNode.cpp
+++ b/Engine/Plugins/Runtime/Metasound/Source/MetasoundStandardNodes/Private/MetasoundPlateReverbNode.cpp
@@ -127,8 +127,8 @@ namespace Metasound
 			FInputs Inputs
 			{
 				InParams.InputData.GetOrCreateDefaultDataReadReference<bool>(PlateReverb::Inputs::BypassName, InParams.OperatorSettings),
-				InParams.InputData.GetOrConstructDataReadReference<FAudioBuffer>(PlateReverb::Inputs::InAudioLeftName, InParams.OperatorSettings),
-				InParams.InputData.GetOrConstructDataReadReference<FAudioBuffer>(PlateReverb::Inputs::InAudioRightName, InParams.OperatorSettings),
+				InParams.InputData.GetOrCreateDefaultDataReadReference<FAudioBuffer>(PlateReverb::Inputs::InAudioLeftName, InParams.OperatorSettings),
+				InParams.InputData.GetOrCreateDefaultDataReadReference<FAudioBuffer>(PlateReverb::Inputs::InAudioRightName, InParams.OperatorSettings),
 				InParams.InputData.GetOrCreateDefaultDataReadReference<float>(PlateReverb::Inputs::DryLevelName, InParams.OperatorSettings),
 				InParams.InputData.GetOrCreateDefaultDataReadReference<float>(PlateReverb::Inputs::WetLevelName, InParams.OperatorSettings),
 				InParams.InputData.GetOrCreateDefaultDataReadReference<float>(PlateReverb::Inputs::LateReflectionsDelayName, InParams.OperatorSettings),
diff --git a/Engine/Plugins/Runtime/Metasound/Source/MetasoundStandardNodes/Private/MetasoundRingModNode.cpp b/Engine/Plugins/Runtime/Metasound/Source/MetasoundStandardNodes/Private/MetasoundRingModNode.cpp
--- a/Engine/Plugins/Runtime/Metasound/Source/MetasoundStandardNodes/Private/MetasoundRingModNode.cpp
+++ b/Engine/Plugins/Runtime/Metasound/Source/MetasoundStandardNodes/Private/MetasoundRingModNode.cpp
@@ -105,8 +105,8 @@ namespace Metasound
 			using namespace RingModVertexNames;
 			const FInputVertexInterfaceData& InputData = InParams.InputData;
 
-			FAudioBufferReadRef CarrierAudioIn = InputData.GetOrConstructDataReadReference<FAudioBuffer>(METASOUND_GET_PARAM_NAME(InputAudioCarrier), InParams.OperatorSettings);
-			FAudioBufferReadRef ModulatorAudioIn = InputData.GetOrConstructDataReadReference<FAudioBuffer>(METASOUND_GET_PARAM_NAME(InputAudioModulator), InParams.OperatorSettings);
+			FAudioBufferReadRef CarrierAudioIn = InputData.GetOrCreateDefaultDataReadReference<FAudioBuffer>(METASOUND_GET_PARAM_NAME(InputAudioCarrier), InParams.OperatorSettings);
+			FAudioBufferReadRef ModulatorAudioIn = InputData.GetOrCreateDefaultDataReadReference<FAudioBuffer>(METASOUND_GET_PARAM_NAME(InputAudioModulator), InParams.OperatorSettings);
 
 			return MakeUnique<FRingModOperator>(InParams.OperatorSettings, CarrierAudioIn, ModulatorAudioIn);
 		}
diff --git a/Engine/Plugins/Runtime/Metasound/Source/MetasoundStandardNodes/Private/MetasoundSampleAndHoldNode.cpp b/Engine/Plugins/Runtime/Metasound/Source/MetasoundStandardNodes/Private/MetasoundSampleAndHoldNode.cpp
--- a/Engine/Plugins/Runtime/Metasound/Source/MetasoundStandardNodes/Private/MetasoundSampleAndHoldNode.cpp
+++ b/Engine/Plugins/Runtime/Metasound/Source/MetasoundStandardNodes/Private/MetasoundSampleAndHoldNode.cpp
@@ -81,8 +81,8 @@ namespace Metasound
 			using namespace SampleAndHoldVertexNames;
 			const FInputVertexInterfaceData& InputData = InParams.InputData;
 
-			FTriggerReadRef InputTriggerSampleAndHold = InputData.GetOrConstructDataReadReference<FTrigger>(METASOUND_GET_PARAM_NAME(InputTriggerSampleAndHold), InParams.OperatorSettings);
-			FAudioBufferReadRef InputAudio = InputData.GetOrConstructDataReadReference<FAudioBuffer>(METASOUND_GET_PARAM_NAME(InputAudio), InParams.OperatorSettings);
+			FTriggerReadRef InputTriggerSampleAndHold = InputData.GetOrCreateDefaultDataReadReference<FTrigger>(METASOUND_GET_PARAM_NAME(InputTriggerSampleAndHold), InParams.OperatorSettings);
+			FAudioBufferReadRef InputAudio = InputData.GetOrCreateDefaultDataReadReference<FAudioBuffer>(METASOUND_GET_PARAM_NAME(InputAudio), InParams.OperatorSettings);
 
 			return MakeUnique<FSampleAndHoldOperator>(InParams, InputTriggerSampleAndHold, InputAudio);
 		}
diff --git a/Engine/Plugins/Runtime/Metasound/Source/MetasoundStandardNodes/Private/MetasoundStereoDelayNode.cpp b/Engine/Plugins/Runtime/Metasound/Source/MetasoundStandardNodes/Private/MetasoundStereoDelayNode.cpp
--- a/Engine/Plugins/Runtime/Metasound/Source/MetasoundStandardNodes/Private/MetasoundStereoDelayNode.cpp
+++ b/Engine/Plugins/Runtime/Metasound/Source/MetasoundStandardNodes/Private/MetasoundStereoDelayNode.cpp
@@ -387,8 +387,8 @@ namespace Metasound
 		using namespace StereoDelay; 
 		const FInputVertexInterfaceData& InputData = InParams.InputData;
 
-		FAudioBufferReadRef LeftAudioIn = InputData.GetOrConstructDataReadReference<FAudioBuffer>(METASOUND_GET_PARAM_NAME(InAudioInputLeft), InParams.OperatorSettings);
-		FAudioBufferReadRef RightAudioIn = InputData.GetOrConstructDataReadReference<FAudioBuffer>(METASOUND_GET_PARAM_NAME(InAudioInputRight), InParams.OperatorSettings);
+		FAudioBufferReadRef LeftAudioIn = InputData.GetOrCreateDefaultDataReadReference<FAudioBuffer>(METASOUND_GET_PARAM_NAME(InAudioInputLeft), InParams.OperatorSettings);
+		FAudioBufferReadRef RightAudioIn = InputData.GetOrCreateDefaultDataReadReference<FAudioBuffer>(METASOUND_GET_PARAM_NAME(InAudioInputRight), InParams.OperatorSettings);
 		FStereoDelayModeReadRef StereoDelayMode = InputData.GetOrCreateDefaultDataReadReference<FEnumStereoDelayMode>(METASOUND_GET_PARAM_NAME(InDelayMode), InParams.OperatorSettings);
 		FTimeReadRef DelayTime = InputData.GetOrCreateDefaultDataReadReference<FTime>(METASOUND_GET_PARAM_NAME(InDelayTime), InParams.OperatorSettings);
 		FFloatReadRef DelayRatio = InputData.GetOrCreateDefaultDataReadReference<float>(METASOUND_GET_PARAM_NAME(InDelayRatio), InParams.OperatorSettings);
@@ -396,7 +396,7 @@ namespace Metasound
 		FFloatReadRef WetLevel = InputData.GetOrCreateDefaultDataReadReference<float>(METASOUND_GET_PARAM_NAME(InWetLevel), InParams.OperatorSettings);
 		FFloatReadRef Feedback = InputData.GetOrCreateDefaultDataReadReference<float>(METASOUND_GET_PARAM_NAME(InFeedbackAmount), InParams.OperatorSettings);
 		FTime MaxDelayTime = InputData.GetOrCreateDefaultValue<FTime>(METASOUND_GET_PARAM_NAME(InParamMaxDelayTime), InParams.OperatorSettings);
-		FTriggerReadRef TriggerReset = InputData.GetOrConstructDataReadReference<FTrigger>(METASOUND_GET_PARAM_NAME(InputResetDelay), InParams.OperatorSettings);
+		FTriggerReadRef TriggerReset = InputData.GetOrCreateDefaultDataReadReference<FTrigger>(METASOUND_GET_PARAM_NAME(InputResetDelay), InParams.OperatorSettings);
 
 		return MakeUnique<FStereoDelayOperator>(InParams.OperatorSettings, LeftAudioIn, RightAudioIn, StereoDelayMode, DelayTime, DelayRatio, DryLevel, WetLevel, Feedback, MaxDelayTime.GetSeconds(), TriggerReset);
 	}
diff --git a/Engine/Plugins/Runtime/Metasound/Source/MetasoundStandardNodes/Private/MetasoundStereoPannerNode.cpp b/Engine/Plugins/Runtime/Metasound/Source/MetasoundStandardNodes/Private/MetasoundStereoPannerNode.cpp
--- a/Engine/Plugins/Runtime/Metasound/Source/MetasoundStandardNodes/Private/MetasoundStereoPannerNode.cpp
+++ b/Engine/Plugins/Runtime/Metasound/Source/MetasoundStandardNodes/Private/MetasoundStereoPannerNode.cpp
@@ -232,7 +232,7 @@ namespace Metasound
 
 		using namespace StereoPannerVertexNames;
 
-		FAudioBufferReadRef AudioIn = InputData.GetOrConstructDataReadReference<FAudioBuffer>(METASOUND_GET_PARAM_NAME(InputAudio), InParams.OperatorSettings);
+		FAudioBufferReadRef AudioIn = InputData.GetOrCreateDefaultDataReadReference<FAudioBuffer>(METASOUND_GET_PARAM_NAME(InputAudio), InParams.OperatorSettings);
 		FFloatReadRef PanningAmount = InputData.GetOrCreateDefaultDataReadReference<float>(METASOUND_GET_PARAM_NAME(InputPanAmount), InParams.OperatorSettings);
 		FPanningLawReadRef PanningLaw = InputData.GetOrCreateDefaultDataReadReference<FEnumPanningLaw>(METASOUND_GET_PARAM_NAME(InputPanningLaw), InParams.OperatorSettings);
 
diff --git a/Engine/Plugins/Runtime/Metasound/Source/MetasoundStandardNodes/Private/MetasoundSuperOscillatorNode.cpp b/Engine/Plugins/Runtime/Metasound/Source/MetasoundStandardNodes/Private/MetasoundSuperOscillatorNode.cpp
--- a/Engine/Plugins/Runtime/Metasound/Source/MetasoundStandardNodes/Private/MetasoundSuperOscillatorNode.cpp
+++ b/Engine/Plugins/Runtime/Metasound/Source/MetasoundStandardNodes/Private/MetasoundSuperOscillatorNode.cpp
@@ -475,17 +475,17 @@ namespace Metasound
 			FSuperOscillatorOperatorConstructParams OpParams
 			{
 				Settings,
-				InputData.GetOrConstructDataReadReference<bool>(METASOUND_GET_PARAM_NAME(EnabledPin), true),
-				InputData.GetOrConstructDataReadReference<bool>(METASOUND_GET_PARAM_NAME(LimitOutputPin), true),
+				InputData.GetOrCreateDefaultDataReadReference<bool>(METASOUND_GET_PARAM_NAME(EnabledPin), Settings),
+				InputData.GetOrCreateDefaultDataReadReference<bool>(METASOUND_GET_PARAM_NAME(LimitOutputPin), Settings),
 				InputData.GetOrCreateDefaultValue<int32>(METASOUND_GET_PARAM_NAME(NumVoicesPin), Settings),
-				InputData.GetOrConstructDataReadReference<float>(METASOUND_GET_PARAM_NAME(BaseFrequencyPin), 440.f),
-				InputData.GetOrConstructDataReadReference<float>(METASOUND_GET_PARAM_NAME(DetuneVolumePin), 0.f),
-				InputData.GetOrConstructDataReadReference<float>(METASOUND_GET_PARAM_NAME(MaxDetunePin), -0.25f),
-				InputData.GetOrConstructDataReadReference<float>(METASOUND_GET_PARAM_NAME(EntropyPin), 0.0f),
-				InputData.GetOrConstructDataReadReference<float>(METASOUND_GET_PARAM_NAME(GlideFactorPin), 0.f),
-				InputData.GetOrConstructDataReadReference<float>(METASOUND_GET_PARAM_NAME(PulseWidthPin), 0.5f),
+				InputData.GetOrCreateDefaultDataReadReference<float>(METASOUND_GET_PARAM_NAME(BaseFrequencyPin), Settings),
+				InputData.GetOrCreateDefaultDataReadReference<float>(METASOUND_GET_PARAM_NAME(DetuneVolumePin), Settings),
+				InputData.GetOrCreateDefaultDataReadReference<float>(METASOUND_GET_PARAM_NAME(MaxDetunePin), Settings),
+				InputData.GetOrCreateDefaultDataReadReference<float>(METASOUND_GET_PARAM_NAME(EntropyPin), Settings),
+				InputData.GetOrCreateDefaultDataReadReference<float>(METASOUND_GET_PARAM_NAME(GlideFactorPin), Settings),
+				InputData.GetOrCreateDefaultDataReadReference<float>(METASOUND_GET_PARAM_NAME(PulseWidthPin), Settings),
 				InputData.IsVertexBound(METASOUND_GET_PARAM_NAME(FrequencyModPin)),
-				InputData.GetOrConstructDataReadReference<FAudioBuffer>(METASOUND_GET_PARAM_NAME(FrequencyModPin), Settings),
+				InputData.GetOrCreateDefaultDataReadReference<FAudioBuffer>(METASOUND_GET_PARAM_NAME(FrequencyModPin), Settings),
 				InputData.GetOrCreateDefaultValue<FEnumLfoWaveshapeType>(METASOUND_GET_PARAM_NAME(WaveTypePin), Settings)
 			};
 			
@@ -602,21 +602,21 @@ namespace Metasound
 			FSuperOscillatorOperatorConstructParams OpParams
 			{
 				Settings,
-				InputData.GetOrConstructDataReadReference<bool>(METASOUND_GET_PARAM_NAME(EnabledPin), true),
-				InputData.GetOrConstructDataReadReference<bool>(METASOUND_GET_PARAM_NAME(LimitOutputPin), true),
+				InputData.GetOrCreateDefaultDataReadReference<bool>(METASOUND_GET_PARAM_NAME(EnabledPin), Settings),
+				InputData.GetOrCreateDefaultDataReadReference<bool>(METASOUND_GET_PARAM_NAME(LimitOutputPin), Settings), 
 				InputData.GetOrCreateDefaultValue<int32>(METASOUND_GET_PARAM_NAME(NumVoicesPin), Settings),
-				InputData.GetOrConstructDataReadReference<float>(METASOUND_GET_PARAM_NAME(BaseFrequencyPin), 440.f),
-				InputData.GetOrConstructDataReadReference<float>(METASOUND_GET_PARAM_NAME(DetuneVolumePin), 0.f),
-				InputData.GetOrConstructDataReadReference<float>(METASOUND_GET_PARAM_NAME(MaxDetunePin), -0.25f),
-				InputData.GetOrConstructDataReadReference<float>(METASOUND_GET_PARAM_NAME(EntropyPin), 0.f),
-				InputData.GetOrConstructDataReadReference<float>(METASOUND_GET_PARAM_NAME(GlideFactorPin), 0.f),
-				InputData.GetOrConstructDataReadReference<float>(METASOUND_GET_PARAM_NAME(PulseWidthPin), 0.5f),
+				InputData.GetOrCreateDefaultDataReadReference<float>(METASOUND_GET_PARAM_NAME(BaseFrequencyPin), Settings),
+				InputData.GetOrCreateDefaultDataReadReference<float>(METASOUND_GET_PARAM_NAME(DetuneVolumePin), Settings),
+				InputData.GetOrCreateDefaultDataReadReference<float>(METASOUND_GET_PARAM_NAME(MaxDetunePin), Settings),
+				InputData.GetOrCreateDefaultDataReadReference<float>(METASOUND_GET_PARAM_NAME(EntropyPin), Settings),
+				InputData.GetOrCreateDefaultDataReadReference<float>(METASOUND_GET_PARAM_NAME(GlideFactorPin), Settings),
+				InputData.GetOrCreateDefaultDataReadReference<float>(METASOUND_GET_PARAM_NAME(PulseWidthPin), Settings),
 				InputData.IsVertexBound(METASOUND_GET_PARAM_NAME(FrequencyModPin)),
-				InputData.GetOrConstructDataReadReference<FAudioBuffer>(METASOUND_GET_PARAM_NAME(FrequencyModPin), Settings),
+				InputData.GetOrCreateDefaultDataReadReference<FAudioBuffer>(METASOUND_GET_PARAM_NAME(FrequencyModPin), Settings),
 				InputData.GetOrCreateDefaultValue<FEnumLfoWaveshapeType>(METASOUND_GET_PARAM_NAME(WaveTypePin), Settings)
 			};
 
-			FFloatReadRef StereoWidthInput = InputData.GetOrConstructDataReadReference<float>(METASOUND_GET_PARAM_NAME(StereoWidthPin), 0.5f);
+			FFloatReadRef StereoWidthInput = InputData.GetOrCreateDefaultDataReadReference<float>(METASOUND_GET_PARAM_NAME(StereoWidthPin), Settings);
 			
 			return MakeUnique<FSuperOscillatorOperatorStereo>(OpParams, StereoWidthInput);
 		}
diff --git a/Engine/Plugins/Runtime/Metasound/Source/MetasoundStandardNodes/Private/MetasoundSwitchAndRampNode.cpp b/Engine/Plugins/Runtime/Metasound/Source/MetasoundStandardNodes/Private/MetasoundSwitchAndRampNode.cpp
--- a/Engine/Plugins/Runtime/Metasound/Source/MetasoundStandardNodes/Private/MetasoundSwitchAndRampNode.cpp
+++ b/Engine/Plugins/Runtime/Metasound/Source/MetasoundStandardNodes/Private/MetasoundSwitchAndRampNode.cpp
@@ -80,8 +80,8 @@ namespace Metasound
 			using namespace SwitchAndRampVertexNames;
 			const FInputVertexInterfaceData& InputData = InParams.InputData;
 
-			FTriggerReadRef TriggerIn = InputData.GetOrConstructDataReadReference<FTrigger>(METASOUND_GET_PARAM_NAME(InputTrigger), InParams.OperatorSettings);
-			FAudioBufferReadRef AudioInput = InputData.GetOrConstructDataReadReference<FAudioBuffer>(METASOUND_GET_PARAM_NAME(InputAudio), InParams.OperatorSettings);
+			FTriggerReadRef TriggerIn = InputData.GetOrCreateDefaultDataReadReference<FTrigger>(METASOUND_GET_PARAM_NAME(InputTrigger), InParams.OperatorSettings);
+			FAudioBufferReadRef AudioInput = InputData.GetOrCreateDefaultDataReadReference<FAudioBuffer>(METASOUND_GET_PARAM_NAME(InputAudio), InParams.OperatorSettings);
 			FTimeReadRef SmoothTime = InputData.GetOrCreateDefaultDataReadReference<FTime>(METASOUND_GET_PARAM_NAME(InputSmoothTime), InParams.OperatorSettings);
 
 			return MakeUnique<FSwitchAndRampOperator>(InParams.OperatorSettings, TriggerIn, AudioInput, SmoothTime);
diff --git a/Engine/Plugins/Runtime/Metasound/Source/MetasoundStandardNodes/Private/MetasoundTriggerAnyNode.cpp b/Engine/Plugins/Runtime/Metasound/Source/MetasoundStandardNodes/Private/MetasoundTriggerAnyNode.cpp
--- a/Engine/Plugins/Runtime/Metasound/Source/MetasoundStandardNodes/Private/MetasoundTriggerAnyNode.cpp
+++ b/Engine/Plugins/Runtime/Metasound/Source/MetasoundStandardNodes/Private/MetasoundTriggerAnyNode.cpp
@@ -121,7 +121,7 @@ namespace Metasound
 
 			for (uint32 i = 0; i < NumInputs; ++i)
 			{
-				InputTriggers.Add(InputData.GetOrConstructDataReadReference<FTrigger>(GetInputTriggerName(i), InParams.OperatorSettings));
+				InputTriggers.Add(InputData.GetOrCreateDefaultDataReadReference<FTrigger>(GetInputTriggerName(i), InParams.OperatorSettings));
 			}
 
 			return MakeUnique<TTriggerAnyOperator<NumInputs>>(InParams.OperatorSettings, MoveTemp(InputTriggers));
diff --git a/Engine/Plugins/Runtime/Metasound/Source/MetasoundStandardNodes/Private/MetasoundTriggerCoinNode.cpp b/Engine/Plugins/Runtime/Metasound/Source/MetasoundStandardNodes/Private/MetasoundTriggerCoinNode.cpp
--- a/Engine/Plugins/Runtime/Metasound/Source/MetasoundStandardNodes/Private/MetasoundTriggerCoinNode.cpp
+++ b/Engine/Plugins/Runtime/Metasound/Source/MetasoundStandardNodes/Private/MetasoundTriggerCoinNode.cpp
@@ -168,8 +168,8 @@ namespace Metasound
 
 		const FInputVertexInterfaceData& InputData = InParams.InputData;
 
-		FTriggerReadRef TriggerIn = InputData.GetOrConstructDataReadReference<FTrigger>(METASOUND_GET_PARAM_NAME(InputTrigger), InParams.OperatorSettings);
-		FTriggerReadRef TriggerResetIn = InputData.GetOrConstructDataReadReference<FTrigger>(METASOUND_GET_PARAM_NAME(InputReset), InParams.OperatorSettings);
+		FTriggerReadRef TriggerIn = InputData.GetOrCreateDefaultDataReadReference<FTrigger>(METASOUND_GET_PARAM_NAME(InputTrigger), InParams.OperatorSettings);
+		FTriggerReadRef TriggerResetIn = InputData.GetOrCreateDefaultDataReadReference<FTrigger>(METASOUND_GET_PARAM_NAME(InputReset), InParams.OperatorSettings);
 		FInt32ReadRef SeedIn = InputData.GetOrCreateDefaultDataReadReference<int32>(METASOUND_GET_PARAM_NAME(InputSeed), InParams.OperatorSettings);
 		FFloatReadRef ProbabilityIn = InputData.GetOrCreateDefaultDataReadReference<float>(METASOUND_GET_PARAM_NAME(InputProbability), InParams.OperatorSettings);
 
diff --git a/Engine/Plugins/Runtime/Metasound/Source/MetasoundStandardNodes/Private/MetasoundTriggerControlNode.cpp b/Engine/Plugins/Runtime/Metasound/Source/MetasoundStandardNodes/Private/MetasoundTriggerControlNode.cpp
--- a/Engine/Plugins/Runtime/Metasound/Source/MetasoundStandardNodes/Private/MetasoundTriggerControlNode.cpp
+++ b/Engine/Plugins/Runtime/Metasound/Source/MetasoundStandardNodes/Private/MetasoundTriggerControlNode.cpp
@@ -165,10 +165,10 @@ namespace Metasound
 
 		const FInputVertexInterfaceData& InputData = InParams.InputData;
 
-		FTriggerReadRef TriggerEnterIn = InputData.GetOrConstructDataReadReference<FTrigger>(METASOUND_GET_PARAM_NAME(InputEnter), InParams.OperatorSettings);
-		FTriggerReadRef TriggerOpenIn = InputData.GetOrConstructDataReadReference<FTrigger>(METASOUND_GET_PARAM_NAME(InputOpen), InParams.OperatorSettings);
-		FTriggerReadRef TriggerCloseIn = InputData.GetOrConstructDataReadReference<FTrigger>(METASOUND_GET_PARAM_NAME(InputClose), InParams.OperatorSettings);
-		FTriggerReadRef TriggerToggleIn = InputData.GetOrConstructDataReadReference<FTrigger>(METASOUND_GET_PARAM_NAME(InputToggle), InParams.OperatorSettings);
+		FTriggerReadRef TriggerEnterIn = InputData.GetOrCreateDefaultDataReadReference<FTrigger>(METASOUND_GET_PARAM_NAME(InputEnter), InParams.OperatorSettings);
+		FTriggerReadRef TriggerOpenIn = InputData.GetOrCreateDefaultDataReadReference<FTrigger>(METASOUND_GET_PARAM_NAME(InputOpen), InParams.OperatorSettings);
+		FTriggerReadRef TriggerCloseIn = InputData.GetOrCreateDefaultDataReadReference<FTrigger>(METASOUND_GET_PARAM_NAME(InputClose), InParams.OperatorSettings);
+		FTriggerReadRef TriggerToggleIn = InputData.GetOrCreateDefaultDataReadReference<FTrigger>(METASOUND_GET_PARAM_NAME(InputToggle), InParams.OperatorSettings);
 		FBoolReadRef bStartClosedIn = InputData.GetOrCreateDefaultDataReadReference<bool>(METASOUND_GET_PARAM_NAME(InputStartClosed), InParams.OperatorSettings);
 
 		return MakeUnique<FTriggerControlOperator>(InParams.OperatorSettings, TriggerEnterIn, TriggerOpenIn, TriggerCloseIn, TriggerToggleIn, bStartClosedIn);
diff --git a/Engine/Plugins/Runtime/Metasound/Source/MetasoundStandardNodes/Private/MetasoundTriggerCounterNode.cpp b/Engine/Plugins/Runtime/Metasound/Source/MetasoundStandardNodes/Private/MetasoundTriggerCounterNode.cpp
--- a/Engine/Plugins/Runtime/Metasound/Source/MetasoundStandardNodes/Private/MetasoundTriggerCounterNode.cpp
+++ b/Engine/Plugins/Runtime/Metasound/Source/MetasoundStandardNodes/Private/MetasoundTriggerCounterNode.cpp
@@ -186,8 +186,8 @@ namespace Metasound
 		using namespace TriggerCounterVertexNames;
 
 		const FInputVertexInterfaceData& InputData = InParams.InputData;
-		FTriggerReadRef TriggerIn = InputData.GetOrConstructDataReadReference<FTrigger>(METASOUND_GET_PARAM_NAME(InputInTrigger), InParams.OperatorSettings);
-		FTriggerReadRef TriggerReset = InputData.GetOrConstructDataReadReference<FTrigger>(METASOUND_GET_PARAM_NAME(InputReset), InParams.OperatorSettings);
+		FTriggerReadRef TriggerIn = InputData.GetOrCreateDefaultDataReadReference<FTrigger>(METASOUND_GET_PARAM_NAME(InputInTrigger), InParams.OperatorSettings);
+		FTriggerReadRef TriggerReset = InputData.GetOrCreateDefaultDataReadReference<FTrigger>(METASOUND_GET_PARAM_NAME(InputReset), InParams.OperatorSettings);
 		FFloatReadRef StartValue = InputData.GetOrCreateDefaultDataReadReference<float>(METASOUND_GET_PARAM_NAME(InputStartValue), InParams.OperatorSettings);
 		FFloatReadRef StepSize = InputData.GetOrCreateDefaultDataReadReference<float>(METASOUND_GET_PARAM_NAME(InputStepSize), InParams.OperatorSettings);
 		FInt32ReadRef AutoResetCount = InputData.GetOrCreateDefaultDataReadReference<int32>(METASOUND_GET_PARAM_NAME(InputAutoResetCount), InParams.OperatorSettings);
diff --git a/Engine/Plugins/Runtime/Metasound/Source/MetasoundStandardNodes/Private/MetasoundTriggerDelayNode.cpp b/Engine/Plugins/Runtime/Metasound/Source/MetasoundStandardNodes/Private/MetasoundTriggerDelayNode.cpp
--- a/Engine/Plugins/Runtime/Metasound/Source/MetasoundStandardNodes/Private/MetasoundTriggerDelayNode.cpp
+++ b/Engine/Plugins/Runtime/Metasound/Source/MetasoundStandardNodes/Private/MetasoundTriggerDelayNode.cpp
@@ -115,8 +115,8 @@ namespace Metasound
 
 		FTimeReadRef Delay = InputData.GetOrCreateDefaultDataReadReference<FTime>(METASOUND_GET_PARAM_NAME(InputDelayTime), InParams.OperatorSettings);
 
-		FTriggerReadRef TriggerIn = InputData.GetOrConstructDataReadReference<FTrigger>(METASOUND_GET_PARAM_NAME(InputInTriggerDelay), InParams.OperatorSettings);
-		FTriggerReadRef TriggerReset = InputData.GetOrConstructDataReadReference<FTrigger>(METASOUND_GET_PARAM_NAME(InputResetDelay), InParams.OperatorSettings);
+		FTriggerReadRef TriggerIn = InputData.GetOrCreateDefaultDataReadReference<FTrigger>(METASOUND_GET_PARAM_NAME(InputInTriggerDelay), InParams.OperatorSettings);
+		FTriggerReadRef TriggerReset = InputData.GetOrCreateDefaultDataReadReference<FTrigger>(METASOUND_GET_PARAM_NAME(InputResetDelay), InParams.OperatorSettings);
 
 		return MakeUnique<FTriggerDelayOperator>(InParams.OperatorSettings, TriggerReset, TriggerIn, Delay);
 	}
diff --git a/Engine/Plugins/Runtime/Metasound/Source/MetasoundStandardNodes/Private/MetasoundTriggerOnThresholdNode.cpp b/Engine/Plugins/Runtime/Metasound/Source/MetasoundStandardNodes/Private/MetasoundTriggerOnThresholdNode.cpp
--- a/Engine/Plugins/Runtime/Metasound/Source/MetasoundStandardNodes/Private/MetasoundTriggerOnThresholdNode.cpp
+++ b/Engine/Plugins/Runtime/Metasound/Source/MetasoundStandardNodes/Private/MetasoundTriggerOnThresholdNode.cpp
@@ -196,7 +196,7 @@ namespace Metasound
 			const FInputVertexInterfaceData& InputData = InParams.InputData;
 			const FOperatorSettings& Settings = InParams.OperatorSettings;
 
-			return InputData.GetOrConstructDataReadReference<FAudioBuffer>(METASOUND_GET_PARAM_NAME(InPin), Settings);
+			return InputData.GetOrCreateDefaultDataReadReference<FAudioBuffer>(METASOUND_GET_PARAM_NAME(InPin), Settings);
 		}
 	};
 
diff --git a/Engine/Plugins/Runtime/Metasound/Source/MetasoundStandardNodes/Private/MetasoundTriggerOnceNode.cpp b/Engine/Plugins/Runtime/Metasound/Source/MetasoundStandardNodes/Private/MetasoundTriggerOnceNode.cpp
--- a/Engine/Plugins/Runtime/Metasound/Source/MetasoundStandardNodes/Private/MetasoundTriggerOnceNode.cpp
+++ b/Engine/Plugins/Runtime/Metasound/Source/MetasoundStandardNodes/Private/MetasoundTriggerOnceNode.cpp
@@ -131,8 +131,8 @@ namespace Metasound
 
 		const FInputVertexInterfaceData& InputData = InParams.InputData;
 
-		FTriggerReadRef TriggerEnterIn = InputData.GetOrConstructDataReadReference<FTrigger>(METASOUND_GET_PARAM_NAME(InputEnter), InParams.OperatorSettings);
-		FTriggerReadRef TriggerResetIn = InputData.GetOrConstructDataReadReference<FTrigger>(METASOUND_GET_PARAM_NAME(InputReset), InParams.OperatorSettings);
+		FTriggerReadRef TriggerEnterIn = InputData.GetOrCreateDefaultDataReadReference<FTrigger>(METASOUND_GET_PARAM_NAME(InputEnter), InParams.OperatorSettings);
+		FTriggerReadRef TriggerResetIn = InputData.GetOrCreateDefaultDataReadReference<FTrigger>(METASOUND_GET_PARAM_NAME(InputReset), InParams.OperatorSettings);
 		FBoolReadRef bStartClosedIn = InputData.GetOrCreateDefaultDataReadReference<bool>(METASOUND_GET_PARAM_NAME(InputStartClosed), InParams.OperatorSettings);
 
 		return MakeUnique<FTriggerOnceOperator>(InParams.OperatorSettings, TriggerEnterIn, TriggerResetIn, bStartClosedIn);
diff --git a/Engine/Plugins/Runtime/Metasound/Source/MetasoundStandardNodes/Private/MetasoundTriggerPipeNode.cpp b/Engine/Plugins/Runtime/Metasound/Source/MetasoundStandardNodes/Private/MetasoundTriggerPipeNode.cpp
--- a/Engine/Plugins/Runtime/Metasound/Source/MetasoundStandardNodes/Private/MetasoundTriggerPipeNode.cpp
+++ b/Engine/Plugins/Runtime/Metasound/Source/MetasoundStandardNodes/Private/MetasoundTriggerPipeNode.cpp
@@ -140,8 +140,8 @@ namespace Metasound
 		using namespace TriggerPipeVertexNames;
 		const FInputVertexInterfaceData& InputData = InParams.InputData;
 
-		FTriggerReadRef TriggerIn = InputData.GetOrConstructDataReadReference<FTrigger>(METASOUND_GET_PARAM_NAME(InputInTrigger), InParams.OperatorSettings);
-		FTriggerReadRef TriggerReset = InputData.GetOrConstructDataReadReference<FTrigger>(METASOUND_GET_PARAM_NAME(InputReset), InParams.OperatorSettings);
+		FTriggerReadRef TriggerIn = InputData.GetOrCreateDefaultDataReadReference<FTrigger>(METASOUND_GET_PARAM_NAME(InputInTrigger), InParams.OperatorSettings);
+		FTriggerReadRef TriggerReset = InputData.GetOrCreateDefaultDataReadReference<FTrigger>(METASOUND_GET_PARAM_NAME(InputReset), InParams.OperatorSettings);
 		FTimeReadRef DelayTime = InputData.GetOrCreateDefaultDataReadReference<FTime>(METASOUND_GET_PARAM_NAME(InputDelayTime), InParams.OperatorSettings);
 
 		return MakeUnique<FTriggerPipeOperator>(InParams.OperatorSettings, TriggerReset, TriggerIn, DelayTime);
diff --git a/Engine/Plugins/Runtime/Metasound/Source/MetasoundStandardNodes/Private/MetasoundTriggerRepeatNode.cpp b/Engine/Plugins/Runtime/Metasound/Source/MetasoundStandardNodes/Private/MetasoundTriggerRepeatNode.cpp
--- a/Engine/Plugins/Runtime/Metasound/Source/MetasoundStandardNodes/Private/MetasoundTriggerRepeatNode.cpp
+++ b/Engine/Plugins/Runtime/Metasound/Source/MetasoundStandardNodes/Private/MetasoundTriggerRepeatNode.cpp
@@ -165,8 +165,8 @@ namespace Metasound
 		const FTriggerRepeatNode& PeriodicTriggerNode = static_cast<const FTriggerRepeatNode&>(InParams.Node);
 		const FInputVertexInterfaceData& InputData = InParams.InputData;
 
-		FTriggerReadRef TriggerEnable = InputData.GetOrConstructDataReadReference<FTrigger>(METASOUND_GET_PARAM_NAME(InputStart), InParams.OperatorSettings);
-		FTriggerReadRef TriggerDisable = InputData.GetOrConstructDataReadReference<FTrigger>(METASOUND_GET_PARAM_NAME(InputStop), InParams.OperatorSettings);
+		FTriggerReadRef TriggerEnable = InputData.GetOrCreateDefaultDataReadReference<FTrigger>(METASOUND_GET_PARAM_NAME(InputStart), InParams.OperatorSettings);
+		FTriggerReadRef TriggerDisable = InputData.GetOrCreateDefaultDataReadReference<FTrigger>(METASOUND_GET_PARAM_NAME(InputStop), InParams.OperatorSettings);
 		FTimeReadRef Period = InputData.GetOrCreateDefaultDataReadReference<FTime>(METASOUND_GET_PARAM_NAME(InputPeriod), InParams.OperatorSettings);
 		FInt32ReadRef NumRepeats = InputData.GetOrCreateDefaultDataReadReference<int32>(METASOUND_GET_PARAM_NAME(InputNumRepeats), InParams.OperatorSettings);
 
diff --git a/Engine/Plugins/Runtime/Metasound/Source/MetasoundStandardNodes/Private/MetasoundTriggerRouteNode.cpp b/Engine/Plugins/Runtime/Metasound/Source/MetasoundStandardNodes/Private/MetasoundTriggerRouteNode.cpp
--- a/Engine/Plugins/Runtime/Metasound/Source/MetasoundStandardNodes/Private/MetasoundTriggerRouteNode.cpp
+++ b/Engine/Plugins/Runtime/Metasound/Source/MetasoundStandardNodes/Private/MetasoundTriggerRouteNode.cpp
@@ -155,7 +155,7 @@ namespace Metasound
 
 			for (uint32 i = 0; i < NumInputs; ++i)
 			{
-				InputTriggers.Add(InputData.GetOrConstructDataReadReference<FTrigger>(GetInputTriggerName(i), InParams.OperatorSettings));
+				InputTriggers.Add(InputData.GetOrCreateDefaultDataReadReference<FTrigger>(GetInputTriggerName(i), InParams.OperatorSettings));
 				InputValues.Add(InputData.GetOrCreateDefaultDataReadReference<ValueType>(GetInputValueName(i), InParams.OperatorSettings));
 			}
 
diff --git a/Engine/Plugins/Runtime/Metasound/Source/MetasoundStandardNodes/Private/MetasoundTriggerSelectNode.cpp b/Engine/Plugins/Runtime/Metasound/Source/MetasoundStandardNodes/Private/MetasoundTriggerSelectNode.cpp
--- a/Engine/Plugins/Runtime/Metasound/Source/MetasoundStandardNodes/Private/MetasoundTriggerSelectNode.cpp
+++ b/Engine/Plugins/Runtime/Metasound/Source/MetasoundStandardNodes/Private/MetasoundTriggerSelectNode.cpp
@@ -128,7 +128,7 @@ namespace Metasound
 
 		const FInputVertexInterfaceData& InputData = InParams.InputData;
 
-		FTriggerReadRef TriggerIn = InputData.GetOrConstructDataReadReference<FTrigger>(METASOUND_GET_PARAM_NAME(InputTrigger), InParams.OperatorSettings);
+		FTriggerReadRef TriggerIn = InputData.GetOrCreateDefaultDataReadReference<FTrigger>(METASOUND_GET_PARAM_NAME(InputTrigger), InParams.OperatorSettings);
 		FInt32ReadRef InputIn = InputData.GetOrCreateDefaultDataReadReference<int32>(METASOUND_GET_PARAM_NAME(InputIndex), InParams.OperatorSettings);
 
 		return MakeUnique<TTriggerSelectOperator>(InParams.OperatorSettings, TriggerIn, InputIn);
diff --git a/Engine/Plugins/Runtime/Metasound/Source/MetasoundStandardNodes/Private/MetasoundTriggerSequenceNode.cpp b/Engine/Plugins/Runtime/Metasound/Source/MetasoundStandardNodes/Private/MetasoundTriggerSequenceNode.cpp
--- a/Engine/Plugins/Runtime/Metasound/Source/MetasoundStandardNodes/Private/MetasoundTriggerSequenceNode.cpp
+++ b/Engine/Plugins/Runtime/Metasound/Source/MetasoundStandardNodes/Private/MetasoundTriggerSequenceNode.cpp
@@ -157,8 +157,8 @@ namespace Metasound
 
 		const FInputVertexInterfaceData& InputData = InParams.InputData;
 
-		FTriggerReadRef TriggerIn = InputData.GetOrConstructDataReadReference<FTrigger>(METASOUND_GET_PARAM_NAME(InputTrigger), InParams.OperatorSettings);
-		FTriggerReadRef TriggerResetIn = InputData.GetOrConstructDataReadReference<FTrigger>(METASOUND_GET_PARAM_NAME(InputReset), InParams.OperatorSettings);
+		FTriggerReadRef TriggerIn = InputData.GetOrCreateDefaultDataReadReference<FTrigger>(METASOUND_GET_PARAM_NAME(InputTrigger), InParams.OperatorSettings);
+		FTriggerReadRef TriggerResetIn = InputData.GetOrCreateDefaultDataReadReference<FTrigger>(METASOUND_GET_PARAM_NAME(InputReset), InParams.OperatorSettings);
 		FBoolReadRef bLoopIn = InputData.GetOrCreateDefaultDataReadReference<bool>(METASOUND_GET_PARAM_NAME(InputLoop), InParams.OperatorSettings);
 
 		return MakeUnique<TTriggerSequenceOperator>(InParams.OperatorSettings, TriggerIn, TriggerResetIn, bLoopIn);
diff --git a/Engine/Plugins/Runtime/Metasound/Source/MetasoundStandardNodes/Private/MetasoundTriggerToggleNode.cpp b/Engine/Plugins/Runtime/Metasound/Source/MetasoundStandardNodes/Private/MetasoundTriggerToggleNode.cpp
--- a/Engine/Plugins/Runtime/Metasound/Source/MetasoundStandardNodes/Private/MetasoundTriggerToggleNode.cpp
+++ b/Engine/Plugins/Runtime/Metasound/Source/MetasoundStandardNodes/Private/MetasoundTriggerToggleNode.cpp
@@ -116,8 +116,8 @@ namespace Metasound
 		using namespace TriggerToggle;
 
 		const FInputVertexInterfaceData& InputData = InParams.InputData;
-		FTriggerReadRef InTriggerOn = InputData.GetOrConstructDataReadReference<FTrigger>(METASOUND_GET_PARAM_NAME(InputOnTrigger), InParams.OperatorSettings);
-		FTriggerReadRef InTriggerOff = InputData.GetOrConstructDataReadReference<FTrigger>(METASOUND_GET_PARAM_NAME(InputOffTrigger), InParams.OperatorSettings);
+		FTriggerReadRef InTriggerOn = InputData.GetOrCreateDefaultDataReadReference<FTrigger>(METASOUND_GET_PARAM_NAME(InputOnTrigger), InParams.OperatorSettings);
+		FTriggerReadRef InTriggerOff = InputData.GetOrCreateDefaultDataReadReference<FTrigger>(METASOUND_GET_PARAM_NAME(InputOffTrigger), InParams.OperatorSettings);
 		FBoolReadRef InInitValue = InputData.GetOrCreateDefaultDataReadReference<bool>(METASOUND_GET_PARAM_NAME(InputInit), InParams.OperatorSettings);
 
 		return MakeUnique<FTriggerToggleOperator>(InParams.OperatorSettings, InTriggerOn, InTriggerOff, InInitValue);
diff --git a/Engine/Plugins/Runtime/Metasound/Source/MetasoundStandardNodes/Private/MetasoundWaveShaperNode.cpp b/Engine/Plugins/Runtime/Metasound/Source/MetasoundStandardNodes/Private/MetasoundWaveShaperNode.cpp
--- a/Engine/Plugins/Runtime/Metasound/Source/MetasoundStandardNodes/Private/MetasoundWaveShaperNode.cpp
+++ b/Engine/Plugins/Runtime/Metasound/Source/MetasoundStandardNodes/Private/MetasoundWaveShaperNode.cpp
@@ -178,7 +178,7 @@ namespace Metasound
 	{
 		const FInputVertexInterfaceData& InputData = InParams.InputData;
 
-		FAudioBufferReadRef AudioIn = InputData.GetOrConstructDataReadReference<FAudioBuffer>(METASOUND_GET_PARAM_NAME(InAudioInput), InParams.OperatorSettings);
+		FAudioBufferReadRef AudioIn = InputData.GetOrCreateDefaultDataReadReference<FAudioBuffer>(METASOUND_GET_PARAM_NAME(InAudioInput), InParams.OperatorSettings);
 		FFloatReadRef InAmount = InputData.GetOrCreateDefaultDataReadReference<float>(METASOUND_GET_PARAM_NAME(InWaveShapeAmount), InParams.OperatorSettings);
 		FFloatReadRef InBias = InputData.GetOrCreateDefaultDataReadReference<float>(METASOUND_GET_PARAM_NAME(InWaveShapeBias), InParams.OperatorSettings);
 		FFloatReadRef OutGain = InputData.GetOrCreateDefaultDataReadReference<float>(METASOUND_GET_PARAM_NAME(InOutputGain), InParams.OperatorSettings);
diff --git a/Engine/Plugins/Runtime/Metasound/Source/MetasoundStandardNodes/Private/MetasoundWaveWriterNode.cpp b/Engine/Plugins/Runtime/Metasound/Source/MetasoundStandardNodes/Private/MetasoundWaveWriterNode.cpp
--- a/Engine/Plugins/Runtime/Metasound/Source/MetasoundStandardNodes/Private/MetasoundWaveWriterNode.cpp
+++ b/Engine/Plugins/Runtime/Metasound/Source/MetasoundStandardNodes/Private/MetasoundWaveWriterNode.cpp
@@ -360,7 +360,7 @@ namespace Metasound
 			{
 				NumConnectedAudioPins++;
 			}
-			InputBuffers.Add(InputData.GetOrConstructDataReadReference<FAudioBuffer>(PinName, InParams.OperatorSettings));
+			InputBuffers.Add(InputData.GetOrCreateDefaultDataReadReference<FAudioBuffer>(PinName, InParams.OperatorSettings));
 		}
 		
 		// Only create a real operator if there's some connected pins.
diff --git a/Engine/Plugins/Runtime/Metasound/Source/MetasoundStandardNodes/Public/MetasoundRandomNode.h b/Engine/Plugins/Runtime/Metasound/Source/MetasoundStandardNodes/Public/MetasoundRandomNode.h
--- a/Engine/Plugins/Runtime/Metasound/Source/MetasoundStandardNodes/Public/MetasoundRandomNode.h
+++ b/Engine/Plugins/Runtime/Metasound/Source/MetasoundStandardNodes/Public/MetasoundRandomNode.h
@@ -320,8 +320,8 @@ namespace Metasound
  			
 			using namespace RandomNodeNames;
 
- 			FTriggerReadRef InNextTrigger = InputData.GetOrConstructDataReadReference<FTrigger>(METASOUND_GET_PARAM_NAME(InputNextTrigger), InParams.OperatorSettings);
-			FTriggerReadRef InResetTrigger = InputData.GetOrConstructDataReadReference<FTrigger>(METASOUND_GET_PARAM_NAME(InputResetTrigger), InParams.OperatorSettings);
+ 			FTriggerReadRef InNextTrigger = InputData.GetOrCreateDefaultDataReadReference<FTrigger>(METASOUND_GET_PARAM_NAME(InputNextTrigger), InParams.OperatorSettings);
+			FTriggerReadRef InResetTrigger = InputData.GetOrCreateDefaultDataReadReference<FTrigger>(METASOUND_GET_PARAM_NAME(InputResetTrigger), InParams.OperatorSettings);
 			FInt32ReadRef InSeedValue = InputData.GetOrCreateDefaultDataReadReference<int32>(METASOUND_GET_PARAM_NAME(InputSeed), InParams.OperatorSettings);
 
 			// note: random bool does not have a range
diff --git a/Engine/Plugins/Runtime/Metasound/Source/MetasoundStandardNodes/Public/MetasoundTriggerAccumulatorNode.h b/Engine/Plugins/Runtime/Metasound/Source/MetasoundStandardNodes/Public/MetasoundTriggerAccumulatorNode.h
--- a/Engine/Plugins/Runtime/Metasound/Source/MetasoundStandardNodes/Public/MetasoundTriggerAccumulatorNode.h
+++ b/Engine/Plugins/Runtime/Metasound/Source/MetasoundStandardNodes/Public/MetasoundTriggerAccumulatorNode.h
@@ -83,7 +83,7 @@ namespace Metasound
 
 			for (uint32 i = 0; i < NumInputs; ++i)
 			{
-				InputTriggers.Add(InputData.GetOrConstructDataReadReference<FTrigger>(METASOUND_GET_PARAM_NAME_WITH_INDEX(InputTrigger, i), InParams.OperatorSettings));
+				InputTriggers.Add(InputData.GetOrCreateDefaultDataReadReference<FTrigger>(METASOUND_GET_PARAM_NAME_WITH_INDEX(InputTrigger, i), InParams.OperatorSettings));
 			}
 
 			return MakeUnique<TTriggerAccumulatorOperator<NumInputs>>(InParams.OperatorSettings, bInAutoReset, MoveTemp(InputTriggers));
diff --git a/Engine/Plugins/Runtime/Metasound/Source/MetasoundStandardNodes/Public/MetasoundTriggerCompareNode.h b/Engine/Plugins/Runtime/Metasound/Source/MetasoundStandardNodes/Public/MetasoundTriggerCompareNode.h
--- a/Engine/Plugins/Runtime/Metasound/Source/MetasoundStandardNodes/Public/MetasoundTriggerCompareNode.h
+++ b/Engine/Plugins/Runtime/Metasound/Source/MetasoundStandardNodes/Public/MetasoundTriggerCompareNode.h
@@ -98,7 +98,7 @@ namespace Metasound
 			
 			const FInputVertexInterfaceData& InputData = InParams.InputData;
 
-			FTriggerReadRef InOnTriggerCompare = InputData.GetOrConstructDataReadReference<FTrigger>(METASOUND_GET_PARAM_NAME(InputCompare), InParams.OperatorSettings);
+			FTriggerReadRef InOnTriggerCompare = InputData.GetOrCreateDefaultDataReadReference<FTrigger>(METASOUND_GET_PARAM_NAME(InputCompare), InParams.OperatorSettings);
 			TDataReadReference<ValueType> InValueA = InputData.GetOrCreateDefaultDataReadReference<ValueType>(METASOUND_GET_PARAM_NAME(InputParamA), InParams.OperatorSettings);
 			TDataReadReference<ValueType> InValueB = InputData.GetOrCreateDefaultDataReadReference<ValueType>(METASOUND_GET_PARAM_NAME(InputParamB), InParams.OperatorSettings);
 			FEnumTriggerComparisonTypeReadRef InComparison = InputData.GetOrCreateDefaultDataReadReference<FEnumTriggerComparisonType>(METASOUND_GET_PARAM_NAME(InputCompareType), InParams.OperatorSettings);
diff --git a/Engine/Plugins/Runtime/Metasound/Source/MetasoundStandardNodes/Public/MetasoundValueNode.h b/Engine/Plugins/Runtime/Metasound/Source/MetasoundStandardNodes/Public/MetasoundValueNode.h
--- a/Engine/Plugins/Runtime/Metasound/Source/MetasoundStandardNodes/Public/MetasoundValueNode.h
+++ b/Engine/Plugins/Runtime/Metasound/Source/MetasoundStandardNodes/Public/MetasoundValueNode.h
@@ -80,8 +80,8 @@ namespace Metasound
 
 			const FInputVertexInterfaceData& InputData = InParams.InputData;
 			
-			FTriggerReadRef SetTrigger = InputData.GetOrConstructDataReadReference<FTrigger>(METASOUND_GET_PARAM_NAME(InputSetTrigger), InParams.OperatorSettings);
-			FTriggerReadRef ResetTrigger = InputData.GetOrConstructDataReadReference<FTrigger>(METASOUND_GET_PARAM_NAME(InputResetTrigger), InParams.OperatorSettings);
+			FTriggerReadRef SetTrigger = InputData.GetOrCreateDefaultDataReadReference<FTrigger>(METASOUND_GET_PARAM_NAME(InputSetTrigger), InParams.OperatorSettings);
+			FTriggerReadRef ResetTrigger = InputData.GetOrCreateDefaultDataReadReference<FTrigger>(METASOUND_GET_PARAM_NAME(InputResetTrigger), InParams.OperatorSettings);
 			TDataReadReference<ValueType> InitValue = InputData.GetOrCreateDefaultDataReadReference<ValueType>(METASOUND_GET_PARAM_NAME(InputInitValue), InParams.OperatorSettings);
 			TDataReadReference<ValueType> TargetValue = InputData.GetOrCreateDefaultDataReadReference<ValueType>(METASOUND_GET_PARAM_NAME(InputTargetValue), InParams.OperatorSettings);
 
