Use new node creation methods in frontend graph builder
Rework frontend graph builder to use default literals instead of literal nodes
#rnx

diff --git a/Engine/Plugins/Runtime/Metasound/Source/MetasoundFrontend/Private/MetasoundFrontendDocumentController.cpp b/Engine/Plugins/Runtime/Metasound/Source/MetasoundFrontend/Private/MetasoundFrontendDocumentController.cpp
--- a/Engine/Plugins/Runtime/Metasound/Source/MetasoundFrontend/Private/MetasoundFrontendDocumentController.cpp
+++ b/Engine/Plugins/Runtime/Metasound/Source/MetasoundFrontend/Private/MetasoundFrontendDocumentController.cpp
@@ -3,6 +3,8 @@
 #include "MetasoundFrontendDocumentController.h"
 
 #include "Algo/ForEach.h"
+#include "Algo/Sort.h"
+#include "Algo/Unique.h"
 #include "HAL/FileManager.h"
 #include "MetasoundFrontendGraphController.h"
 #include "MetasoundFrontendDocumentIdGenerator.h"
@@ -425,10 +427,57 @@ namespace Metasound
 			return ClassPtr;
 		}
 
+		void FDocumentController::DeduplicateDependencies()
+		{
+			if (FMetasoundFrontendDocument* Document = DocumentPtr.Get())
+			{
+				auto CompareForUnique = [&](const FMetasoundFrontendClass& InLHS, const FMetasoundFrontendClass& InRHS)
+				{
+					return InLHS.ID == InRHS.ID;
+				};
+
+				// Remove duplicate entries, keeping the highest version dependencies. 
+				auto CompareForSort = [&](const FMetasoundFrontendClass& InLHS, const FMetasoundFrontendClass& InRHS)
+				{
+					// Sort by ID
+					if (InLHS.ID < InRHS.ID)
+					{
+						return true;
+					}
+					else if (InRHS.ID < InLHS.ID)
+					{
+						return false;
+					}
+					// If IDs are equal, sort by version number descending
+					else if (InLHS.Metadata.GetVersion() > InRHS.Metadata.GetVersion())
+					{
+						return true;
+					}
+					else if (InLHS.Metadata.GetVersion() < InRHS.Metadata.GetVersion())
+					{
+						return false ;
+					}
+					else
+					{
+						// if IDs and version numbers are equal, sort by number of inputs & outputs descending
+						const int32 NumLHSVertices = InLHS.Interface.Inputs.Num() + InLHS.Interface.Outputs.Num();
+						const int32 NumRHSVertices = InRHS.Interface.Inputs.Num() + InRHS.Interface.Outputs.Num();
+
+						return NumLHSVertices > NumRHSVertices;
+					}
+				};
+
+				Algo::Sort(Document->Dependencies, CompareForSort);
+				Document->Dependencies.SetNum(Algo::Unique(Document->Dependencies, CompareForUnique));
+			}
+		}
 		void FDocumentController::RemoveUnreferencedDependencies()
 		{
 			if (FMetasoundFrontendDocument* Document = DocumentPtr.Get())
 			{
+				// Remove duplicate entries, keeping the highest version dependencies. 
+				DeduplicateDependencies();
+
 				int32 NumDependenciesRemovedThisItr = 0;
 
 				// Repeatedly remove unreferenced dependencies until there are
diff --git a/Engine/Plugins/Runtime/Metasound/Source/MetasoundFrontend/Private/MetasoundFrontendDocumentController.h b/Engine/Plugins/Runtime/Metasound/Source/MetasoundFrontend/Private/MetasoundFrontendDocumentController.h
--- a/Engine/Plugins/Runtime/Metasound/Source/MetasoundFrontend/Private/MetasoundFrontendDocumentController.h
+++ b/Engine/Plugins/Runtime/Metasound/Source/MetasoundFrontend/Private/MetasoundFrontendDocumentController.h
@@ -107,6 +107,7 @@ namespace Metasound
 			FConstDocumentAccess ShareAccess() const override;
 		private:
 
+			void DeduplicateDependencies();
 			bool AddDuplicateSubgraph(const FMetasoundFrontendGraphClass& InGraphToCopy, const FMetasoundFrontendDocument& InOtherDocument);
 
 			FDocumentAccessPtr DocumentPtr;
diff --git a/Engine/Plugins/Runtime/Metasound/Source/MetasoundFrontend/Private/MetasoundFrontendGraphBuilder.cpp b/Engine/Plugins/Runtime/Metasound/Source/MetasoundFrontend/Private/MetasoundFrontendGraphBuilder.cpp
--- a/Engine/Plugins/Runtime/Metasound/Source/MetasoundFrontend/Private/MetasoundFrontendGraphBuilder.cpp
+++ b/Engine/Plugins/Runtime/Metasound/Source/MetasoundFrontend/Private/MetasoundFrontendGraphBuilder.cpp
@@ -24,20 +24,10 @@ namespace Metasound::Frontend
 {
 	namespace GraphBuilderPrivate
 	{
-		struct FDefaultLiteralData
-		{
-			FGuid DestinationNodeID;
-			FGuid DestinationVertexID;
-			FVertexName DestinationVertexKey;
-			FName TypeName;
-			FLiteralNodeConstructorParams InitParams;
-		};
-
 		// Map of Input VertexID to variable data required to construct and connect default variable
 		using FNodeIDVertexID = TTuple<FGuid, FGuid>;
 		using FDependencyByIDMap = TMap<FGuid, const FMetasoundFrontendClass*>;
 		using FSharedNodeByIDMap = TMap<FGuid, TSharedPtr<const INode>>;
-		using FDefaultInputByIDMap = TMap<FNodeIDVertexID, FDefaultLiteralData>;
 
 		// Context used throughout entire graph build process
 		// (for both a root and nested subgraphs)
@@ -58,20 +48,8 @@ namespace Metasound::Frontend
 			const FMetasoundFrontendGraphClass& GraphClass;
 			const FMetasoundFrontendGraph& PagedGraph;
 			FBuildContext& BuildContext;
-			FDefaultInputByIDMap DefaultInputs;
 		};
 
-
-		FNodeInitData CreateNodeInitData(const FMetasoundFrontendNode& InNode)
-		{
-			FNodeInitData InitData;
-
-			InitData.InstanceName = InNode.Name;
-			InitData.InstanceID = InNode.GetID();
-
-			return InitData;
-		}
-
 		template <typename ResolveType>
 		FGuid ResolveTargetPageID(const ResolveType& InToResolve)
 		{
@@ -128,49 +106,6 @@ namespace Metasound::Frontend
 			return Literal;
 		}
 
-		const FMetasoundFrontendVariable* FindVariableForVariableNode(const FMetasoundFrontendNode& InVariableNode, const FMetasoundFrontendGraph& InGraph)
-		{
-			const FGuid& DesiredID = InVariableNode.GetID();
-			return InGraph.Variables.FindByPredicate([&](const FMetasoundFrontendVariable& InVar) { return InVar.VariableNodeID == DesiredID; });
-		}
-
-		TUniquePtr<INode> CreateVariableNode(const FBuildContext& InContext, const FMetasoundFrontendNode& InNode, const FMetasoundFrontendGraph& InGraph)
-		{
-			using namespace Metasound::Frontend;
-
-			const FMetasoundFrontendVariable* FrontendVariable = FindVariableForVariableNode(InNode, InGraph);
-
-			if (nullptr != FrontendVariable)
-			{
-				IDataTypeRegistry& DataTypeRegistry = IDataTypeRegistry::Get();
-				const bool IsLiteralParsableByDataType = DataTypeRegistry.IsLiteralTypeSupported(FrontendVariable->TypeName, FrontendVariable->Literal.GetType());
-
-				if (IsLiteralParsableByDataType)
-				{
-					FLiteral Literal = FrontendVariable->Literal.ToLiteral(FrontendVariable->TypeName, &InContext.DataTypeRegistry, &InContext.ProxyDataCache);
-
-					FVariableNodeConstructorParams InitParams =
-					{
-						InNode.Name,
-						InNode.GetID(),
-						MoveTemp(Literal)
-					};
-
-					return DataTypeRegistry.CreateVariableNode(FrontendVariable->TypeName, MoveTemp(InitParams));
-				}
-				else
-				{
-					UE_LOG(LogMetaSound, Error, TEXT("Cannot create variable node [NodeID:%s]. [Variable:%s] cannot be constructed with the provided literal type."), *InNode.GetID().ToString(), *FrontendVariable->Name.ToString());
-				}
-			}
-			else
-			{
-				UE_LOG(LogMetaSound, Error, TEXT("Cannot create variable node [NodeID:%s]. No variable found for variable node."), *InNode.GetID().ToString());
-			}
-
-			return TUniquePtr<INode>(nullptr);
-		}
-
 		const FMetasoundFrontendLiteral* FindInputLiteralForInputNode(const FMetasoundFrontendNode& InInputNode, const FMetasoundFrontendClass& InInputNodeClass, const FMetasoundFrontendClassInput& InOwningGraphClassInput)
 		{
 			using namespace Frontend;
@@ -229,174 +164,372 @@ namespace Metasound::Frontend
 		}
 
 
-		TUniquePtr<INode> CreateInputNode(const FBuildContext& InContext, const FMetasoundFrontendNode& InNode, const FMetasoundFrontendClass& InClass, const FMetasoundFrontendClassInput& InOwningGraphClassInput)
+		FEnvironmentVertexInterface CreateEnvironmentVertexInterface(const FMetasoundFrontendNodeInterface& InNodeInterface)
 		{
-			using namespace Metasound::Frontend;
-
-			const FMetasoundFrontendLiteral* FrontendLiteral = FindInputLiteralForInputNode(InNode, InClass, InOwningGraphClassInput);
-
-			if (nullptr != FrontendLiteral)
+			if (InNodeInterface.Environment.Num())
 			{
-				if (ensure(InNode.Interface.Inputs.Num() == 1))
+				auto Convert = [](const FMetasoundFrontendVertex& InFrontendVertex) -> FEnvironmentVertex 
 				{
-					IDataTypeRegistry& DataTypeRegistry = IDataTypeRegistry::Get();
-					const FMetasoundFrontendVertex& InputVertex = InNode.Interface.Inputs[0];
+					return FEnvironmentVertex{InFrontendVertex.Name, FText::GetEmpty()};
+				};
 
-					const bool IsLiteralParsableByDataType = DataTypeRegistry.IsLiteralTypeSupported(InputVertex.TypeName, FrontendLiteral->GetType());
+				TArray<FEnvironmentVertex> Vertices;
+				Vertices.Reserve(InNodeInterface.Environment.Num());
+				Algo::Transform(InNodeInterface.Environment, Vertices, Convert);
 
-					if (IsLiteralParsableByDataType)
-					{
-						FLiteral Literal = FrontendLiteral->ToLiteral(InputVertex.TypeName, &InContext.DataTypeRegistry, &InContext.ProxyDataCache);
+				return FEnvironmentVertexInterface{MoveTemp(Vertices)};
+			}
+			else
+			{
+				return FEnvironmentVertexInterface{};
+			}
+		}
 
-						FInputNodeConstructorParams InitParams =
+		FOutputVertexInterface CreateOutputVertexInterface(const FBuildContext& InContext, const FMetasoundFrontendNode& InNode, const FMetasoundFrontendClass& InClass)
+		{
+			// This exists here until we bring sub interface functionality to the frontend. 
+			// This check will no longer be valid at that point. We will have to implement
+			// a way to lookup the class's vertex access type info from the class def, or 
+			// move that info into the node's def.
+			check(InClass.Interface.Outputs.Num() == InNode.Interface.Outputs.Num());
+
+			// TODO: migrate sub interface builders from MetasoundVertex.cpp when
+			// sub interfaces exist on the frontend document
+			if (InNode.Interface.Outputs.Num())
+			{
+				TArray<FOutputDataVertex> Vertices;
+				Vertices.Reserve(InNode.Interface.Outputs.Num());
+
+				switch (InClass.Metadata.GetType())
+				{
+					// The vertices on an input/output node do not have matching
+					// names with the vertices on input/output node classes. The
+					// node class vertex names are blank, and the input/output node
+					// vertex names match the name of the input or output. Here 
+					// we have to grab data from both the class and node interfaces
+					case EMetasoundFrontendClassType::Input:
 						{
-							InNode.Name,
-							InNode.GetID(),
-							InputVertex.Name,
-							MoveTemp(Literal)
-						};
-
-						ensureAlwaysMsgf(InOwningGraphClassInput.AccessType != EMetasoundFrontendVertexAccessType::Unset, TEXT("Graph Class Input cannot be set to access type of 'Unset'"));
-						if (InOwningGraphClassInput.AccessType == EMetasoundFrontendVertexAccessType::Reference)
+							if (InNode.Interface.Outputs.Num() == 1)
+							{
+								const FMetasoundFrontendVertex& NodeOutputVertex = InNode.Interface.Outputs[0];
+								const FMetasoundFrontendClassOutput& ClassOutputVertex = InClass.Interface.Outputs[0];
+								// Note: stripping metadata here. If we make this function public, we may want to keep metadata if WITH_EDITORONLY_DATA
+								Vertices.Add(FOutputDataVertex{
+											NodeOutputVertex.Name, 
+											NodeOutputVertex.TypeName, 
+											FDataVertexMetadata{}, 
+											FrontendVertexAccessTypeToCoreVertexAccessType(ClassOutputVertex.AccessType)
+										});
+							}
+							else
+							{
+								// This is an edge case in the off scenario that someone
+								// is trying to build a graph that used the `TInputNode<FStereoAudio>` 
+								// specialization.  Auto-update should have removed these
+								// long ago, but this log will help identify any issues.
+								// This log is added in UE 5.6 and if no one is hitting
+								// this issue, can be removed in UE 5.7.
+								UE_LOG(LogMetaSound, Error, TEXT("The MetaSound %s has an malformed input node %s. Please replace the input node and resave"), *InContext.DebugAssetName, *InNode.Name.ToString());
+							}
+						}
+						break;
+
+					case EMetasoundFrontendClassType::Output:
 						{
-							return DataTypeRegistry.CreateInputNode(InputVertex.TypeName, MoveTemp(InitParams));
+							check(InNode.Interface.Outputs.Num() == 1);
+							const FMetasoundFrontendVertex& NodeOutputVertex = InNode.Interface.Outputs[0];
+							const FMetasoundFrontendClassOutput& ClassOutputVertex = InClass.Interface.Outputs[0];
+							// Note: stripping metadata here. If we make this function public, we may want to keep metadata if WITH_EDITORONLY_DATA
+							Vertices.Add(FOutputDataVertex{
+										NodeOutputVertex.Name, 
+										NodeOutputVertex.TypeName, 
+										FDataVertexMetadata{}, 
+										FrontendVertexAccessTypeToCoreVertexAccessType(ClassOutputVertex.AccessType)
+									});
+
 						}
-						else // InOwningGraphClassInput.AccessType == EMetasoundFrontendVertexAccessType::Value
+						break;
+
+					default:
 						{
-							return DataTypeRegistry.CreateConstructorInputNode(InputVertex.TypeName, MoveTemp(InitParams));
+							// As-it-stands, node interfaces and class interfaces
+							// match for all other nodes. This will change when
+							// node config is done as a feature. 
+							for (const FMetasoundFrontendClassOutput& ClassOutput : InClass.Interface.Outputs)
+							{
+								// Note: stripping metadata here. If we make this function public, we may want to keep metadata if WITH_EDITORONLY_DATA
+								Vertices.Add(FOutputDataVertex{ClassOutput.Name, ClassOutput.TypeName, FDataVertexMetadata{}, FrontendVertexAccessTypeToCoreVertexAccessType(ClassOutput.AccessType)});
+							}
 						}
-					}
-					else
-					{
-						UE_LOG(LogMetaSound, Error, TEXT("Cannot create input node [NodeID:%s]. [Vertex:%s] cannot be constructed with the provided literal type."), *InNode.GetID().ToString(), *InputVertex.Name.ToString());
-					}
+						break;
 				}
+
+				return FOutputVertexInterface{MoveTemp(Vertices), TArray<VertexPrivate::FSubInterfaceLayout>()};
 			}
 			else
 			{
-				UE_LOG(LogMetaSound, Error, TEXT("Cannot create input node [NodeID:%s]. No default literal set for input node."), *InNode.GetID().ToString());
+				return FOutputVertexInterface{};
 			}
-
-			return TUniquePtr<INode>(nullptr);
 		}
 
-		TArray<FDefaultLiteralData> GetInputDefaultLiteralData(const FBuildContext& InContext, const FMetasoundFrontendNode& InNode, const FNodeInitData& InInitData, const TSet<FNodeIDVertexID>& InEdgeDestinations)
+		FInputVertexInterface CreateInputVertexInterface(const FBuildContext& InContext, const FMetasoundFrontendNode& InNode, const FMetasoundFrontendClass& InClass, const FMetasoundFrontendClassInput* InOwningGraphClassInput=nullptr)
 		{
-			TArray<FDefaultLiteralData> DefaultLiteralData;
-
-			TArray<FMetasoundFrontendVertex> InputVertices = InNode.Interface.Inputs;
-			for (const FMetasoundFrontendVertexLiteral& Literal : InNode.InputLiterals)
+			// This exists here until we bring sub interface functionality to the frontend. 
+			// This check will no longer be valid at that point. We will have to implement
+			// a way to lookup the class's vertex access type info from the class def, or 
+			// move that info into the node's def.
+			check(InClass.Interface.Inputs.Num() == InNode.Interface.Inputs.Num());
+
+			// TODO: migrate sub interface builders from MetasoundVertex.cpp when
+			// sub interfaces exist on the frontend document
+			if (InNode.Interface.Inputs.Num())
 			{
-				FLiteralNodeConstructorParams InitParams;
-				FName TypeName;
-
-				FGuid VertexID = Literal.VertexID;
-				FVertexName DestinationVertexName;
-				bool bRequiresDefault = false;
-				auto RemoveAndBuildParams = [&](const FMetasoundFrontendVertex& Vertex)
+				auto MakeVertex = [&InContext](const FVertexName& Name, const FName& DataTypeName, EMetasoundFrontendVertexAccessType InFrontendAccessType, const FMetasoundFrontendLiteral* InFrontendLiteral) -> FInputDataVertex
 				{
-					if (Vertex.VertexID == VertexID)
+					EVertexAccessType AccessType = FrontendVertexAccessTypeToCoreVertexAccessType(InFrontendAccessType);
+					// Note: stripping metadata here. If we make this function public, we may want to keep metadata if WITH_EDITORONLY_DATA
+					if (InFrontendLiteral)
 					{
-						// Only build params and forward along to connect to dynamically generated
-						// literal node if edge is not explicitly connected to vertex as destination.
-						bRequiresDefault = !InEdgeDestinations.Contains({InNode.GetID(), Vertex.VertexID});
-						if (bRequiresDefault)
-						{
-							InitParams.Literal = Literal.Value.ToLiteral(Vertex.TypeName, &InContext.DataTypeRegistry, &InContext.ProxyDataCache);
-							InitParams.InstanceID = Frontend::CreateLocallyUniqueId();
-							InitParams.NodeName = "Literal";
-							TypeName = Vertex.TypeName;
-
-							DestinationVertexName = Vertex.Name;
-						}
-
-						return true;
+						FLiteral CoreLiteral = InFrontendLiteral->ToLiteral(DataTypeName, &InContext.DataTypeRegistry, &InContext.ProxyDataCache);
+						return FInputDataVertex{Name, DataTypeName, FDataVertexMetadata{}, AccessType, MoveTemp(CoreLiteral)};
+					}
+					else
+					{
+						return FInputDataVertex{Name, DataTypeName, FDataVertexMetadata{}, AccessType};
 					}
-
-					return false;
 				};
 
-				const bool bRemoved = InputVertices.RemoveAllSwap(RemoveAndBuildParams, EAllowShrinking::No) > 0;
-				if (ensure(bRemoved))
+				TArray<FInputDataVertex> Vertices;
+				Vertices.Reserve(InNode.Interface.Inputs.Num());
+
+				switch (InClass.Metadata.GetType())
 				{
-					if (bRequiresDefault)
-					{
-						DefaultLiteralData.Emplace(FDefaultLiteralData
+					// The vertices on an input/output node do not have matching
+					// names with the vertices on input/output node classes. The
+					// node class vertex names are blank, and the input/output node
+					// vertex names match the name of the input or output. Here 
+					// we have to grab data from both the class and node interfaces
+					case EMetasoundFrontendClassType::Input:
+						{
+							if (ensureMsgf(InOwningGraphClassInput, TEXT("Input nodes require the owning graph class input in order to retrieve the default literal")))
 							{
-								InNode.GetID(),
-								VertexID,
-								DestinationVertexName,
-								TypeName,
-								MoveTemp(InitParams)
-							});
-					}
+								check(InNode.Interface.Inputs.Num() == 1);
+								const FMetasoundFrontendVertex& NodeInputVertex = InNode.Interface.Inputs[0];
+								const FMetasoundFrontendClassInput& ClassInputVertex = InClass.Interface.Inputs[0];
+								const FMetasoundFrontendLiteral* FrontendLiteral = FindInputLiteralForInputNode(InNode, InClass, *InOwningGraphClassInput);
+								
+								Vertices.Add(MakeVertex(NodeInputVertex.Name, NodeInputVertex.TypeName, ClassInputVertex.AccessType, FrontendLiteral));
+							}
+						}
+						break;
+					case EMetasoundFrontendClassType::Output:
+						{
+							if (InNode.Interface.Inputs.Num() == 1)
+							{
+								const FMetasoundFrontendVertex& NodeInputVertex = InNode.Interface.Inputs[0];
+								const FMetasoundFrontendClassInput& ClassInputVertex = InClass.Interface.Inputs[0];
+								const FMetasoundFrontendLiteral* FrontendLiteral = FindLiteralForInputVertex(NodeInputVertex.Name, InNode, ClassInputVertex);
+								Vertices.Add(MakeVertex(NodeInputVertex.Name, NodeInputVertex.TypeName, ClassInputVertex.AccessType, FrontendLiteral));
+							}
+							else
+							{
+								// This is an edge case in the off scenario that someone
+								// is trying to build a graph that used the `TOutputNode<FStereoAudio>` 
+								// specialization.  Auto-update should have removed these
+								// long ago, but this log will help identify any issues.
+								// This log is added in UE 5.6 and if no one is hitting
+								// this issue, can be removed in UE 5.7.
+								UE_LOG(LogMetaSound, Error, TEXT("The MetaSound %s has an malformed output node %s. Please replace the output node and resave"), *InContext.DebugAssetName, *InNode.Name.ToString());
+							}
+						}
+						break;
+
+					default:
+						{
+							// As-it-stands, node interfaces and class interfaces
+							// match for all other nodes. This will change when
+							// node config is done as a feature. 
+							for (const FMetasoundFrontendClassInput& ClassInput : InClass.Interface.Inputs)
+							{
+								const FMetasoundFrontendLiteral* FrontendLiteral = FindLiteralForInputVertex(ClassInput.Name, InNode, ClassInput);
+								Vertices.Add(MakeVertex(ClassInput.Name, ClassInput.TypeName, ClassInput.AccessType, FrontendLiteral));
+							}
+						}
+						break;
 				}
+
+				return FInputVertexInterface{MoveTemp(Vertices), TArray<VertexPrivate::FSubInterfaceLayout>()};
+			}
+			else
+			{
+				return FInputVertexInterface{};
 			}
+		}
 
-			return DefaultLiteralData;
+		FVertexInterface CreateVertexInterface(const FBuildContext& InContext, const FMetasoundFrontendNode& InNode, const FMetasoundFrontendClass& InClass, const FMetasoundFrontendClassInput* InOwningGraphClassInput=nullptr)
+		{
+			return FVertexInterface(
+				CreateInputVertexInterface(InContext, InNode, InClass, InOwningGraphClassInput), 
+				CreateOutputVertexInterface(InContext, InNode, InClass), 
+				CreateEnvironmentVertexInterface(InNode.Interface));
 		}
 
+		FNodeData CreateNodeData(const FBuildContext& InContext, const FMetasoundFrontendNode& InNode, const FMetasoundFrontendClass& InClass, const FMetasoundFrontendClassInput* InOwningGraphClassInput=nullptr)
+		{
+			return FNodeData(InNode.Name, InNode.GetID(), CreateVertexInterface(InContext, InNode, InClass, InOwningGraphClassInput));
+		}
 
-		TUniquePtr<INode> CreateOutputNode(const FBuildContext& InContext, const FMetasoundFrontendNode& InNode, const FMetasoundFrontendClass& InClass, FBuildGraphContext& InGraphContext, const TSet<FNodeIDVertexID>& InEdgeDestinations)
+		const FMetasoundFrontendVariable* FindVariableForVariableNode(const FMetasoundFrontendNode& InVariableNode, const FMetasoundFrontendGraph& InGraph)
 		{
-			using namespace Metasound::Frontend;
+			const FGuid& DesiredID = InVariableNode.GetID();
+			return InGraph.Variables.FindByPredicate([&](const FMetasoundFrontendVariable& InVar) { return InVar.VariableNodeID == DesiredID; });
+		}
 
-			check(InClass.Metadata.GetType() == EMetasoundFrontendClassType::Output);
+		TUniquePtr<INode> CreateVariableNode(const FBuildContext& InContext, const FMetasoundFrontendNode& InNode, const FMetasoundFrontendClass& InClass, const FMetasoundFrontendGraph& InGraph)
+		{
+			check(InClass.Metadata.GetType() == EMetasoundFrontendClassType::Variable);
 			check(InNode.ClassID == InClass.ID);
 
-			if (ensure(InNode.Interface.Outputs.Num() == 1))
+			// Check frontend node and class interfaces are compatible
+			if (InClass.Interface.Inputs.Num() != InClass.Interface.Inputs.Num())
+			{
+				UE_LOG(LogMetaSound, Error, TEXT("MetaSound %s contains mismatched number of inputs (%d / %d) on input node %s class %s"), *InContext.DebugAssetName, InNode.Interface.Inputs.Num(), InClass.Interface.Inputs.Num(), *InNode.Name.ToString(), *InClass.Metadata.GetClassName().ToString());
+				return TUniquePtr<INode>(nullptr);
+			}
+			else if (InClass.Interface.Outputs.Num() != InClass.Interface.Outputs.Num())
 			{
-				const FMetasoundFrontendVertex& OutputVertex = InNode.Interface.Outputs[0];
+				UE_LOG(LogMetaSound, Error, TEXT("MetaSound %s contains mismatched number of outputs (%d / %d) on output node %s class %s"), *InContext.DebugAssetName, InNode.Interface.Outputs.Num(), InClass.Interface.Outputs.Num(), *InNode.Name.ToString(), *InClass.Metadata.GetClassName().ToString());
+				return TUniquePtr<INode>(nullptr);
+			}
 
-				FOutputNodeConstructorParams InitParams =
-				{
-					InNode.Name,
-					InNode.GetID(),
-					OutputVertex.Name
-				};
+			// Find the variable object associated with the node.
+			const FMetasoundFrontendVariable* FrontendVariable = FindVariableForVariableNode(InNode, InGraph);
 
-				{
-					const FNodeInitData InitData = GraphBuilderPrivate::CreateNodeInitData(InNode);
-					TArray<FDefaultLiteralData> DefaultLiteralData = GetInputDefaultLiteralData(InContext, InNode, InitData, InEdgeDestinations);
-					for (FDefaultLiteralData& Data : DefaultLiteralData)
-					{
-						InGraphContext.DefaultInputs.Emplace(FNodeIDVertexID { InNode.GetID(), Data.DestinationVertexID }, MoveTemp(Data));
-					}
-				}
+			if (nullptr != FrontendVariable)
+			{
+				const IDataTypeRegistry& DataTypeRegistry = InContext.DataTypeRegistry;
+				const bool IsLiteralParsableByDataType = DataTypeRegistry.IsLiteralTypeSupported(FrontendVariable->TypeName, FrontendVariable->Literal.GetType());
 
-				ensure(InClass.Interface.Outputs.Num() == 1);
-				if (InClass.Interface.Outputs[0].AccessType == EMetasoundFrontendVertexAccessType::Reference)
+				if (IsLiteralParsableByDataType)
 				{
-					return IDataTypeRegistry::Get().CreateOutputNode(OutputVertex.TypeName, MoveTemp(InitParams));
+					FLiteral Literal = FrontendVariable->Literal.ToLiteral(FrontendVariable->TypeName, &InContext.DataTypeRegistry, &InContext.ProxyDataCache);
+					return DataTypeRegistry.CreateVariableNode(FrontendVariable->TypeName, MoveTemp(Literal), CreateNodeData(InContext, InNode, InClass));
 				}
-				else // InClass.Interface.Outputs[0].AccessType == EMetasoundFrontendVertexAccessType::Value
+				else
 				{
-					return IDataTypeRegistry::Get().CreateConstructorOutputNode(OutputVertex.TypeName, MoveTemp(InitParams));
+					UE_LOG(LogMetaSound, Error, TEXT("Cannot create variable node [NodeID:%s]. [Variable:%s] cannot be constructed with the provided literal type."), *InNode.GetID().ToString(), *FrontendVariable->Name.ToString());
 				}
 			}
+			else
+			{
+				UE_LOG(LogMetaSound, Error, TEXT("Cannot create variable node [NodeID:%s]. No variable found for variable node."), *InNode.GetID().ToString());
+			}
+
+			return TUniquePtr<INode>(nullptr);
+		}
+
+
+		TUniquePtr<INode> CreateInputNode(const FBuildContext& InContext, const FMetasoundFrontendNode& InNode, const FMetasoundFrontendClass& InClass, const FMetasoundFrontendClassInput& InOwningGraphClassInput)
+		{
+			check(InClass.Metadata.GetType() == EMetasoundFrontendClassType::Input);
+			check(InNode.ClassID == InClass.ID);
+
+			// Check that the frontend node and class is correct.
+			if (InNode.Interface.Inputs.Num() != 1)
+			{
+				UE_LOG(LogMetaSound, Error, TEXT("MetaSound %s contains invalid number of inputs (%d) on input node %s"), *InContext.DebugAssetName, InNode.Interface.Inputs.Num(), *InNode.Name.ToString());
+				return TUniquePtr<INode>(nullptr);
+			}
+			else if (InClass.Interface.Inputs.Num() != 1)
+			{
+				UE_LOG(LogMetaSound, Error, TEXT("MetaSound %s contains invalid number of inputs (%d) on input class %s"), *InContext.DebugAssetName, InClass.Interface.Inputs.Num(), *InClass.Metadata.GetClassName().ToString());
+				return TUniquePtr<INode>(nullptr);
+			}
+			else if (InClass.Interface.Outputs.Num() != InClass.Interface.Outputs.Num())
+			{
+				UE_LOG(LogMetaSound, Error, TEXT("MetaSound %s contains mismatched number of outputs (%d / %d) on output node %s class %s"), *InContext.DebugAssetName, InNode.Interface.Outputs.Num(), InClass.Interface.Outputs.Num(), *InNode.Name.ToString(), *InClass.Metadata.GetClassName().ToString());
+				return TUniquePtr<INode>(nullptr);
+			}
+
+			// Create the input node
+			const FMetasoundFrontendVertex& InputVertex = InNode.Interface.Inputs[0];
+
+			if (InOwningGraphClassInput.AccessType == EMetasoundFrontendVertexAccessType::Reference)
+			{
+				return InContext.DataTypeRegistry.CreateInputNode(InputVertex.TypeName, CreateNodeData(InContext, InNode, InClass, &InOwningGraphClassInput));
+			}
+			else if (InOwningGraphClassInput.AccessType == EMetasoundFrontendVertexAccessType::Value)
+			{
+				return InContext.DataTypeRegistry.CreateConstructorInputNode(InputVertex.TypeName, CreateNodeData(InContext, InNode, InClass, &InOwningGraphClassInput));
+			}
+			else
+			{
+				UE_LOG(LogMetaSound, Error, TEXT("MetaSound %s contains invalid input access type on input %s"), *InContext.DebugAssetName, *InNode.Name.ToString());
+			}
+
 			return TUniquePtr<INode>(nullptr);
 		}
 
-		TUniquePtr<INode> CreateExternalNode(const FBuildContext& InContext, const FMetasoundFrontendNode& InNode, const FMetasoundFrontendClass& InClass, FBuildGraphContext& InGraphContext, const TSet<FNodeIDVertexID>& InEdgeDestinations)
+		TUniquePtr<INode> CreateOutputNode(const FBuildContext& InContext, const FMetasoundFrontendNode& InNode, const FMetasoundFrontendClass& InClass, FBuildGraphContext& InGraphContext)
+		{
+			check(InClass.Metadata.GetType() == EMetasoundFrontendClassType::Output);
+			check(InNode.ClassID == InClass.ID);
+
+			if (InClass.Interface.Inputs.Num() != InClass.Interface.Inputs.Num())
+			{
+				UE_LOG(LogMetaSound, Error, TEXT("MetaSound %s contains mismatched number of inputs (%d / %d) on input node %s class %s"), *InContext.DebugAssetName, InNode.Interface.Inputs.Num(), InClass.Interface.Inputs.Num(), *InNode.Name.ToString(), *InClass.Metadata.GetClassName().ToString());
+				return TUniquePtr<INode>(nullptr);
+			}
+			else if (InNode.Interface.Outputs.Num() != 1)
+			{
+				UE_LOG(LogMetaSound, Error, TEXT("MetaSound %s contains invalid number of outputs (%d) on output node %s"), *InContext.DebugAssetName, InNode.Interface.Outputs.Num(), *InNode.Name.ToString());
+				return TUniquePtr<INode>(nullptr);
+			}
+			else if (InClass.Interface.Outputs.Num() != 1)
+			{
+				UE_LOG(LogMetaSound, Error, TEXT("MetaSound %s contains invalid number of outputs (%d) on output class %s"), *InContext.DebugAssetName, InClass.Interface.Outputs.Num(), *InClass.Metadata.GetClassName().ToString());
+				return TUniquePtr<INode>(nullptr);
+			}
+
+			const FMetasoundFrontendVertex& OutputVertex = InNode.Interface.Outputs[0];
+
+			if (InClass.Interface.Outputs[0].AccessType == EMetasoundFrontendVertexAccessType::Reference)
+			{
+				return InContext.DataTypeRegistry.CreateOutputNode(OutputVertex.TypeName, CreateNodeData(InContext, InNode, InClass));
+			}
+			else if (InClass.Interface.Outputs[0].AccessType == EMetasoundFrontendVertexAccessType::Value)
+			{
+				return InContext.DataTypeRegistry.CreateConstructorOutputNode(OutputVertex.TypeName, CreateNodeData(InContext, InNode, InClass));
+			}
+			else
+			{
+				UE_LOG(LogMetaSound, Error, TEXT("MetaSound %s contains invalid output access type on output %s"), *InContext.DebugAssetName, *InNode.Name.ToString());
+			}
+
+			return TUniquePtr<INode>(nullptr);
+		}
+
+		TUniquePtr<INode> CreateExternalNode(const FBuildContext& InContext, const FMetasoundFrontendNode& InNode, const FMetasoundFrontendClass& InClass)
 		{
 			using namespace Frontend;
 
 			check(InNode.ClassID == InClass.ID);
 
-			const FNodeInitData InitData = GraphBuilderPrivate::CreateNodeInitData(InNode);
+			// Check frontend node and class interfaces are compatible
+			if (InClass.Interface.Inputs.Num() != InClass.Interface.Inputs.Num())
 			{
-				TArray<FDefaultLiteralData> DefaultLiteralData = GetInputDefaultLiteralData(InContext, InNode, InitData, InEdgeDestinations);
-				for (FDefaultLiteralData& Data : DefaultLiteralData)
-				{
-					InGraphContext.DefaultInputs.Emplace(FNodeIDVertexID{ InNode.GetID(), Data.DestinationVertexID }, MoveTemp(Data));
-				}
+				UE_LOG(LogMetaSound, Error, TEXT("MetaSound %s contains mismatched number of inputs (%d / %d) on input node %s class %s"), *InContext.DebugAssetName, InNode.Interface.Inputs.Num(), InClass.Interface.Inputs.Num(), *InNode.Name.ToString(), *InClass.Metadata.GetClassName().ToString());
+				return TUniquePtr<INode>(nullptr);
+			}
+			else if (InClass.Interface.Outputs.Num() != InClass.Interface.Outputs.Num())
+			{
+				UE_LOG(LogMetaSound, Error, TEXT("MetaSound %s contains mismatched number of outputs (%d / %d) on output node %s class %s"), *InContext.DebugAssetName, InNode.Interface.Outputs.Num(), InClass.Interface.Outputs.Num(), *InNode.Name.ToString(), *InClass.Metadata.GetClassName().ToString());
+				return TUniquePtr<INode>(nullptr);
 			}
 
-			// TODO: handle check to see if node interface conforms to class interface here. 
-			// TODO: check to see if external object supports class interface.
 
 			const FNodeRegistryKey Key = FNodeRegistryKey(InClass.Metadata);
-			return FMetasoundFrontendRegistryContainer::Get()->CreateNode(Key, InitData);
+			return INodeClassRegistry::Get()->CreateNode(Key, CreateNodeData(InContext, InNode, InClass));
 		}
 
 		const FMetasoundFrontendClassInput* FindClassInputForInputNode(const FMetasoundFrontendGraphClass& InOwningGraph, const FMetasoundFrontendNode& InInputNode)
@@ -433,15 +566,13 @@ namespace Metasound::Frontend
 			return nullptr;
 		}
 
-		bool AddNodesToGraph(FBuildGraphContext& InGraphContext)
+		void LogFailedToCreateNode(const FString& InAssetName, const FMetasoundFrontendNode& InNode, const FMetasoundFrontendClass& InNodeClass)
 		{
-			TSet<FNodeIDVertexID> GraphEdgeDestinations;
-			const TArray<FMetasoundFrontendEdge>& GraphEdges = InGraphContext.PagedGraph.Edges;
-			Algo::Transform(GraphEdges, GraphEdgeDestinations, [](const FMetasoundFrontendEdge& Edge) 
-			{
-				return FNodeIDVertexID{Edge.ToNodeID, Edge.ToVertexID};
-			});
+			UE_LOG(LogMetaSound, Error, TEXT("Metasound '%s': Failed to create node [NodeID:%s, NodeName:%s] from registry [Class: %s %s]"), *InAssetName, *InNode.GetID().ToString(), *InNode.Name.ToString(), *InNodeClass.Metadata.GetClassName().ToString(), *InNodeClass.Metadata.GetVersion().ToString());
+		}
 
+		bool AddNodesToGraph(FBuildGraphContext& InGraphContext)
+		{
 			for (const FMetasoundFrontendNode& Node : InGraphContext.PagedGraph.Nodes)
 			{
 				const FMetasoundFrontendClass* NodeClass = InGraphContext.BuildContext.FrontendClasses.FindRef(Node.ClassID);
@@ -457,7 +588,14 @@ namespace Metasound::Frontend
 							if (nullptr != ClassInput)
 							{
 								TSharedPtr<const INode> InputNode(CreateInputNode(InGraphContext.BuildContext, Node, *NodeClass, *ClassInput).Release());
-								InGraphContext.Graph->AddInputNode(Node.GetID(), ClassInput->Name, InputNode);
+								if (InputNode.IsValid())
+								{
+									InGraphContext.Graph->AddInputNode(Node.GetID(), ClassInput->Name, InputNode);
+								}
+								else
+								{
+									LogFailedToCreateNode(InGraphContext.BuildContext.DebugAssetName, Node, *NodeClass);
+								}
 							}
 							else
 							{
@@ -473,8 +611,15 @@ namespace Metasound::Frontend
 							const FMetasoundFrontendClassOutput* ClassOutput = FindClassOutputForOutputNode(InGraphContext.GraphClass, Node);
 							if (nullptr != ClassOutput)
 							{
-								TSharedPtr<const INode> OutputNode(CreateOutputNode(InGraphContext.BuildContext, Node, *NodeClass, InGraphContext, GraphEdgeDestinations).Release());
-								InGraphContext.Graph->AddOutputNode(Node.GetID(), ClassOutput->Name, OutputNode);
+								TSharedPtr<const INode> OutputNode(CreateOutputNode(InGraphContext.BuildContext, Node, *NodeClass, InGraphContext).Release());
+								if (OutputNode)
+								{
+									InGraphContext.Graph->AddOutputNode(Node.GetID(), ClassOutput->Name, OutputNode);
+								}
+								else
+								{
+									LogFailedToCreateNode(InGraphContext.BuildContext.DebugAssetName, Node, *NodeClass);
+								}
 							}
 							else
 							{
@@ -509,8 +654,15 @@ namespace Metasound::Frontend
 
 						case EMetasoundFrontendClassType::Variable:
 						{
-							TSharedPtr<const INode> VariableNode(CreateVariableNode(InGraphContext.BuildContext, Node, InGraphContext.PagedGraph).Release());
-							InGraphContext.Graph->AddNode(Node.GetID(), VariableNode);
+							TSharedPtr<const INode> VariableNode(CreateVariableNode(InGraphContext.BuildContext, Node, *NodeClass, InGraphContext.PagedGraph).Release());
+							if (VariableNode)
+							{
+								InGraphContext.Graph->AddNode(Node.GetID(), VariableNode);
+							}
+							else
+							{
+								LogFailedToCreateNode(InGraphContext.BuildContext.DebugAssetName, Node, *NodeClass);
+							}
 						}
 						break;
 
@@ -523,8 +675,15 @@ namespace Metasound::Frontend
 						case EMetasoundFrontendClassType::External:
 						default:
 						{
-							TSharedPtr<const INode> ExternalNode(CreateExternalNode(InGraphContext.BuildContext, Node, *NodeClass, InGraphContext, GraphEdgeDestinations).Release());
-							InGraphContext.Graph->AddNode(Node.GetID(), ExternalNode);
+							TSharedPtr<const INode> ExternalNode(CreateExternalNode(InGraphContext.BuildContext, Node, *NodeClass).Release());
+							if (ExternalNode)
+							{
+								InGraphContext.Graph->AddNode(Node.GetID(), ExternalNode);
+							}
+							else
+							{
+								LogFailedToCreateNode(InGraphContext.BuildContext.DebugAssetName, Node, *NodeClass);
+							}
 						}
 						break;
 					}
@@ -609,70 +768,16 @@ namespace Metasound::Frontend
 				bool bSuccess = InGraphContext.Graph->AddDataEdge(*FromNode, FromVertexKey,  *ToNode, ToVertexKey);
 
 				// If succeeded, remove input as viable vertex to construct default variable, as it has been superceded by a connection.
-				if (bSuccess)
-				{
-					FNodeIDVertexID DestinationPair { ToNode->GetInstanceID(), DestinationNodeAndVertex->Vertex->VertexID };
-					InGraphContext.DefaultInputs.Remove(DestinationPair);
-				}
-				else
-				{
-					UE_LOG(LogMetaSound, Error, TEXT("MetaSound '%s': Failed to connect edge from [NodeID:%s, VertexID:%s] to [NodeID:%s, VertexID:%s]"), *InGraphContext.BuildContext.DebugAssetName, *Edge.FromNodeID.ToString(), *Edge.FromVertexID.ToString(), *Edge.ToNodeID.ToString(), *Edge.ToVertexID.ToString());
-					return false;
-				}
-			}
-
-			return true;
-		}
-
-		bool AddDefaultInputLiterals(FBuildGraphContext& InGraphContext)
-		{
-			using namespace Metasound::Frontend;
-			using namespace LiteralNodeNames; 
-
-			using FIterator = FDefaultInputByIDMap::TIterator;
-
-			for (FDefaultInputByIDMap::ElementType& Pair : InGraphContext.DefaultInputs)
-			{
-				FDefaultLiteralData& LiteralData = Pair.Value;
-				const FGuid LiteralNodeID = Frontend::CreateLocallyUniqueId();
-
-				// 1. Construct and add the default variable to the graph
-				{
-					TUniquePtr<const INode> DefaultLiteral = IDataTypeRegistry::Get().CreateLiteralNode(LiteralData.TypeName, MoveTemp(LiteralData.InitParams));
-					InGraphContext.Graph->AddNode(LiteralNodeID, TSharedPtr<const INode>(DefaultLiteral.Release()));
-				}
-
-				// 2. Connect the default variable to the expected input
-				const INode* FromNode = InGraphContext.Graph->FindNode(LiteralNodeID);
-				if (nullptr == FromNode)
-				{
-					UE_LOG(LogMetaSound, Error, TEXT("MetaSound '%s': Failed to find node in graph [NodeID:%s]"), *InGraphContext.BuildContext.DebugAssetName, *LiteralNodeID.ToString());
-					return false;
-				}
-				const FVertexName& FromVertexName = METASOUND_GET_PARAM_NAME(OutputValue);
-
-				const INode* ToNode = InGraphContext.Graph->FindNode(LiteralData.DestinationNodeID);
-				if (nullptr == ToNode)
-				{
-					UE_LOG(LogMetaSound, Error, TEXT("MetaSound '%s': Failed to node in graph [NodeID:%s]"), *InGraphContext.BuildContext.DebugAssetName, *LiteralData.DestinationNodeID.ToString());
-					return false;
-				}
-				const FVertexName& ToVertexName = LiteralData.DestinationVertexKey;
-
-				bool bSuccess = InGraphContext.Graph->AddDataEdge(*FromNode, FromVertexName, *ToNode, ToVertexName);
 				if (!bSuccess)
 				{
-					UE_LOG(LogMetaSound, Error, TEXT("MetaSound '%s': Failed to connect default variable edge: from '%s' to '%s'"), *InGraphContext.BuildContext.DebugAssetName, *FromVertexName.ToString(), *ToVertexName.ToString());
+					UE_LOG(LogMetaSound, Error, TEXT("MetaSound '%s': Failed to connect edge from [NodeID:%s, VertexID:%s] to [NodeID:%s, VertexID:%s]"), *InGraphContext.BuildContext.DebugAssetName, *Edge.FromNodeID.ToString(), *Edge.FromVertexID.ToString(), *Edge.ToNodeID.ToString(), *Edge.ToVertexID.ToString());
 					return false;
 				}
 			}
 
-			// Clear default inputs because literals have been moved out of default input map.
-			InGraphContext.DefaultInputs.Reset();
 			return true;
 		}
 
-
 		bool SortSubgraphDependencies(TArray<const FMetasoundFrontendGraphClass*>& Subgraphs)
 		{
 			// Helper for caching and querying subgraph dependencies
@@ -742,11 +847,6 @@ namespace Metasound::Frontend
 				bSuccess = AddEdgesToGraph(BuildGraphContext);
 			}
 
-			if (bSuccess)
-			{
-				bSuccess = AddDefaultInputLiterals(BuildGraphContext);
-			}
-
 			if (bSuccess)
 			{
 				return MoveTemp(BuildGraphContext.Graph);
