Static Analysis fix for MetasoundFrontendGraphBuilder
#rnx

diff --git a/Engine/Plugins/Runtime/Metasound/Source/MetasoundFrontend/Private/MetasoundFrontendGraphBuilder.cpp b/Engine/Plugins/Runtime/Metasound/Source/MetasoundFrontend/Private/MetasoundFrontendGraphBuilder.cpp
--- a/Engine/Plugins/Runtime/Metasound/Source/MetasoundFrontend/Private/MetasoundFrontendGraphBuilder.cpp
+++ b/Engine/Plugins/Runtime/Metasound/Source/MetasoundFrontend/Private/MetasoundFrontendGraphBuilder.cpp
@@ -391,12 +391,12 @@ namespace Metasound::Frontend
 			check(InNode.ClassID == InClass.ID);
 
 			// Check frontend node and class interfaces are compatible
-			if (InClass.Interface.Inputs.Num() != InClass.Interface.Inputs.Num())
+			if (InClass.Interface.Inputs.Num() != InNode.Interface.Inputs.Num())
 			{
 				UE_LOG(LogMetaSound, Error, TEXT("MetaSound %s contains mismatched number of inputs (%d / %d) on input node %s class %s"), *InContext.DebugAssetName, InNode.Interface.Inputs.Num(), InClass.Interface.Inputs.Num(), *InNode.Name.ToString(), *InClass.Metadata.GetClassName().ToString());
 				return TUniquePtr<INode>(nullptr);
 			}
-			else if (InClass.Interface.Outputs.Num() != InClass.Interface.Outputs.Num())
+			else if (InClass.Interface.Outputs.Num() != InNode.Interface.Outputs.Num())
 			{
 				UE_LOG(LogMetaSound, Error, TEXT("MetaSound %s contains mismatched number of outputs (%d / %d) on output node %s class %s"), *InContext.DebugAssetName, InNode.Interface.Outputs.Num(), InClass.Interface.Outputs.Num(), *InNode.Name.ToString(), *InClass.Metadata.GetClassName().ToString());
 				return TUniquePtr<INode>(nullptr);
@@ -445,7 +445,7 @@ namespace Metasound::Frontend
 				UE_LOG(LogMetaSound, Error, TEXT("MetaSound %s contains invalid number of inputs (%d) on input class %s"), *InContext.DebugAssetName, InClass.Interface.Inputs.Num(), *InClass.Metadata.GetClassName().ToString());
 				return TUniquePtr<INode>(nullptr);
 			}
-			else if (InClass.Interface.Outputs.Num() != InClass.Interface.Outputs.Num())
+			else if (InClass.Interface.Outputs.Num() != InNode.Interface.Outputs.Num())
 			{
 				UE_LOG(LogMetaSound, Error, TEXT("MetaSound %s contains mismatched number of outputs (%d / %d) on output node %s class %s"), *InContext.DebugAssetName, InNode.Interface.Outputs.Num(), InClass.Interface.Outputs.Num(), *InNode.Name.ToString(), *InClass.Metadata.GetClassName().ToString());
 				return TUniquePtr<INode>(nullptr);
@@ -475,7 +475,7 @@ namespace Metasound::Frontend
 			check(InClass.Metadata.GetType() == EMetasoundFrontendClassType::Output);
 			check(InNode.ClassID == InClass.ID);
 
-			if (InClass.Interface.Inputs.Num() != InClass.Interface.Inputs.Num())
+			if (InClass.Interface.Inputs.Num() != InNode.Interface.Inputs.Num())
 			{
 				UE_LOG(LogMetaSound, Error, TEXT("MetaSound %s contains mismatched number of inputs (%d / %d) on input node %s class %s"), *InContext.DebugAssetName, InNode.Interface.Inputs.Num(), InClass.Interface.Inputs.Num(), *InNode.Name.ToString(), *InClass.Metadata.GetClassName().ToString());
 				return TUniquePtr<INode>(nullptr);
@@ -516,12 +516,12 @@ namespace Metasound::Frontend
 			check(InNode.ClassID == InClass.ID);
 
 			// Check frontend node and class interfaces are compatible
-			if (InClass.Interface.Inputs.Num() != InClass.Interface.Inputs.Num())
+			if (InClass.Interface.Inputs.Num() != InNode.Interface.Inputs.Num())
 			{
 				UE_LOG(LogMetaSound, Error, TEXT("MetaSound %s contains mismatched number of inputs (%d / %d) on input node %s class %s"), *InContext.DebugAssetName, InNode.Interface.Inputs.Num(), InClass.Interface.Inputs.Num(), *InNode.Name.ToString(), *InClass.Metadata.GetClassName().ToString());
 				return TUniquePtr<INode>(nullptr);
 			}
-			else if (InClass.Interface.Outputs.Num() != InClass.Interface.Outputs.Num())
+			else if (InClass.Interface.Outputs.Num() != InNode.Interface.Outputs.Num())
 			{
 				UE_LOG(LogMetaSound, Error, TEXT("MetaSound %s contains mismatched number of outputs (%d / %d) on output node %s class %s"), *InContext.DebugAssetName, InNode.Interface.Outputs.Num(), InClass.Interface.Outputs.Num(), *InNode.Name.ToString(), *InClass.Metadata.GetClassName().ToString());
 				return TUniquePtr<INode>(nullptr);
