Refactor literal node usage in the MetaSound dynamic operator for node configuration, various unit test fixups

diff --git a/Engine/Plugins/Runtime/Metasound/Source/MetasoundEngine/Private/MetasoundBuilderSubsystem.cpp b/Engine/Plugins/Runtime/Metasound/Source/MetasoundEngine/Private/MetasoundBuilderSubsystem.cpp
--- a/Engine/Plugins/Runtime/Metasound/Source/MetasoundEngine/Private/MetasoundBuilderSubsystem.cpp
+++ b/Engine/Plugins/Runtime/Metasound/Source/MetasoundEngine/Private/MetasoundBuilderSubsystem.cpp
@@ -65,16 +65,6 @@ namespace Metasound::Engine
 			return Literal;
 		}
 
-		TUniquePtr<INode> CreateDynamicNodeFromFrontendLiteral(const FName DataType, const FMetasoundFrontendLiteral& InLiteral)
-		{
-			using namespace Frontend;
-			FLiteral ValueLiteral = InLiteral.ToLiteral(DataType);
-
-			// TODO: Node name "Literal" is always the same.  Consolidate and deprecate providing unique node name to avoid unnecessary FName table bloat.
-			FLiteralNodeConstructorParams Params { "Literal", FGuid::NewGuid(), MoveTemp(ValueLiteral) };
-			return IDataTypeRegistry::Get().CreateLiteralNode(DataType, MoveTemp(Params));
-		}
-
 		TOptional<FMetasoundFrontendLiteral> TryResolveNodeInputDefault(const FMetaSoundFrontendDocumentBuilder& Builder, const FGuid& InNodeID, FName VertexName)
 		{
 			TOptional<FMetasoundFrontendLiteral> DefaultLiteral;
@@ -508,8 +498,9 @@ void UMetaSoundSourceBuilder::OnNodeAdded(int32 NodeIndex) const
 			case EMetasoundFrontendClassType::External:
 			case EMetasoundFrontendClassType::Graph:
 			{
-				const Metasound::FNodeInitData InitData { AddedNode.Name, AddedNode.GetID() };
-				NewNode = NodeRegistry.CreateNode(ClassKey, InitData);
+				FVertexInterface VertexInterface;
+				NodeRegistry.FindDefaultVertexInterface(ClassKey, VertexInterface);
+				NewNode = NodeRegistry.CreateNode(ClassKey, FNodeData{ AddedNode.Name, AddedNode.GetID() , VertexInterface });
 			}
 			break;
 
@@ -540,40 +531,53 @@ void UMetaSoundSourceBuilder::OnNodeAdded(int32 NodeIndex) const
 					DefaultLiteral = &(NodeClass->Interface.Inputs.Last().FindConstDefaultChecked(Frontend::DefaultPageID));
 				}
 
-				FInputNodeConstructorParams InitData
-				{
-					AddedNode.Name,
-					AddedNode.GetID(),
-					InputVertex.Name,
-					DefaultLiteral->ToLiteral(DataTypeName)
-				};
+				TSharedPtr<const FNodeClassMetadata> InputClassMetadata = DataTypeRegistry.GetInputClassMetadata(DataTypeName);
+				checkf(InputClassMetadata, TEXT("Input class metadata could not be found from data type registry"));
 
-				NewNode = DataTypeRegistry.CreateInputNode(DataTypeName, MoveTemp(InitData));
+				FVertexInterface InputNodeInterface = InputClassMetadata->DefaultInterface;
+				// Vertex names must be set for input nodes
+				InputNodeInterface.GetInputInterface().At(0).VertexName = InputVertex.Name;
+				InputNodeInterface.GetOutputInterface().At(0).VertexName = InputVertex.Name;
+				InputNodeInterface.GetInputInterface().At(0).SetDefaultLiteral(DefaultLiteral->ToLiteral(DataTypeName));
+
+				NewNode = DataTypeRegistry.CreateInputNode(DataTypeName, FNodeData{ AddedNode.Name, AddedNode.GetID(), InputNodeInterface });
 			}
 			break;
 
 			case EMetasoundFrontendClassType::Variable:
 			{
 				const FName& DataTypeName = NodeClass->Metadata.GetClassName().Name;
-				FDefaultLiteralNodeConstructorParams InitData { AddedNode.Name, AddedNode.GetID(), DataTypeRegistry.CreateDefaultLiteral(DataTypeName)};
-				NewNode = DataTypeRegistry.CreateVariableNode(DataTypeName, MoveTemp(InitData));
+				TSharedPtr<const FNodeClassMetadata> VariableMetadata = DataTypeRegistry.GetVariableClassMetadata(DataTypeName);
+				checkf(VariableMetadata, TEXT("Variable class metadata could not be found from data type registry"));
+				NewNode = DataTypeRegistry.CreateVariableNode(DataTypeName, DataTypeRegistry.CreateDefaultLiteral(DataTypeName),
+					FNodeData{ AddedNode.Name, AddedNode.GetID(), VariableMetadata->DefaultInterface }); 
 			}
 			break;
 
 			case EMetasoundFrontendClassType::Literal:
 			{
+				UE_LOG(LogMetaSound, Warning, TEXT("Builder '%s': Adding literal nodes is being deprecated. Please set literal values on the connected node interface directly."), *GetName());
 				const FName& DataTypeName = NodeClass->Metadata.GetClassName().Name;
+PRAGMA_DISABLE_DEPRECATION_WARNINGS
 				FDefaultLiteralNodeConstructorParams InitData { AddedNode.Name, AddedNode.GetID(), DataTypeRegistry.CreateDefaultLiteral(DataTypeName)};
 				NewNode = DataTypeRegistry.CreateLiteralNode(DataTypeName, MoveTemp(InitData));
+PRAGMA_ENABLE_DEPRECATION_WARNINGS
 			}
 			break;
 
 			case EMetasoundFrontendClassType::Output:
 			{
 				const FName& DataTypeName = NodeClass->Metadata.GetClassName().Name;
-				const FMetasoundFrontendVertex& OutputVertex = AddedNode.Interface.Outputs.Last();
-				FDefaultNamedVertexNodeConstructorParams InitData { AddedNode.Name, AddedNode.GetID(), OutputVertex.Name };
-				NewNode = DataTypeRegistry.CreateOutputNode(DataTypeName, MoveTemp(InitData));
+				const FName& OutputVertexName = AddedNode.Interface.Outputs.Last().Name;
+
+				TSharedPtr<const FNodeClassMetadata> OutputClassMetadata = DataTypeRegistry.GetOutputClassMetadata(DataTypeName);
+				checkf(OutputClassMetadata, TEXT("Output class metadata could not be found from data type registry"));
+
+				FVertexInterface OutputNodeInterface = OutputClassMetadata->DefaultInterface;
+				// Vertex names must be set for output nodes
+				OutputNodeInterface.GetInputInterface().At(0).VertexName = OutputVertexName;
+				OutputNodeInterface.GetOutputInterface().At(0).VertexName = OutputVertexName;
+				NewNode = DataTypeRegistry.CreateOutputNode(DataTypeName, FNodeData{ AddedNode.Name, AddedNode.GetID(), OutputNodeInterface });
 			}
 			break;
 
@@ -610,8 +614,8 @@ void UMetaSoundSourceBuilder::OnNodeInputLiteralSet(int32 NodeIndex, int32 Verte
 			using namespace Metasound::Engine;
 
 			const FMetasoundFrontendLiteral& InputDefault = Node.InputLiterals[LiteralIndex].Value;
-			TUniquePtr<INode> LiteralNode = BuilderSubsystemPrivate::CreateDynamicNodeFromFrontendLiteral(Input.TypeName, InputDefault);
-			Transactor.SetValue(Node.GetID(), Input.Name, MoveTemp(LiteralNode));
+			
+			Transactor.SetValue(Node.GetID(), Input.Name, InputDefault.ToLiteral(Input.TypeName), &UMetaSoundSourceBuilder::CreateDataReference);
 			return true;
 		});
 	}
@@ -669,8 +673,7 @@ void UMetaSoundSourceBuilder::OnRemoveSwappingEdge(int32 SwapIndex, int32 LastIn
 
 			if (ensureAlwaysMsgf(InputDefault.IsSet(), TEXT("Could not dynamically assign default literal upon removing edge: literal should be assigned by either the frontend document's input or the class definition")))
 			{
-				TUniquePtr<INode> LiteralNode = BuilderSubsystemPrivate::CreateDynamicNodeFromFrontendLiteral(ToNodeInput->TypeName, *InputDefault);
-				Transactor.RemoveDataEdge(EdgeBeingRemoved.FromNodeID, FromNodeOutput->Name, EdgeBeingRemoved.ToNodeID, ToNodeInput->Name, MoveTemp(LiteralNode));
+				Transactor.RemoveDataEdge(EdgeBeingRemoved.FromNodeID, FromNodeOutput->Name, EdgeBeingRemoved.ToNodeID, ToNodeInput->Name, InputDefault->ToLiteral(ToNodeInput->TypeName), &UMetaSoundSourceBuilder::CreateDataReference);
 				return true;
 			}
 		}
@@ -765,8 +768,7 @@ void UMetaSoundSourceBuilder::OnRemovingNodeInputLiteral(int32 NodeIndex, int32
 			TOptional<FMetasoundFrontendLiteral> InputDefault = BuilderSubsystemPrivate::TryResolveNodeInputDefault(Builder, Node.GetID(), Input.Name);
 			if (ensureAlwaysMsgf(InputDefault.IsSet(), TEXT("Could not dynamically assign default literal from class definition upon removing input '%s' literal: document's dependency entry invalid and has no default assigned"), *Input.Name.ToString()))
 			{
-				TUniquePtr<INode> LiteralNode = BuilderSubsystemPrivate::CreateDynamicNodeFromFrontendLiteral(Input.TypeName, *InputDefault);
-				Transactor.SetValue(Node.GetID(), Input.Name, MoveTemp(LiteralNode));
+				Transactor.SetValue(Node.GetID(), Input.Name, InputDefault->ToLiteral(Input.TypeName), CreateDataReference);
 				return true;
 			}
 
diff --git a/Engine/Plugins/Runtime/Metasound/Source/MetasoundEngineTest/Private/EngineTestMetaSoundAutomatedNodeTest.cpp b/Engine/Plugins/Runtime/Metasound/Source/MetasoundEngineTest/Private/EngineTestMetaSoundAutomatedNodeTest.cpp
--- a/Engine/Plugins/Runtime/Metasound/Source/MetasoundEngineTest/Private/EngineTestMetaSoundAutomatedNodeTest.cpp
+++ b/Engine/Plugins/Runtime/Metasound/Source/MetasoundEngineTest/Private/EngineTestMetaSoundAutomatedNodeTest.cpp
@@ -682,47 +682,59 @@ namespace Metasound::EngineTest{
 			case EMetasoundFrontendClassType::External:
 			case EMetasoundFrontendClassType::Graph:
 			{
-				FNodeInitData NodeInitData{TestNodeName, TestNodeID};
-				Node = NodeRegistry->CreateNode(InNodeRegistryKey, NodeInitData);
+				FVertexInterface VertexInterface;
+				NodeRegistry->FindDefaultVertexInterface(InNodeRegistryKey, VertexInterface);
+				Node = NodeRegistry->CreateNode(InNodeRegistryKey, FNodeData{ TestNodeName, TestNodeID , VertexInterface });
 			}
 			break;
 
 			case EMetasoundFrontendClassType::Input:
 			{
 				FName DataTypeName = NodeClass.Metadata.GetClassName().Name;
-				FInputNodeConstructorParams NodeInitData
-				{
-					TestNodeName, 
-					TestNodeID, 
-					TestVertexName,
-					NodeClass.Interface.Inputs[0].FindConstDefaultChecked(Frontend::DefaultPageID).ToLiteral(DataTypeName)
-				};
+				TSharedPtr<const FNodeClassMetadata> InputClassMetadata = DataTypeRegistry.GetInputClassMetadata(DataTypeName);
+				checkf(InputClassMetadata, TEXT("Input class metadata could not be found from data type registry"));
+
+				FVertexInterface InputNodeInterface = InputClassMetadata->DefaultInterface;
+				// Vertex names must be set for input nodes
+				InputNodeInterface.GetInputInterface().At(0).VertexName = TestVertexName;
+				InputNodeInterface.GetOutputInterface().At(0).VertexName = TestVertexName;
+				InputNodeInterface.GetInputInterface().At(0).SetDefaultLiteral(NodeClass.Interface.Inputs[0].FindConstDefaultChecked(Frontend::DefaultPageID).ToLiteral(DataTypeName));
 
-				Node = DataTypeRegistry.CreateInputNode(DataTypeName, MoveTemp(NodeInitData));
+				Node = DataTypeRegistry.CreateInputNode(DataTypeName, FNodeData{ TestNodeName, TestNodeID, InputNodeInterface });
 			}
 			break;
 
 			case EMetasoundFrontendClassType::Variable:
 			{
 				FName DataTypeName = NodeClass.Metadata.GetClassName().Name;
-				FDefaultLiteralNodeConstructorParams NodeInitData{TestNodeName, TestNodeID, DataTypeRegistry.CreateDefaultLiteral(DataTypeName)};
-				Node = DataTypeRegistry.CreateVariableNode(DataTypeName, MoveTemp(NodeInitData));
+				TSharedPtr<const FNodeClassMetadata> VariableMetadata = DataTypeRegistry.GetVariableClassMetadata(DataTypeName);
+				checkf(VariableMetadata, TEXT("Variable class metadata could not be found from data type registry"));
+				Node = DataTypeRegistry.CreateVariableNode(DataTypeName, DataTypeRegistry.CreateDefaultLiteral(DataTypeName),
+					FNodeData{ TestNodeName, TestNodeID, VariableMetadata->DefaultInterface });
 			}
 			break;
 
 			case EMetasoundFrontendClassType::Literal:
 			{
 				FName DataTypeName = NodeClass.Metadata.GetClassName().Name;
+PRAGMA_DISABLE_DEPRECATION_WARNINGS
 				FDefaultLiteralNodeConstructorParams NodeInitData{TestNodeName, TestNodeID, DataTypeRegistry.CreateDefaultLiteral(DataTypeName)};
 				Node = DataTypeRegistry.CreateLiteralNode(DataTypeName, MoveTemp(NodeInitData));
+PRAGMA_ENABLE_DEPRECATION_WARNINGS
 			}
 			break;
 
 			case EMetasoundFrontendClassType::Output:
 			{
 				FName DataTypeName = NodeClass.Metadata.GetClassName().Name;
-				FDefaultNamedVertexNodeConstructorParams NodeInitData{TestNodeName, TestNodeID, TestVertexName};
-				Node = DataTypeRegistry.CreateOutputNode(DataTypeName, MoveTemp(NodeInitData));
+				TSharedPtr<const FNodeClassMetadata> OutputClassMetadata = DataTypeRegistry.GetOutputClassMetadata(DataTypeName);
+				checkf(OutputClassMetadata, TEXT("Output class metadata could not be found from data type registry"));
+
+				FVertexInterface OutputNodeInterface = OutputClassMetadata->DefaultInterface;
+				// Vertex names must be set for output nodes
+				OutputNodeInterface.GetInputInterface().At(0).VertexName = TestVertexName;
+				OutputNodeInterface.GetOutputInterface().At(0).VertexName = TestVertexName;
+				Node = DataTypeRegistry.CreateOutputNode(DataTypeName, FNodeData{ TestNodeName, TestNodeID, OutputNodeInterface });
 			}
 			break;
 
diff --git a/Engine/Plugins/Runtime/Metasound/Source/MetasoundEngineTest/Private/MetasoundDynamicGeneratorTests.cpp b/Engine/Plugins/Runtime/Metasound/Source/MetasoundEngineTest/Private/MetasoundDynamicGeneratorTests.cpp
--- a/Engine/Plugins/Runtime/Metasound/Source/MetasoundEngineTest/Private/MetasoundDynamicGeneratorTests.cpp
+++ b/Engine/Plugins/Runtime/Metasound/Source/MetasoundEngineTest/Private/MetasoundDynamicGeneratorTests.cpp
@@ -63,22 +63,26 @@ namespace Metasound::Test::Generator::Dynamic
 		template<typename DataType>
 		bool AddInput(const FVertexName& Name, const FGuid& NodeGuid, const FLiteral& DefaultLiteral)
 		{
-			FInputNodeConstructorParams InputNodeParams
+			TSharedPtr<const FNodeClassMetadata> InputClassMetadata = DataRegistry.GetInputClassMetadata(GetMetasoundDataTypeName<DataType>());
+			if (!InputClassMetadata)
 			{
-				Name,
-				NodeGuid,
-				Name,
-				DefaultLiteral
-			};
-			
-			TUniquePtr<INode> Node = DataRegistry.CreateInputNode(GetMetasoundDataTypeName<DataType>(), MoveTemp(InputNodeParams));
+				return false;
+			}
+
+			FVertexInterface InputNodeInterface = InputClassMetadata->DefaultInterface;
+			// Vertex names must be set for input nodes
+			InputNodeInterface.GetInputInterface().At(0).VertexName = Name;
+			InputNodeInterface.GetOutputInterface().At(0).VertexName = Name;
+			InputNodeInterface.GetInputInterface().At(0).SetDefaultLiteral(DefaultLiteral);
+
+			TUniquePtr<INode> Node = DataRegistry.CreateInputNode(GetMetasoundDataTypeName<DataType>(), FNodeData{ Name, NodeGuid, InputNodeInterface });
 
 			if (!Node.IsValid())
 			{
 				return false;
 			}
 
-			Transactor.AddNode(InputNodeParams.InstanceID, MoveTemp(Node));
+			Transactor.AddNode(NodeGuid, MoveTemp(Node));
 
 			const auto CreateDataReference = [](
 				const FOperatorSettings& InSettings,
@@ -91,7 +95,7 @@ namespace Metasound::Test::Generator::Dynamic
 			};
 			
 			Transactor.AddInputDataDestination(
-				InputNodeParams.InstanceID,
+				NodeGuid,
 				Name,
 				DefaultLiteral,
 				CreateDataReference);
@@ -108,22 +112,25 @@ namespace Metasound::Test::Generator::Dynamic
 		template<typename DataType>
 		bool AddOutput(const FVertexName& Name, const FGuid& NodeGuid)
 		{
-			FOutputNodeConstructorParams OutputNodeParams
+			TSharedPtr<const FNodeClassMetadata> OutputClassMetadata = DataRegistry.GetOutputClassMetadata(GetMetasoundDataTypeName<DataType>());
+			if (!OutputClassMetadata)
 			{
-				Name,
-				NodeGuid,
-				Name
-			};
+				return false;
+			}
+			FVertexInterface OutputNodeInterface = OutputClassMetadata->DefaultInterface;
+			// Vertex names must be set for output nodes
+			OutputNodeInterface.GetInputInterface().At(0).VertexName = Name;
+			OutputNodeInterface.GetOutputInterface().At(0).VertexName = Name;
 
-			TUniquePtr<INode> Node = DataRegistry.CreateOutputNode(GetMetasoundDataTypeName<DataType>(), MoveTemp(OutputNodeParams));
+			TUniquePtr<INode> Node = DataRegistry.CreateOutputNode(GetMetasoundDataTypeName<DataType>(), FNodeData{ Name, NodeGuid, OutputNodeInterface });
 
 			if (!Node.IsValid())
 			{
 				return false;
 			}
 
-			Transactor.AddNode(OutputNodeParams.InstanceID, MoveTemp(Node));
-			Transactor.AddOutputDataSource(OutputNodeParams.InstanceID, Name);
+			Transactor.AddNode(NodeGuid, MoveTemp(Node));
+			Transactor.AddOutputDataSource(NodeGuid, Name);
 			return true;
 		}
 
diff --git a/Engine/Plugins/Runtime/Metasound/Source/MetasoundFrontend/Private/MetasoundFrontendGraph.cpp b/Engine/Plugins/Runtime/Metasound/Source/MetasoundFrontend/Private/MetasoundFrontendGraph.cpp
--- a/Engine/Plugins/Runtime/Metasound/Source/MetasoundFrontend/Private/MetasoundFrontendGraph.cpp
+++ b/Engine/Plugins/Runtime/Metasound/Source/MetasoundFrontend/Private/MetasoundFrontendGraph.cpp
@@ -23,7 +23,9 @@ namespace Metasound
 
 			// Input nodes need an extra Index value to keep track of their position in the graph's inputs.
 			InputNodes.Add(InIndex, InNode.Get());
+PRAGMA_DISABLE_DEPRECATION_WARNINGS
 			AddInputNode(InDependencyId, InVertexName, MoveTemp(InNode));
+PRAGMA_ENABLE_DEPRECATION_WARNINGS
 		}
 	}
 
@@ -36,7 +38,9 @@ namespace Metasound
 
 			// Output nodes need an extra Index value to keep track of their position in the graph's inputs.
 			OutputNodes.Add(InIndex, InNode.Get());
+PRAGMA_DISABLE_DEPRECATION_WARNINGS
 			AddOutputNode(InNodeID, InVertexName, MoveTemp(InNode));
+PRAGMA_ENABLE_DEPRECATION_WARNINGS
 		}
 	}
 
diff --git a/Engine/Plugins/Runtime/Metasound/Source/MetasoundFrontend/Private/MetasoundFrontendGraphBuilder.cpp b/Engine/Plugins/Runtime/Metasound/Source/MetasoundFrontend/Private/MetasoundFrontendGraphBuilder.cpp
--- a/Engine/Plugins/Runtime/Metasound/Source/MetasoundFrontend/Private/MetasoundFrontendGraphBuilder.cpp
+++ b/Engine/Plugins/Runtime/Metasound/Source/MetasoundFrontend/Private/MetasoundFrontendGraphBuilder.cpp
@@ -587,10 +587,10 @@ namespace Metasound::Frontend
 
 							if (nullptr != ClassInput)
 							{
-								TSharedPtr<const INode> InputNode(CreateInputNode(InGraphContext.BuildContext, Node, *NodeClass, *ClassInput).Release());
+								TUniquePtr<INode> InputNode = CreateInputNode(InGraphContext.BuildContext, Node, *NodeClass, *ClassInput);
 								if (InputNode.IsValid())
 								{
-									InGraphContext.Graph->AddInputNode(Node.GetID(), ClassInput->Name, InputNode);
+									InGraphContext.Graph->AddInputNode(Node.GetID(), ClassInput->Name, MoveTemp(InputNode));
 								}
 								else
 								{
@@ -611,10 +611,10 @@ namespace Metasound::Frontend
 							const FMetasoundFrontendClassOutput* ClassOutput = FindClassOutputForOutputNode(InGraphContext.GraphClass, Node);
 							if (nullptr != ClassOutput)
 							{
-								TSharedPtr<const INode> OutputNode(CreateOutputNode(InGraphContext.BuildContext, Node, *NodeClass, InGraphContext).Release());
+								TUniquePtr<INode> OutputNode = CreateOutputNode(InGraphContext.BuildContext, Node, *NodeClass, InGraphContext);
 								if (OutputNode)
 								{
-									InGraphContext.Graph->AddOutputNode(Node.GetID(), ClassOutput->Name, OutputNode);
+									InGraphContext.Graph->AddOutputNode(Node.GetID(), ClassOutput->Name, MoveTemp(OutputNode));
 								}
 								else
 								{
@@ -636,7 +636,9 @@ namespace Metasound::Frontend
 
 							if (SubgraphPtr.IsValid())
 							{
+PRAGMA_DISABLE_DEPRECATION_WARNINGS
 								InGraphContext.Graph->AddNode(Node.GetID(), SubgraphPtr);
+PRAGMA_ENABLE_DEPRECATION_WARNINGS
 							}
 							else
 							{
@@ -654,10 +656,10 @@ namespace Metasound::Frontend
 
 						case EMetasoundFrontendClassType::Variable:
 						{
-							TSharedPtr<const INode> VariableNode(CreateVariableNode(InGraphContext.BuildContext, Node, *NodeClass, InGraphContext.PagedGraph).Release());
+							TUniquePtr<INode> VariableNode = CreateVariableNode(InGraphContext.BuildContext, Node, *NodeClass, InGraphContext.PagedGraph);
 							if (VariableNode)
 							{
-								InGraphContext.Graph->AddNode(Node.GetID(), VariableNode);
+								InGraphContext.Graph->AddNode(Node.GetID(), MoveTemp(VariableNode));
 							}
 							else
 							{
@@ -675,10 +677,10 @@ namespace Metasound::Frontend
 						case EMetasoundFrontendClassType::External:
 						default:
 						{
-							TSharedPtr<const INode> ExternalNode(CreateExternalNode(InGraphContext.BuildContext, Node, *NodeClass).Release());
+							TUniquePtr<INode> ExternalNode = CreateExternalNode(InGraphContext.BuildContext, Node, *NodeClass);
 							if (ExternalNode)
 							{
-								InGraphContext.Graph->AddNode(Node.GetID(), ExternalNode);
+								InGraphContext.Graph->AddNode(Node.GetID(), MoveTemp(ExternalNode));
 							}
 							else
 							{
diff --git a/Engine/Plugins/Runtime/Metasound/Source/MetasoundFrontend/Private/MetasoundFrontendNodeClassRegistry.cpp b/Engine/Plugins/Runtime/Metasound/Source/MetasoundFrontend/Private/MetasoundFrontendNodeClassRegistry.cpp
--- a/Engine/Plugins/Runtime/Metasound/Source/MetasoundFrontend/Private/MetasoundFrontendNodeClassRegistry.cpp
+++ b/Engine/Plugins/Runtime/Metasound/Source/MetasoundFrontend/Private/MetasoundFrontendNodeClassRegistry.cpp
@@ -152,6 +152,7 @@ namespace Metasound::Frontend
 			, Factory(MakeShared<FGraphOperatorFactoryAdapter>(*InGraphToWrap))
 			, Graph(MoveTemp(InGraphToWrap))
 			{
+				Interface = Graph->GetVertexInterface();
 			}
 
 			virtual const FName& GetInstanceName() const override
@@ -173,7 +174,19 @@ namespace Metasound::Frontend
 
 			virtual const FVertexInterface& GetVertexInterface() const override
 			{
-				return Graph->GetVertexInterface();
+				return Interface;
+			}
+
+			virtual void SetDefaultInput(const FVertexName& InVertexName, const FLiteral& InLiteral) override
+			{
+				if (FInputDataVertex* Vertex = Interface.GetInputInterface().Find(InVertexName))
+				{
+					Vertex->SetDefaultLiteral(InLiteral);
+				}
+				else
+				{
+					UE_LOG(LogMetaSound, Error, TEXT("Could not set default input. Could not find input vertex %s on node %s of node class %s"), *InVertexName.ToString(), *GetInstanceName().ToString(), *GetMetadata().ClassName.ToString());
+				}
 			}
 
 			UE_DEPRECATED(5.6, "Mutable vertex interfaces are no longer part of the node interface.")
@@ -182,7 +195,7 @@ namespace Metasound::Frontend
 				PRAGMA_DISABLE_DEPRECATION_WARNINGS
 				// Cannot set vertex interface because of const reference to IGraph
 				// Return true if the supplied interface is the same as the existing interface.
-				return Graph->GetVertexInterface() == InInterface;
+				return Interface == InInterface;
 				PRAGMA_ENABLE_DEPRECATION_WARNINGS
 			}
 
@@ -204,6 +217,7 @@ namespace Metasound::Frontend
 			FGuid InstanceID;
 			TSharedRef<FGraphOperatorFactoryAdapter> Factory;
 			TSharedRef<const IGraph> Graph;
+			FVertexInterface Interface;
 		};
 
 
diff --git a/Engine/Plugins/Runtime/Metasound/Source/MetasoundGraphCore/Private/MetasoundBasicNode.cpp b/Engine/Plugins/Runtime/Metasound/Source/MetasoundGraphCore/Private/MetasoundBasicNode.cpp
--- a/Engine/Plugins/Runtime/Metasound/Source/MetasoundGraphCore/Private/MetasoundBasicNode.cpp
+++ b/Engine/Plugins/Runtime/Metasound/Source/MetasoundGraphCore/Private/MetasoundBasicNode.cpp
@@ -35,4 +35,16 @@ namespace Metasound
 	{
 		return NodeData.Interface;
 	}
+
+	void FBasicNode::SetDefaultInput(const FVertexName& InVertexName, const FLiteral& InLiteral)
+	{
+		if (FInputDataVertex* Vertex = NodeData.Interface.GetInputInterface().Find(InVertexName))
+		{
+			Vertex->SetDefaultLiteral(InLiteral);
+		}
+		else
+		{
+			UE_LOG(LogMetaSound, Error, TEXT("Could not set default input. Could not find input vertex %s on node %s of node class %s"), *InVertexName.ToString(), *NodeData.Name.ToString(), *ClassMetadata->ClassName.ToString());
+		}
+	}
 }
diff --git a/Engine/Plugins/Runtime/Metasound/Source/MetasoundGraphCore/Private/MetasoundDynamicOperator.cpp b/Engine/Plugins/Runtime/Metasound/Source/MetasoundGraphCore/Private/MetasoundDynamicOperator.cpp
--- a/Engine/Plugins/Runtime/Metasound/Source/MetasoundGraphCore/Private/MetasoundDynamicOperator.cpp
+++ b/Engine/Plugins/Runtime/Metasound/Source/MetasoundGraphCore/Private/MetasoundDynamicOperator.cpp
@@ -564,12 +564,50 @@ namespace Metasound
 			return EDynamicOperatorTransformQueueAction::Continue;
 		}
 
+		// Disconnect a specific edge
+		FRemoveOperatorConnection::FRemoveOperatorConnection(FOperatorID InFromOpID, const FName& InFromVertName, FOperatorID InToOpID, const FName& InToVertName, FAnyDataReference InDataReference)
+			: SetOperatorInputTransform(InToOpID, InToVertName, MoveTemp(InDataReference))
+			, FromOpID(InFromOpID)
+			, ToOpID(InToOpID)
+			, FromVertName(InFromVertName)
+			, ToVertName(InToVertName)
+		{
+		}
+
+		EDynamicOperatorTransformQueueAction FRemoveOperatorConnection::Transform(FDynamicGraphOperatorData& InGraphOperatorData)
+		{
+			using namespace DirectedGraphAlgo;
+
+			METASOUND_TRACE_CPUPROFILER_EVENT_SCOPE(Metasound::DynamicOperator::DisconnectOperators)
+
+			// Clean up the old connection.
+			FOperatorInfo* FromOpInfo = InGraphOperatorData.OperatorMap.Find(FromOpID);
+
+			if (nullptr == FromOpInfo)
+			{
+				UE_LOG(LogMetaSound, Error, TEXT("Could not find operator with ID %d when disconnecting from %d:%s to %d:%s"), FromOpID, FromOpID, *FromVertName.ToString(), ToOpID, *ToVertName.ToString());
+				return EDynamicOperatorTransformQueueAction::Continue;
+			}
+
+			// Remove destinations from operator map.
+			auto IsDestinationToRemove = [&](const FGraphOperatorData::FVertexDestination& Dst)
+			{
+				return (Dst.OperatorID == ToOpID) && (Dst.VertexName == ToVertName);
+			};
+			FromOpInfo->OutputConnections.FindOrAdd(FromVertName).RemoveAll(IsDestinationToRemove); 
+
+			// Set data reference on newly unconnected input 
+			EDynamicOperatorTransformQueueAction NextAction = SetOperatorInputTransform.Transform(InGraphOperatorData);
+			check(NextAction == EDynamicOperatorTransformQueueAction::Continue);
+
+			return EDynamicOperatorTransformQueueAction::Continue;
+		}
+
 		// Set the value on an unconnected operator input
-		FSetOperatorInput::FSetOperatorInput(FOperatorID InToOpID, const FName& InToVert, const FLiteral& InLiteral, FLiteralAssignmentFunction InAssignFunc)
+		FSetOperatorInput::FSetOperatorInput(FOperatorID InToOpID, const FName& InToVert, FAnyDataReference InDataRef)
 		: ToOpID(InToOpID)
 		, ToVert(InToVert)
-		, Literal(InLiteral)
-		, AssignFunc(InAssignFunc)
+		, DataRef(InDataRef)
 		{
 		}
 
@@ -593,8 +631,13 @@ namespace Metasound
 				UE_LOG(LogMetaSound, Error, TEXT("Could not find input data reference with vertex name %s when setting %d:%s"), *ToVert.ToString(), ToOpID, *ToVert.ToString());
 				return EDynamicOperatorTransformQueueAction::Continue;
 			}
+						
+			// Propagate the data reference update through the graph
+			PropagateBindUpdate(ToOpID, ToVert, MoveTemp(DataRef), InGraphOperatorData);
 
-			AssignFunc(InGraphOperatorData.OperatorSettings, Literal, *Ref);
+			// Refresh output vertex interface data in case any output nodes were updated
+			// when bind updates were propagated through the graph.
+			UpdateOutputVertexData(InGraphOperatorData);
 
 			return EDynamicOperatorTransformQueueAction::Continue;
 		}
diff --git a/Engine/Plugins/Runtime/Metasound/Source/MetasoundGraphCore/Private/MetasoundDynamicOperator.h b/Engine/Plugins/Runtime/Metasound/Source/MetasoundGraphCore/Private/MetasoundDynamicOperator.h
--- a/Engine/Plugins/Runtime/Metasound/Source/MetasoundGraphCore/Private/MetasoundDynamicOperator.h
+++ b/Engine/Plugins/Runtime/Metasound/Source/MetasoundGraphCore/Private/MetasoundDynamicOperator.h
@@ -211,7 +211,7 @@ namespace Metasound
 		class FSwapOperatorConnection : public IDynamicOperatorTransform
 		{
 		public:
-			FSwapOperatorConnection(FOperatorID InOriginalFromOpID, const FName& InOrignialFromVert, FOperatorID InNewFromOpID, const FName& InNewFromVert, FOperatorID InToOpID, const FName& InToVert);
+			FSwapOperatorConnection(FOperatorID InOriginalFromOpID, const FName& InOriginalFromVert, FOperatorID InNewFromOpID, const FName& InNewFromVert, FOperatorID InToOpID, const FName& InToVert);
 			virtual EDynamicOperatorTransformQueueAction Transform(FDynamicGraphOperatorData& InGraphOperatorData) override;
 		private:
 			FConnectOperators ConnectTransform;
@@ -226,14 +226,30 @@ namespace Metasound
 		class FSetOperatorInput : public IDynamicOperatorTransform
 		{
 		public:
-			FSetOperatorInput(FOperatorID InToOpID, const FName& InToVert, const FLiteral& InLiteral, FLiteralAssignmentFunction InAssignFunc);
+			FSetOperatorInput(FOperatorID InToOpID, const FName& InToVert, FAnyDataReference InDataRef);
+
 			virtual EDynamicOperatorTransformQueueAction Transform(FDynamicGraphOperatorData& InGraphOperatorData) override;
 
 		private:
 			FOperatorID ToOpID;
 			FName ToVert;
-			FLiteral Literal;
-			FLiteralAssignmentFunction AssignFunc;
+			FAnyDataReference DataRef;
+		};
+
+		/** A transform which disconnects two vertices in the graph, replacing the input with a given data reference. */
+		class FRemoveOperatorConnection : public IDynamicOperatorTransform
+		{
+		public:
+			FRemoveOperatorConnection(FOperatorID InFromOpID, const FName& InFromVertName,  FOperatorID InToOpID, const FName& InToVertName, FAnyDataReference InReplacementDataRef);
+			virtual EDynamicOperatorTransformQueueAction Transform(FDynamicGraphOperatorData& InGraphOperatorData) override;
+
+		private:
+			FSetOperatorInput SetOperatorInputTransform;
+
+			FOperatorID FromOpID;
+			FOperatorID ToOpID;
+			FName FromVertName;
+			FName ToVertName;
 		};
 
 		/** A transform which groups multiple transformations together and forces them to all complete
diff --git a/Engine/Plugins/Runtime/Metasound/Source/MetasoundGraphCore/Private/MetasoundDynamicOperatorTransactor.cpp b/Engine/Plugins/Runtime/Metasound/Source/MetasoundGraphCore/Private/MetasoundDynamicOperatorTransactor.cpp
--- a/Engine/Plugins/Runtime/Metasound/Source/MetasoundGraphCore/Private/MetasoundDynamicOperatorTransactor.cpp
+++ b/Engine/Plugins/Runtime/Metasound/Source/MetasoundGraphCore/Private/MetasoundDynamicOperatorTransactor.cpp
@@ -652,10 +652,11 @@ namespace Metasound
 				return;
 			}
 
-			TSharedRef<const INode> NodePtr(InNode.Release());
+			// Cache pointer because InNode gets moved. 
+			const INode* NodePtr = InNode.Get();
 
-			Graph.AddNode(InNodeID, NodePtr);
-			int32 OperatorOrdinal = GraphSorter.InsertOperator(GetOperatorID(*NodePtr), FDynamicGraphIncrementalSorter::EInsertLocation::Last);
+			Graph.AddNode(InNodeID, MoveTemp(InNode));
+			int32 OperatorOrdinal = GraphSorter.InsertOperator(GetOperatorID(NodePtr), FDynamicGraphIncrementalSorter::EInsertLocation::Last);
 
 			EnqueueInsertOperatorTransform(*NodePtr, OperatorOrdinal);
 		}
@@ -779,11 +780,57 @@ namespace Metasound
 			if (ToNode->GetVertexInterface().GetInputVertex(InToVertex).DataTypeName == GetMetasoundDataTypeName<FAudioBuffer>())
 			{
 				// Handle audio edge removal with a fade out.
-				EnqueueFadeAndRemoveEdgeOperatorTransform(*FromNode, InFromVertex, *ToNode, InToVertex, LiteralNodeRef, LiteralOrdinal);
+				EnqueueFadeAndRemoveEdgeOperatorTransform_Deprecated(*FromNode, InFromVertex, *ToNode, InToVertex, LiteralNodeRef, LiteralOrdinal);
 			}
 			else
 			{
-				EnqueueRemoveEdgeOperatorTransform(*FromNode, InFromVertex, *ToNode, InToVertex, LiteralNodeRef, LiteralOrdinal);
+				EnqueueRemoveEdgeOperatorTransform_Deprecated(*FromNode, InFromVertex, *ToNode, InToVertex, LiteralNodeRef, LiteralOrdinal);
+			}
+		}
+
+		void FDynamicOperatorTransactor::RemoveDataEdge(const FGuid& InFromNodeID, const FVertexName& InFromVertexName, const FGuid& InToNodeID, const FVertexName& InToVertexName, FLiteral InReplacementLiteral, const FReferenceCreationFunction InReferenceCreateFunc)
+		{
+			using namespace DirectedGraphAlgo;
+
+			METASOUND_TRACE_CPUPROFILER_EVENT_SCOPE(Metasound::FDynamicOperatorTransactor::RemoveDataEdge);
+
+			const INode* FromNode = Graph.FindNode(InFromNodeID);
+			const INode* ToNode = Graph.FindNode(InToNodeID);
+			const FOperatorID FromOperatorID = DirectedGraphAlgo::GetOperatorID(FromNode);
+			const FOperatorID ToOperatorID = DirectedGraphAlgo::GetOperatorID(ToNode);
+
+			if ((nullptr == ToNode) || (nullptr == FromNode))
+			{
+				UE_LOG(LogMetaSound, Error, TEXT("Cannot remove edge from %s:%s to %s:%s because of missing node"), *InFromNodeID.ToString(), *InFromVertexName.ToString(), *InToNodeID.ToString(), *InToVertexName.ToString());
+				return;
+			}
+
+			if (!ToNode->GetVertexInterface().ContainsInputVertex(InToVertexName))
+			{
+				UE_LOG(LogMetaSound, Error, TEXT("Cannot remove edge from %s:%s to %s:%s because of destination node does not contain vertex %s."), *InFromNodeID.ToString(), *InFromVertexName.ToString(), *InToNodeID.ToString(), *InToVertexName.ToString(), *InToVertexName.ToString());
+				return;
+			}
+
+			bool bSuccess = Graph.RemoveDataEdge(*FromNode, InFromVertexName, *ToNode, InToVertexName);
+			if (!bSuccess)
+			{
+				UE_LOG(LogMetaSound, Error, TEXT("Failed to remove edge from %s:%s to %s:%s on internal graph."), *InFromNodeID.ToString(), *InFromVertexName.ToString(), *InToNodeID.ToString(), *InToVertexName.ToString());
+				return;
+			}
+			GraphSorter.RemoveDataEdge(FromOperatorID, ToOperatorID);
+
+			// Update default stored on the node
+			Graph.SetNodeDefaultInput(InToNodeID, InToVertexName, InReplacementLiteral);
+
+			// Immediately disconnect non-audio edges. 
+			if (ToNode->GetVertexInterface().GetInputVertex(InToVertexName).DataTypeName == GetMetasoundDataTypeName<FAudioBuffer>())
+			{
+				// Handle audio edge removal with a fade out.
+				EnqueueFadeAndRemoveEdgeOperatorTransform(*FromNode, InFromVertexName, *ToNode, InToVertexName, InReplacementLiteral, InReferenceCreateFunc);
+			}
+			else
+			{
+				EnqueueRemoveEdgeOperatorTransform(*FromNode, InFromVertexName, *ToNode, InToVertexName, InReplacementLiteral, InReferenceCreateFunc);
 			}
 		}
 
@@ -819,7 +866,44 @@ namespace Metasound
 			AddDataEdgeInternal(*LiteralNode, DynamicOperatorTransactorPrivate::LiteralNodeOutputVertexName, InNodeID, *Node, InVertex);
 
 			// Add literal node after calling "AddDataEdgeInternal" so that AddDataEdgeInternal can check if there is a prior existing literal node.
-			LiteralNodeMap.Add(FLiteralNodeID{InNodeID, InVertex}, MoveTemp(InLiteralNode));
+			LiteralNodeMap.Add(FLiteralNodeID{InNodeID, InVertex }, MoveTemp(InLiteralNode));
+		}
+
+		void FDynamicOperatorTransactor::SetValue(const FGuid& InNodeID, const FVertexName& InVertexName, const FLiteral& InLiteral, const FReferenceCreationFunction InReferenceCreateFunc)
+		{
+			METASOUND_TRACE_CPUPROFILER_EVENT_SCOPE(Metasound::FDynamicOperatorTransactor::SetValue);
+
+			const INode* Node = Graph.FindNode(InNodeID);
+			if (nullptr == Node)
+			{
+				UE_LOG(LogMetaSound, Error, TEXT("Cannot set node value of %s:%s because of missing node"), *InNodeID.ToString(), *InVertexName.ToString());
+				return;
+			}
+
+			Graph.SetNodeDefaultInput(InNodeID, InVertexName, InLiteral);
+
+			const FOperatorID OperatorID = DirectedGraphAlgo::GetOperatorID(Node);
+			const FInputDataVertex& InputVertex = Node->GetVertexInterface().GetInputVertex(InVertexName);
+			EDataReferenceAccessType ReferenceAccessType = InputVertex.AccessType == EVertexAccessType::Value ? EDataReferenceAccessType::Value : EDataReferenceAccessType::Write;
+
+			auto CreateSetValueTransform = [&](const FDynamicOperatorInfo& InInfo) -> TUniquePtr<IDynamicOperatorTransform>
+			{
+				TOptional<FAnyDataReference> NewDataReference = InReferenceCreateFunc(InInfo.OperatorSettings, InputVertex.DataTypeName, InLiteral, ReferenceAccessType);
+
+				if (NewDataReference.IsSet())
+				{
+					return MakeUnique<FSetOperatorInput>(FSetOperatorInput(
+						OperatorID, InVertexName, MoveTemp(*NewDataReference)
+					));
+				}
+				else
+				{
+					UE_LOG(LogMetaSound, Error, TEXT("Cannot Set Value %s:%s because of failure to create data reference for transform."), *InNodeID.ToString(), *InVertexName.ToString());
+					return TUniquePtr<IDynamicOperatorTransform>(nullptr);
+				}
+			};
+
+			EnqueueTransformOnOperatorQueues(CreateSetValueTransform);
 		}
 
 		/** Add an input data destination to describe how data provided 
@@ -857,7 +941,7 @@ namespace Metasound
 				TOptional<FAnyDataReference> NewDataReference = InFunc(InInfo.OperatorSettings, InputVertex.DataTypeName, InDefaultLiteral, ReferenceAccessType);
 				if (NewDataReference.IsSet())
 				{
-					return MakeUnique<FAddInput>(OperatorID, InVertexName, *NewDataReference);
+					return MakeUnique<FAddInput>(OperatorID, InVertexName, MoveTemp(*NewDataReference));
 				}
 				else
 				{
@@ -960,7 +1044,6 @@ namespace Metasound
 				Graph.RemoveDataEdge(*PriorLiteralNode, DynamicOperatorTransactorPrivate::LiteralNodeOutputVertexName, InToNode, InToVertex);
 				GraphSorter.RemoveDataEdge(PriorLiteralOperatorID, ToOperatorID);
 			}
-
 			Graph.AddDataEdge(InFromNode, InFromVertex, InToNode, InToVertex);
 			TArray<FOrdinalSwap> OrdinalUpdates;
 			GraphSorter.AddDataEdge(FromOperatorID, ToOperatorID, OrdinalUpdates);
@@ -976,12 +1059,12 @@ namespace Metasound
 			if (InputVertex->DataTypeName == GetMetasoundDataTypeName<FAudioBuffer>())
 			{
 				// If edge is audio, then the connection needs to be faded
-				EnqueueFadeAndAddEdgeOperatorTransform(InFromNode, InFromVertex, InToNode, InToVertex, PriorLiteralNode.Get(), OrdinalUpdates);
+				EnqueueFadeAndAddEdgeOperatorTransform(InFromNode, InFromVertex, InToNode, InToVertex, OrdinalUpdates);
 			}
 			else
 			{
 				// If the edge is not audio, then no fading is performed. 
-				EnqueueAddEdgeOperatorTransform(InFromNode, InFromVertex, InToNode, InToVertex, PriorLiteralNode.Get(), OrdinalUpdates);
+				EnqueueAddEdgeOperatorTransform(InFromNode, InFromVertex, InToNode, InToVertex, OrdinalUpdates);
 			}
 		}
 		
@@ -1042,12 +1125,11 @@ namespace Metasound
 			RemoveNodeInternal(InNode, bInRemoveDataEdgesWithNode);
 		}
 
-		void FDynamicOperatorTransactor::EnqueueAddEdgeOperatorTransform(const INode& InFromNode, const FVertexName& InFromVertex, const INode& InToNode, const FVertexName& InToVertex, const INode* InPriorLiteralNode, const TArray<FOrdinalSwap>& InOrdinalUpdates)
+		void FDynamicOperatorTransactor::EnqueueAddEdgeOperatorTransform(const INode& InFromNode, const FVertexName& InFromVertex, const INode& InToNode, const FVertexName& InToVertex, const TArray<FOrdinalSwap>& InOrdinalUpdates)
 		{
 			/* enqueue an update. */
 			const FOperatorID FromOperatorID = DirectedGraphAlgo::GetOperatorID(InFromNode);
 			const FOperatorID ToOperatorID = DirectedGraphAlgo::GetOperatorID(InToNode);
-			const FOperatorID PriorLiteralOperatorID = DirectedGraphAlgo::GetOperatorID(InPriorLiteralNode);
 
 			// Create transforms for runtime
 			auto CreateAddEdgeTransform = [&](const FDynamicOperatorInfo& InInfo) -> TUniquePtr<IDynamicOperatorTransform>
@@ -1059,15 +1141,7 @@ namespace Metasound
 					AtomicTransforms.Add(MakeUnique<FSwapOperatorOrdinalsAndSort>(InOrdinalUpdates));
 				}
 
-				if (InPriorLiteralNode)
-				{
-					AtomicTransforms.Add(MakeUnique<FSwapOperatorConnection>(PriorLiteralOperatorID, DynamicOperatorTransactorPrivate::LiteralNodeOutputVertexName, FromOperatorID, InFromVertex, ToOperatorID, InToVertex));
-					AtomicTransforms.Add(MakeUnique<FRemoveOperator>(PriorLiteralOperatorID, TArray<FOperatorID>() /* OperatorsConnectedToInput */));
-				}
-				else
-				{
-					AtomicTransforms.Add(MakeUnique<FConnectOperators>(FromOperatorID, InFromVertex, ToOperatorID, InToVertex));
-				}
+				AtomicTransforms.Add(MakeUnique<FConnectOperators>(FromOperatorID, InFromVertex, ToOperatorID, InToVertex));
 
 				return MakeUnique<FAtomicTransform>(MoveTemp(AtomicTransforms));
 			};
@@ -1075,11 +1149,10 @@ namespace Metasound
 			EnqueueTransformOnOperatorQueues(CreateAddEdgeTransform);
 		}
 
-		void FDynamicOperatorTransactor::EnqueueFadeAndAddEdgeOperatorTransform(const INode& InFromNode, const FVertexName& InFromVertex, const INode& InToNode, const FVertexName& InToVertex, const INode* InPriorLiteralNode, const TArray<FOrdinalSwap>& InOrdinalUpdates)
+		void FDynamicOperatorTransactor::EnqueueFadeAndAddEdgeOperatorTransform(const INode& InFromNode, const FVertexName& InFromVertex, const INode& InToNode, const FVertexName& InToVertex, const TArray<FOrdinalSwap>& InOrdinalUpdates)
 		{
 			const FOperatorID FromOperatorID = DirectedGraphAlgo::GetOperatorID(InFromNode);
 			const FOperatorID ToOperatorID = DirectedGraphAlgo::GetOperatorID(InToNode);
-			const FOperatorID PriorLiteralOperatorID = DirectedGraphAlgo::GetOperatorID(InPriorLiteralNode);
 
 			// Fade inputs on the receiving node when adding an edge. We don't fade the source node's outputs
 			// because those outputs could also be connected to other nodes which we do not want to fade. 
@@ -1095,17 +1168,8 @@ namespace Metasound
 				{
 					AtomicTransforms.Add(MakeUnique<FSwapOperatorOrdinalsAndSort>(InOrdinalUpdates));
 				}
-				if (InPriorLiteralNode)
-				{
-					AtomicTransforms.Add(MakeUnique<FSwapOperatorConnection>(PriorLiteralOperatorID, DynamicOperatorTransactorPrivate::LiteralNodeOutputVertexName, FromOperatorID, InFromVertex, ToOperatorID, InToVertex));
-					// We assume that any FAudioBuffer created from a literal node is just silent audio. That means we can
-					// Remove that literal node without fading out the audio from the literal node connected to the destination node. 
-					AtomicTransforms.Add(MakeUnique<FRemoveOperator>(DirectedGraphAlgo::GetOperatorID(InPriorLiteralNode), TArray<FOperatorID>() /* OperatorsConnectedToInput */));
-				}
-				else
-				{
-					AtomicTransforms.Add(MakeUnique<FConnectOperators>(FromOperatorID, InFromVertex, ToOperatorID, InToVertex));
-				}
+				
+				AtomicTransforms.Add(MakeUnique<FConnectOperators>(FromOperatorID, InFromVertex, ToOperatorID, InToVertex));
 
 				AtomicTransforms.Add(MakeUnique<FBeginAudioFadeTransform>(ToOperatorID, EAudioFadeType::FadeIn, InputsToFade, OutputsToFade));
 				
@@ -1167,7 +1231,7 @@ namespace Metasound
 			EnqueueTransformOnOperatorQueues(CreateRemoveNodeTransform);
 		}
 
-		void FDynamicOperatorTransactor::EnqueueRemoveEdgeOperatorTransform(const INode& InFromNode, const FVertexName& InFromVertex, const INode& InToNode, const FVertexName& InToVertex, const INode& InReplacementLiteralNode, int32 InLiteralOrdinal)
+		void FDynamicOperatorTransactor::EnqueueRemoveEdgeOperatorTransform_Deprecated(const INode& InFromNode, const FVertexName& InFromVertexName, const INode& InToNode, const FVertexName& InToVertexName, const INode& InReplacementLiteralNode, int32 InLiteralOrdinal)
 		{
 			using namespace DynamicOperatorTransactorPrivate;
 
@@ -1181,7 +1245,7 @@ namespace Metasound
 				TUniquePtr<IDynamicOperatorTransform> AddNodeTransform = CreateInsertOperatorTransform(InReplacementLiteralNode, InLiteralOrdinal, InInfo.OperatorSettings, InInfo.Environment, InInfo.GraphRenderCost.Get());
 
 				// Swap prior connection with new connections.
-				TUniquePtr<IDynamicOperatorTransform> ConnectOperatorsTransform = MakeUnique<FSwapOperatorConnection>(FromOperatorID, InFromVertex, LiteralOperatorID, DynamicOperatorTransactorPrivate::LiteralNodeOutputVertexName, ToOperatorID, InToVertex);
+				TUniquePtr<IDynamicOperatorTransform> ConnectOperatorsTransform = MakeUnique<FSwapOperatorConnection>(FromOperatorID, InFromVertexName, LiteralOperatorID, DynamicOperatorTransactorPrivate::LiteralNodeOutputVertexName, ToOperatorID, InToVertexName);
 
 				if (AddNodeTransform.IsValid() && ConnectOperatorsTransform.IsValid())
 				{
@@ -1196,19 +1260,68 @@ namespace Metasound
 				else
 				{
 
-					UE_LOG(LogMetaSound, Error, TEXT("Cannot remove edge from %s:%s to %s:%s because of failure to create all transforms needed to perform operatorn."), *GetDebugNodeNameString(InFromNode), *InFromVertex.ToString(), *GetDebugNodeNameString(InToNode), *InToVertex.ToString());
+					UE_LOG(LogMetaSound, Error, TEXT("Cannot remove edge from %s:%s to %s:%s because of failure to create all transforms needed to perform operatorn."), *GetDebugNodeNameString(InFromNode), *InFromVertexName.ToString(), *GetDebugNodeNameString(InToNode), *InToVertexName.ToString());
 					return TUniquePtr<IDynamicOperatorTransform>(nullptr);
 				}
 			};
 
 			EnqueueTransformOnOperatorQueues(CreateRemoveEdgeTransform);
 		}
+		
+		void FDynamicOperatorTransactor::EnqueueRemoveEdgeOperatorTransform(const INode& InFromNode, const FVertexName& InFromVertexName, const INode& InToNode, const FVertexName& InToVertexName, const FLiteral& InReplacementLiteral, const FReferenceCreationFunction InReferenceCreateFunc)
+		{
+			using namespace DynamicOperatorTransactorPrivate;
+
+			const FOperatorID FromOperatorID = DirectedGraphAlgo::GetOperatorID(InFromNode);
+			const FOperatorID ToOperatorID = DirectedGraphAlgo::GetOperatorID(InToNode);
+			const FInputDataVertex& InputVertex = InToNode.GetVertexInterface().GetInputVertex(InToVertexName);
+			EDataReferenceAccessType ReferenceAccessType = InputVertex.AccessType == EVertexAccessType::Value ? EDataReferenceAccessType::Value : EDataReferenceAccessType::Write;
+
+			auto CreateRemoveEdgeTransform = [&](const FDynamicOperatorInfo& InInfo) -> TUniquePtr<IDynamicOperatorTransform>
+			{
+				TOptional<FAnyDataReference> NewDataReference = InReferenceCreateFunc(InInfo.OperatorSettings, InputVertex.DataTypeName, InReplacementLiteral, ReferenceAccessType);
+				if (NewDataReference.IsSet())
+				{
+					return MakeUnique<FRemoveOperatorConnection>(FromOperatorID, InFromVertexName, ToOperatorID, InToVertexName, MoveTemp(*NewDataReference));
+				}
+				else
+				{
+					UE_LOG(LogMetaSound, Error, TEXT("Cannot add Remove Edge %s:%s because of failure to create data reference for replacement literal."), *GetDebugNodeNameString(InToNode), *InToVertexName.ToString());
+					return TUniquePtr<IDynamicOperatorTransform>(nullptr);
+				}
+			};
+
+			EnqueueTransformOnOperatorQueues(CreateRemoveEdgeTransform);
+		}
+		
+		void FDynamicOperatorTransactor::EnqueueFadeAndRemoveEdgeOperatorTransform_Deprecated(const INode& InFromNode, const FVertexName& InFromVertexName, const INode& InToNode, const FVertexName& InToVertexName, const INode& InReplacementLiteralNode, int32 InLiteralOrdinal)
+		{
+			// Fade the input to the node getting disconnected rather than the output of the source node. The source node
+			// may be connected to other nodes and fading it's output would fade all the other connected nodes' inputs. 
+			TArrayView<const FVertexName> InputVerticesToFade(&InToVertexName, 1);
+			TArrayView<const FVertexName> OutputVerticesToFade;
+
+			EnqueueBeginFadeOperatorTransform(InToNode, EAudioFadeType::FadeOut, InputVerticesToFade, OutputVerticesToFade);
+
+			// Replace input with literal. This assumes that the replacement audio
+			// buffer contains silent audio. The fade transform will get the input
+			// audio to silent which will then seamlessly be swapped with a silent
+			// audio buffer as a permanent connection. 
+			//
+			// If we ever find ourselves creating audio buffers with literals which 
+			// are anything other than silent buffers, we should rework this operation
+			// to do either a cross-fade, or an additional "fade in" to the new value. 
+			EnqueueRemoveEdgeOperatorTransform_Deprecated(InFromNode, InFromVertexName, InToNode, InToVertexName, InReplacementLiteralNode, InLiteralOrdinal);
+
+			// Remove fade operation. 
+			EnqueueEndFadeOperatorTransform(InToNode);
+		}
 
-		void FDynamicOperatorTransactor::EnqueueFadeAndRemoveEdgeOperatorTransform(const INode& InFromNode, const FVertexName& InFromVertex, const INode& InToNode, const FVertexName& InToVertex, const INode& InReplacementLiteralNode, int32 InLiteralOrdinal)
+		void FDynamicOperatorTransactor::EnqueueFadeAndRemoveEdgeOperatorTransform(const INode& InFromNode, const FVertexName& InFromVertexName, const INode& InToNode, const FVertexName& InToVertexName, const FLiteral& InReplacementLiteral, const FReferenceCreationFunction InReferenceCreateFunc)
 		{
 			// Fade the input to the node getting disconnected rather than the output of the source node. The source node
 			// may be connected to other nodes and fading it's output would fade all the other connected nodes' inputs. 
-			TArrayView<const FVertexName> InputVerticesToFade(&InToVertex, 1);
+			TArrayView<const FVertexName> InputVerticesToFade(&InToVertexName, 1);
 			TArrayView<const FVertexName> OutputVerticesToFade;
 
 			EnqueueBeginFadeOperatorTransform(InToNode, EAudioFadeType::FadeOut, InputVerticesToFade, OutputVerticesToFade);
@@ -1221,7 +1334,7 @@ namespace Metasound
 			// If we ever find ourselves creating audio buffers with literals which 
 			// are anything other than silent buffers, we should rework this operation
 			// to do either a cross-fade, or an additional "fade in" to the new value. 
-			EnqueueRemoveEdgeOperatorTransform(InFromNode, InFromVertex, InToNode, InToVertex, InReplacementLiteralNode, InLiteralOrdinal);
+			EnqueueRemoveEdgeOperatorTransform(InFromNode, InFromVertexName, InToNode, InToVertexName, InReplacementLiteral, InReferenceCreateFunc);
 
 			// Remove fade operation. 
 			EnqueueEndFadeOperatorTransform(InToNode);
diff --git a/Engine/Plugins/Runtime/Metasound/Source/MetasoundGraphCore/Private/MetasoundGraph.cpp b/Engine/Plugins/Runtime/Metasound/Source/MetasoundGraphCore/Private/MetasoundGraph.cpp
--- a/Engine/Plugins/Runtime/Metasound/Source/MetasoundGraphCore/Private/MetasoundGraph.cpp
+++ b/Engine/Plugins/Runtime/Metasound/Source/MetasoundGraphCore/Private/MetasoundGraph.cpp
@@ -10,6 +10,16 @@
 
 namespace Metasound
 {
+	const INode* FGraph::FDeprecationNodeStorageAdapter::GetConstNode() const
+	{
+		return Storage.Get();
+	}
+
+	INode* FGraph::FDeprecationNodeStorageAdapter::GetMutableNode() const
+	{
+		return MutableNode;
+	}
+
 	FGraph::FGraph(const FString& InInstanceName, const FGuid& InInstanceID)
 	: InstanceName(InInstanceName)
 	, InstanceID(InInstanceID)
@@ -179,12 +189,41 @@ namespace Metasound
 		Nodes.Add(InNodeID, InNode);
 	}
 
+	void FGraph::AddNode(const FGuid& InNodeID, TUniquePtr<INode> InNode)
+	{
+		Nodes.Add(InNodeID, MoveTemp(InNode));
+	}
+
+	void FGraph::SetNodeDefaultInput(const FGuid& InNodeID, const FVertexName& InVertexName, const FLiteral& InLiteral)
+	{
+		if (INode* Node = FindMutableNode(InNodeID))
+		{
+			Node->SetDefaultInput(InVertexName, InLiteral);
+		}
+		else
+		{
+			UE_LOG(LogMetaSound, Error, TEXT("Could not set input vertex literal on vertex %s. No node with ID %s found in graph %s"), *InVertexName.ToString(), *InNodeID.ToString(), *GetInstanceName().ToString())
+		}
+	}
+
 	void FGraph::AddInputNode(FGuid InNodeId, const FVertexName& InVertexName, TSharedPtr<const INode> InNode)
 	{
 		if (InNode.IsValid())
 		{
 			AddInputDataDestination(*InNode, InVertexName);
-			AddNode(InNodeId, InNode);
+PRAGMA_DISABLE_DEPRECATION_WARNINGS
+			AddNode(InNodeId, MoveTemp(InNode));
+PRAGMA_ENABLE_DEPRECATION_WARNINGS
+
+		}
+	}
+
+	void FGraph::AddInputNode(FGuid InNodeId, const FVertexName& InVertexName, TUniquePtr<INode> InNode)
+	{
+		if (InNode.IsValid())
+		{
+			AddInputDataDestination(*InNode, InVertexName);
+			AddNode(InNodeId, MoveTemp(InNode));
 		}
 	}
 
@@ -193,15 +232,35 @@ namespace Metasound
 		if (InNode.IsValid())
 		{
 			AddOutputDataSource(*InNode, InVertexName);
+PRAGMA_DISABLE_DEPRECATION_WARNINGS
 			AddNode(InNodeID, InNode);
+PRAGMA_ENABLE_DEPRECATION_WARNINGS
+		}
+	}
+
+	void FGraph::AddOutputNode(FGuid InNodeID, const FVertexName& InVertexName, TUniquePtr<INode> InNode)
+	{
+		if (InNode.IsValid())
+		{
+			AddOutputDataSource(*InNode, InVertexName);
+			AddNode(InNodeID, MoveTemp(InNode));
 		}
 	}
 
 	const INode* FGraph::FindNode(const FGuid& InNodeID) const
 	{
-		if (const TSharedPtr<const INode>* NodePtr = Nodes.Find(InNodeID))
+		if (const FDeprecationNodeStorageAdapter* NodeStorage = Nodes.Find(InNodeID))
 		{
-			return NodePtr->Get();
+			return NodeStorage->GetConstNode();
+		}
+		return nullptr;
+	}
+
+	INode* FGraph::FindMutableNode(const FGuid& InNodeID)
+	{
+		if (const FDeprecationNodeStorageAdapter* NodeStorage = Nodes.Find(InNodeID))
+		{
+			return NodeStorage->GetMutableNode();
 		}
 		return nullptr;
 	}
@@ -246,15 +305,15 @@ namespace Metasound
 		}
 
 		// Filter all nodes by the connected nodes to get the unconnected nodes.
-		auto IsNodeUnconnected = [&](const TPair<FGuid, TSharedPtr<const INode>>& InGuidAndNode)
+		auto IsNodeUnconnected = [&](const TPair<FGuid, FDeprecationNodeStorageAdapter>& InGuidAndNode)
 		{
-			const INode* NodePtr = InGuidAndNode.Value.Get();
+			const INode* NodePtr = InGuidAndNode.Value.GetConstNode();
 			return (nullptr != NodePtr) && (INDEX_NONE == Algo::BinarySearch(ConnectedNodes, NodePtr));
 		};
 
-		auto RemoveSharedPtr = [&](const TPair<FGuid, TSharedPtr<const INode>>& InGuidAndNode) -> TPair<FGuid, const INode*>
+		auto RemoveSharedPtr = [&](const TPair<FGuid, FDeprecationNodeStorageAdapter>& InGuidAndNode) -> TPair<FGuid, const INode*>
 		{
-			return TPair<FGuid, const INode*>(InGuidAndNode.Key, InGuidAndNode.Value.Get());
+			return TPair<FGuid, const INode*>(InGuidAndNode.Key, InGuidAndNode.Value.GetConstNode());
 		};
 
 		Algo::TransformIf(Nodes, OutUnconnectedNodes, IsNodeUnconnected, RemoveSharedPtr);
@@ -266,10 +325,10 @@ namespace Metasound
 	{
 		if (bInRemoveDataEdgesWithNode)
 		{
-			TSharedPtr<const INode> NodePtr;
+			FDeprecationNodeStorageAdapter NodePtr;
 			if (Nodes.RemoveAndCopyValue(InNodeID, NodePtr))
 			{
-				RemoveDataEdgesWithNode(*NodePtr);
+				RemoveDataEdgesWithNode(*NodePtr.GetConstNode());
 				return true;
 			}
 			return false;
@@ -290,6 +349,40 @@ namespace Metasound
 		return Metadata.DefaultInterface;
 	}
 
+	void FGraph::SetDefaultInput(const FVertexName& InVertexName, const FLiteral& InLiteral)
+	{
+		// There's some significant redundancy when storing default literals on a
+		// graph's input. We update them all. 
+		if (FInputDataVertex* InputVertex = Metadata.DefaultInterface.GetInputInterface().Find(InVertexName))
+		{
+			// Update literal on exposed vertex
+			InputVertex->SetDefaultLiteral(InLiteral);
+
+			// Update literal on input data destination
+			const INode* InputNode = nullptr;
+			for (auto& Pair : InputDestinations)
+			{
+				if (Pair.Key.Get<1>() == InVertexName)
+				{
+					Pair.Value.Vertex.SetDefaultLiteral(InLiteral);
+					InputNode = Pair.Key.Get<0>();
+					break;
+				}
+			}
+
+			// Update literal on input node
+			if (InputNode)
+			{
+				// This check for mutable nodes will not be necessary once we remove
+				// storage for TSharedPtr<const INodes>. At that point we can assume
+				// that all internal nodes can be mutated safely. 
+				if (INode* MutableInputNode = FindMutableNode(InputNode->GetInstanceID()))
+				{
+					MutableInputNode->SetDefaultInput(InVertexName, InLiteral);
+				}
+			}
+		}
+	}
 
 	bool FGraph::SetVertexInterface(const FVertexInterface& InInterface)
 	{
diff --git a/Engine/Plugins/Runtime/Metasound/Source/MetasoundGraphCore/Private/MetasoundNode.cpp b/Engine/Plugins/Runtime/Metasound/Source/MetasoundGraphCore/Private/MetasoundNode.cpp
--- a/Engine/Plugins/Runtime/Metasound/Source/MetasoundGraphCore/Private/MetasoundNode.cpp
+++ b/Engine/Plugins/Runtime/Metasound/Source/MetasoundGraphCore/Private/MetasoundNode.cpp
@@ -36,25 +36,15 @@ namespace Metasound
 		return Info.DefaultInterface;
 	}
 
-	bool FNode::SetVertexInterface(const FVertexInterface& InInterface)
+	void FNode::SetDefaultInput(const FVertexName& InVertexName, const FLiteral& InLiteral)
 	{
-		// This code looks a bit counter-intuitive. Here's the explanation: 
-		// 1. This node doesn't support manipulating the interface. The default 
-		//    interface is the only interface that works with this node.
-		// 2. The only valid argument to this function is so set the interface 
-		//    to it's default state. 
-		// 3. If the input vertex interface is exactly the same as the default 
-		//    interface, then we can say that this function executed 
-		//    successfully.
-		// 4. Return "true" if the supplied interface is the same as the 
-		//    default interface, false otherwise. 
-		return GetVertexInterface() == InInterface;
-	}
-
-	bool FNode::IsVertexInterfaceSupported(const FVertexInterface& InInterface) const
-	{
-		// This node doesn't support manipulating the interface. The default 
-		// interface is the only interface that works with this node.
-		return GetVertexInterface() == InInterface;
+		if (FInputDataVertex* Vertex = Info.DefaultInterface.GetInputInterface().Find(InVertexName))
+		{
+			Vertex->SetDefaultLiteral(InLiteral);
+		}
+		else
+		{
+			UE_LOG(LogMetaSound, Error, TEXT("Cannot set default input. Vertex %s does not exist on node %s"), *InVertexName.ToString(), *InstanceName.ToString());
+		}
 	}
 }
diff --git a/Engine/Plugins/Runtime/Metasound/Source/MetasoundGraphCore/Private/MetasoundNodeInterface.cpp b/Engine/Plugins/Runtime/Metasound/Source/MetasoundGraphCore/Private/MetasoundNodeInterface.cpp
--- a/Engine/Plugins/Runtime/Metasound/Source/MetasoundGraphCore/Private/MetasoundNodeInterface.cpp
+++ b/Engine/Plugins/Runtime/Metasound/Source/MetasoundGraphCore/Private/MetasoundNodeInterface.cpp
@@ -167,4 +167,16 @@ namespace Metasound
 			return InLeft.From < InRight.From;
 		}
 	}
+
+#if !METASOUND_PURE_VIRTUAL_SETDEFAULTINPUT
+	void INodeBase::SetDefaultInput(const FVertexName& InVertexName, const FLiteral& InLiteral)
+	{
+		static bool bDidWarn = false;
+		if (!bDidWarn)
+		{
+			UE_LOG(LogMetaSound, Warning, TEXT("Ignoring input default for vertex %s. Please implement INodeInterface::SetDefaultInput(...) for the class representing node %s. This method will become pure virtual on or after release 5.8. Define METASOUND_PURE_VIRTUAL_SETDEFAULTINPUT in order to build with this method as a pure virtual on the interface."), *InVertexName.ToString(), *GetInstanceName().ToString());
+			bDidWarn = true;
+		}
+	}
+#endif
 }
diff --git a/Engine/Plugins/Runtime/Metasound/Source/MetasoundGraphCore/Private/MetasoundOperatorBuilder.cpp b/Engine/Plugins/Runtime/Metasound/Source/MetasoundGraphCore/Private/MetasoundOperatorBuilder.cpp
--- a/Engine/Plugins/Runtime/Metasound/Source/MetasoundGraphCore/Private/MetasoundOperatorBuilder.cpp
+++ b/Engine/Plugins/Runtime/Metasound/Source/MetasoundGraphCore/Private/MetasoundOperatorBuilder.cpp
@@ -493,6 +493,10 @@ namespace Metasound
 					AddBuildError<FInvalidConnectionDataTypeError>(InOutContext.Results.Errors, FDataEdge{Source, Destination});
 				}
 			}
+			else
+			{
+				OperatorInfo.VertexData.GetInputs().SetDefaultLiteral(Destination.Vertex.VertexName, Destination.Vertex.GetDefaultLiteral());
+			}
 		}
 
 		return BuildStatus;
diff --git a/Engine/Plugins/Runtime/Metasound/Source/MetasoundGraphCore/Private/MetasoundVertexData.cpp b/Engine/Plugins/Runtime/Metasound/Source/MetasoundGraphCore/Private/MetasoundVertexData.cpp
--- a/Engine/Plugins/Runtime/Metasound/Source/MetasoundGraphCore/Private/MetasoundVertexData.cpp
+++ b/Engine/Plugins/Runtime/Metasound/Source/MetasoundGraphCore/Private/MetasoundVertexData.cpp
@@ -91,6 +91,11 @@ namespace Metasound
 			return Vertex;
 		}
 
+		void FInputBinding::SetDefaultLiteral(const FLiteral& InLiteral)
+		{
+			Vertex.SetDefaultLiteral(InLiteral);
+		}
+
 		bool FInputBinding::IsBound() const
 		{
 			return Data.IsSet();
@@ -570,6 +575,18 @@ namespace Metasound
 		}
 	}
 
+	void FInputVertexInterfaceData::SetDefaultLiteral(const FVertexName& InVertexName, const FLiteral& InLiteral)
+	{
+		if (FInputBinding* Binding = Find(InVertexName))
+		{
+			Binding->SetDefaultLiteral(InLiteral);
+		}
+		else
+		{
+			UE_LOG(LogMetaSound, Warning, TEXT("Cannot set default literal. Failed to find input vertex with name '%s'."), *InVertexName.ToString());
+		}
+	}
+
 	EDataReferenceAccessType FInputVertexInterfaceData::GetVertexDataAccessType(const FVertexName& InVertexName) const
 	{
 		return VertexDataPrivate::GetVertexDataAccessType(Bindings, InVertexName);
diff --git a/Engine/Plugins/Runtime/Metasound/Source/MetasoundGraphCore/Public/MetasoundBasicNode.h b/Engine/Plugins/Runtime/Metasound/Source/MetasoundGraphCore/Public/MetasoundBasicNode.h
--- a/Engine/Plugins/Runtime/Metasound/Source/MetasoundGraphCore/Public/MetasoundBasicNode.h
+++ b/Engine/Plugins/Runtime/Metasound/Source/MetasoundGraphCore/Public/MetasoundBasicNode.h
@@ -31,6 +31,9 @@ namespace Metasound
 			/** Return the interface associated with this node instance. */
 			virtual const FVertexInterface& GetVertexInterface() const override;
 
+			/** Set the default literal for an input vertex. */
+			virtual void SetDefaultInput(const FVertexName& InVertexName, const FLiteral& InLiteral) override;
+
 		private:
 
 			FNodeData NodeData;
diff --git a/Engine/Plugins/Runtime/Metasound/Source/MetasoundGraphCore/Public/MetasoundDynamicOperatorTransactor.h b/Engine/Plugins/Runtime/Metasound/Source/MetasoundGraphCore/Public/MetasoundDynamicOperatorTransactor.h
--- a/Engine/Plugins/Runtime/Metasound/Source/MetasoundGraphCore/Public/MetasoundDynamicOperatorTransactor.h
+++ b/Engine/Plugins/Runtime/Metasound/Source/MetasoundGraphCore/Public/MetasoundDynamicOperatorTransactor.h
@@ -18,6 +18,13 @@
 #define METASOUND_DEBUG_DYNAMIC_TRANSACTOR !UE_BUILD_SHIPPING
 #endif
 
+// In UE 5.6, registered node are expected to support the constructor signature Constructor(FNodeData, TSharedRef<const FNodeClassMetadata>)
+// Because there are many existing nodes, it may take time to update them. For convenience, the deprecations related to this change are
+// configurable via a preprocessor macro so that the deprecation warnings do not drown out other compiler errors and warnings.
+#ifndef UE_METASOUND_ENABLE_5_6_NODE_REGISTRATION_DEPRECATION_WARNINGS
+#define UE_METASOUND_ENABLE_5_6_NODE_REGISTRATION_DEPRECATION_WARNINGS (0)
+#endif
+
 namespace Metasound
 {
 	class FOperatorSettings;
@@ -208,11 +215,19 @@ namespace Metasound
 			 */
 			void AddDataEdge(const FGuid& InFromNodeID, const FVertexName& InFromVertex, const FGuid& InToNodeID, const FVertexName& InToVertex);
 
+#if UE_METASOUND_ENABLE_5_6_NODE_REGISTRATION_DEPRECATION_WARNINGS
+			UE_DEPRECATED(5.6, "Literal nodes will no longer be supported. Use RemoveDataEdge with an FLiteral instead.")
+#endif
 			/** Remove the given data edge. */
-			void RemoveDataEdge(const FGuid& InFromNode, const FVertexName& InFromVertex, const FGuid& InToNode, const FVertexName& InToVertex, TUniquePtr<INode> InReplacementLiteralNode);
+			void RemoveDataEdge(const FGuid& InFromNode, const FVertexName& InFromVertexName, const FGuid& InToNode, const FVertexName& InToVertexName, TUniquePtr<INode> InReplacementLiteralNode);
+			void RemoveDataEdge(const FGuid& InFromNode, const FVertexName& InFromVertexName, const FGuid& InToNode, const FVertexName& InToVertexName, FLiteral InReplacementLiteral, const FReferenceCreationFunction InReferenceCreateFunc);
 
+#if UE_METASOUND_ENABLE_5_6_NODE_REGISTRATION_DEPRECATION_WARNINGS
+			UE_DEPRECATED(5.6, "Setting values using a literal node will no longer be supported. Set literal on the vertex interface instead.")
+#endif
 			/** Set the value on a unconnected node input vertex. */
-			void SetValue(const FGuid& InNodeID, const FVertexName& InVertex, TUniquePtr<INode> InLiteralNode);
+			void SetValue(const FGuid& InNodeID, const FVertexName& InVertexName, TUniquePtr<INode> InLiteralNode);
+			void SetValue(const FGuid& InNodeID, const FVertexName& InVertexName, const FLiteral& InLiteral, const FReferenceCreationFunction InReferenceCreateFunc);
 
 			/** Add an input data destination to describe how data provided 
 			 * outside this graph should be routed internally.
@@ -251,10 +266,15 @@ namespace Metasound
 			void EnqueueBeginFadeOperatorTransform(const INode& InNode, EAudioFadeType InFadeType, TArrayView<const FVertexName> InInputsToFade, TArrayView<const FVertexName> InOutputsToFade);
 			void EnqueueEndFadeOperatorTransform(const INode& InNode);
 
-			void EnqueueRemoveEdgeOperatorTransform(const INode& InFromNode, const FVertexName& InFromVertex, const INode& InToNode, const FVertexName& InToVertex, const INode& InReplacementLiteralNode, int32 InLiteralOrdinal);
-			void EnqueueFadeAndRemoveEdgeOperatorTransform(const INode& InFromNode, const FVertexName& InFromVertex, const INode& InToNode, const FVertexName& InToVertex, const INode& InReplacementLiteralNode, int32 InLiteralOrdinal);
-			void EnqueueAddEdgeOperatorTransform(const INode& InFromNode, const FVertexName& InFromVertex, const INode& InToNode, const FVertexName& InToVertex, const INode* InPriorLiteralNode, const TArray<FOrdinalSwap>& InOrdinalUpdates);
-			void EnqueueFadeAndAddEdgeOperatorTransform(const INode& InFromNode, const FVertexName& InFromVertex, const INode& InToNode, const FVertexName& InToVertex, const INode* InPriorLiteralNode, const TArray<FOrdinalSwap>& InOrdinalUpdates);
+			// Only used by deprecated 5.6 functions
+			void EnqueueRemoveEdgeOperatorTransform_Deprecated(const INode& InFromNode, const FVertexName& InFromVertexName, const INode& InToNode, const FVertexName& InToVertexName, const INode& InReplacementLiteralNode, int32 InLiteralOrdinal);
+			void EnqueueFadeAndRemoveEdgeOperatorTransform_Deprecated(const INode& InFromNode, const FVertexName& InFromVertexName, const INode& InToNode, const FVertexName& InToVertexName, const INode& InReplacementLiteralNode, int32 InLiteralOrdinal);
+
+			void EnqueueRemoveEdgeOperatorTransform(const INode& InFromNode, const FVertexName& InFromVertexName, const INode& InToNode, const FVertexName& InToVertexName, const FLiteral& InReplacementLiteral, const FReferenceCreationFunction InReferenceCreateFunc);
+			void EnqueueFadeAndRemoveEdgeOperatorTransform(const INode& InFromNode, const FVertexName& InFromVertexName, const INode& InToNode, const FVertexName& InToVertexName, const FLiteral& InReplacementLiteral, const FReferenceCreationFunction InReferenceCreateFunc);
+
+			void EnqueueAddEdgeOperatorTransform(const INode& InFromNode, const FVertexName& InFromVertex, const INode& InToNode, const FVertexName& InToVertex, const TArray<FOrdinalSwap>& InOrdinalUpdates);
+			void EnqueueFadeAndAddEdgeOperatorTransform(const INode& InFromNode, const FVertexName& InFromVertex, const INode& InToNode, const FVertexName& InToVertex, const TArray<FOrdinalSwap>& InOrdinalUpdates);
 
 			void AddDataEdgeInternal(const INode& InFromNode, const FVertexName& InFromVertex, const FGuid& InToNodeID, const INode& InToNode, const FVertexName& InToVertex);
 			struct FDynamicOperatorInfo
diff --git a/Engine/Plugins/Runtime/Metasound/Source/MetasoundGraphCore/Public/MetasoundGraph.h b/Engine/Plugins/Runtime/Metasound/Source/MetasoundGraphCore/Public/MetasoundGraph.h
--- a/Engine/Plugins/Runtime/Metasound/Source/MetasoundGraphCore/Public/MetasoundGraph.h
+++ b/Engine/Plugins/Runtime/Metasound/Source/MetasoundGraphCore/Public/MetasoundGraph.h
@@ -37,12 +37,16 @@ namespace Metasound
 			/** Return the current vertex interface. */
 			virtual const FVertexInterface& GetVertexInterface() const override;
 
+			/** Set the default value associated with an input vertex. */
+			virtual void SetDefaultInput(const FVertexName& InVertexName, const FLiteral& InLiteral) override;
+
 			/** Set the vertex interface. If the vertex was successfully changed, returns true. 
 			 *
 			 * @param InInterface - New interface for node. 
 			 *
 			 * @return True on success, false otherwise.
 			 */
+			UE_DEPRECATED(5.6, "FGraphs can only set their interface by adding inputs and outputs.")
 			virtual bool SetVertexInterface(const FVertexInterface& InInterface) override;
 
 			/** Expresses whether a specific vertex interface is supported.
@@ -51,6 +55,7 @@ namespace Metasound
 			 *
 			 * @return True if the interface is supported, false otherwise. 
 			 */
+			UE_DEPRECATED(5.6, "This function no longer serves a purpose.")
 			virtual bool IsVertexInterfaceSupported(const FVertexInterface& InInterface) const override;
 
 			/** Get vertices which contain input parameters. */
@@ -95,12 +100,23 @@ namespace Metasound
 			/** Removes all edges with connected to the node. */
 			void RemoveDataEdgesWithNode(const INode& InNode);
 
+			UE_DEPRECATED(5.6, TEXT("Use AddNode(...) which accepts a unique ptr"))
+			void AddNode(const FGuid& InNodeID, TSharedPtr<const INode> InNode);
+
 			/** Store a node on this graph. 
 			 *
 			 * @param InNodeID - A unique ID associated with the node.
 			 * @param InNode - A shared pointer to a node. 
 			 */
-			void AddNode(const FGuid& InNodeID, TSharedPtr<const INode> InNode);
+			void AddNode(const FGuid& InNodeID, TUniquePtr<INode> InNode);
+
+			/** Set the input default literal for a node that exists in the graph. 
+			 *
+			 * @param InNodeID - ID of not to update. 
+			 * @param InVertexName - Name of input vertex on the node.
+			 * @param InLiteral - Literal used to construct default input. 
+			 */
+			void SetNodeDefaultInput(const FGuid& InNodeID, const FVertexName& InVertexName, const FLiteral& InLiteral);
 
 			/** Add an input node to this graph.
 			 *
@@ -108,16 +124,34 @@ namespace Metasound
 			 * @param InVertexName - The key for the graph input vertex.
 			 * @param InNode - A shared pointer to an input node. 
 			 */
+			UE_DEPRECATED(5.6, TEXT("Use AddInputNode(...) which accepts a unique ptr"))
 			void AddInputNode(FGuid InNodeID, const FVertexName& InVertexName, TSharedPtr<const INode> InNode);
 
+			/** Add an input node to this graph.
+			 *
+			 * @param InNodeID - A unique ID associated with the node.
+			 * @param InVertexName - The key for the graph input vertex.
+			 * @param InNode - A unique pointer to an input node. 
+			 */
+			void AddInputNode(FGuid InNodeID, const FVertexName& InVertexName, TUniquePtr<INode> InNode);
+
 			/** Add an output node to this graph.
 			 *
 			 * @param InNodeID - A unique ID associated with the node.
 			 * @param InVertexName - The key for the graph output vertex.
 			 * @param InNode - A shared pointer to an output node. 
 			 */
+			UE_DEPRECATED(5.6, TEXT("Use AddOutputNode(...) which accepts a unique ptr"))
 			void AddOutputNode(FGuid InNodeID, const FVertexName& InVertexName, TSharedPtr<const INode> InNode);
 
+			/** Add an output node to this graph.
+			 *
+			 * @param InNodeID - A unique ID associated with the node.
+			 * @param InVertexName - The key for the graph output vertex.
+			 * @param InNode - A shared pointer to an output node. 
+			 */
+			void AddOutputNode(FGuid InNodeID, const FVertexName& InVertexName, TUniquePtr<INode> InNode);
+
 			/** Retrieve node by node ID.
 			 *
 			 * @param InNodeID - The NodeID of the requested Node.
@@ -196,6 +230,9 @@ namespace Metasound
 			virtual FOperatorFactorySharedRef GetDefaultOperatorFactory() const override;
 
 		private:
+			// Non-const access to nodes should not be exposed publicly so that
+			// graph edges can be assured to be consistent. 
+			INode* FindMutableNode(const FGuid& InNodeID);
 
 			class FFactory : public IOperatorFactory
 			{
@@ -211,7 +248,35 @@ namespace Metasound
 			FNodeClassMetadata Metadata;
 
 			TArray<FDataEdge> Edges;
-			TSortedMap<FGuid, TSharedPtr<const INode>> Nodes;
+
+			// This struct is used to handle back-compatibility with the deprecated
+			// API. To support the old API, we still hold onto shared pointers of
+			// nodes in the graph. Once the deprecated API is removed, we can remove
+			// this struct and simply store unique ptrs. 
+			struct FDeprecationNodeStorageAdapter
+			{
+				FDeprecationNodeStorageAdapter() = default;
+
+				FDeprecationNodeStorageAdapter(TUniquePtr<INode> InNode)
+				: Storage(InNode.Get())
+				{
+					MutableNode = InNode.Release(); // Node lifetime managed by storage in TSharedPtr<>
+				}
+
+				FDeprecationNodeStorageAdapter(TSharedPtr<const INode> InNode)
+				: Storage(MoveTemp(InNode))
+				{
+				}
+
+				const INode* GetConstNode() const;
+				INode* GetMutableNode() const;
+
+			private:
+				INode* MutableNode;
+				TSharedPtr<const INode> Storage;
+			};
+
+			TSortedMap<FGuid, FDeprecationNodeStorageAdapter> Nodes;
 
 			FInputDataDestinationCollection InputDestinations;
 			FOutputDataSourceCollection OutputSources;
diff --git a/Engine/Plugins/Runtime/Metasound/Source/MetasoundGraphCore/Public/MetasoundNode.h b/Engine/Plugins/Runtime/Metasound/Source/MetasoundGraphCore/Public/MetasoundNode.h
--- a/Engine/Plugins/Runtime/Metasound/Source/MetasoundGraphCore/Public/MetasoundNode.h
+++ b/Engine/Plugins/Runtime/Metasound/Source/MetasoundGraphCore/Public/MetasoundNode.h
@@ -26,11 +26,7 @@ namespace Metasound
 
 			virtual const FVertexInterface& GetVertexInterface() const override;
 
-			UE_DEPRECATED(5.6, "FNode will not support modification to the FVertexInterface")
-			virtual bool SetVertexInterface(const FVertexInterface& InInterface) override;
-
-			UE_DEPRECATED(5.6, "FNode will not support modification to the FVertexInterface")
-			virtual bool IsVertexInterfaceSupported(const FVertexInterface& InInterface) const override;
+			virtual void SetDefaultInput(const FVertexName& InVertexName, const FLiteral& InLiteral) override;
 
 		private:
 
diff --git a/Engine/Plugins/Runtime/Metasound/Source/MetasoundGraphCore/Public/MetasoundNodeInterface.h b/Engine/Plugins/Runtime/Metasound/Source/MetasoundGraphCore/Public/MetasoundNodeInterface.h
--- a/Engine/Plugins/Runtime/Metasound/Source/MetasoundGraphCore/Public/MetasoundNodeInterface.h
+++ b/Engine/Plugins/Runtime/Metasound/Source/MetasoundGraphCore/Public/MetasoundNodeInterface.h
@@ -11,6 +11,13 @@
 #include "MetasoundLiteral.h"
 #include "Misc/Guid.h"
 
+// This define helps in migrating code which has not yet implemented the purevirtual
+// INode::SetDefaultInput(...) which was introduced in UE 5.6. By UE 5.8 or later,
+// the method will become pure virtual and this macro will no longer exist. 
+#ifndef METASOUND_PURE_VIRTUAL_SETDEFAULTINPUT
+#define METASOUND_PURE_VIRTUAL_SETDEFAULTINPUT 0
+#endif
+
 namespace Metasound
 {
 	extern const FString METASOUNDGRAPHCORE_API PluginAuthor;
@@ -165,6 +172,14 @@ namespace Metasound
 			/** Return the current vertex interface. */
 			virtual const FVertexInterface& GetVertexInterface() const = 0;
 
+			/** Set the default input literal for a vertex. This literal will be 
+			 * used if nothing is connected to the node's input vertex. */
+#if METASOUND_PURE_VIRTUAL_SETDEFAULTINPUT
+			virtual void SetDefaultInput(const FVertexName& InVertexName, const FLiteral& InLiteral) = 0;
+#else 
+			virtual METASOUNDGRAPHCORE_API void SetDefaultInput(const FVertexName& InVertexName, const FLiteral& InLiteral);
+#endif
+
 			/** Set the vertex interface. If the vertex was successfully changed, returns true. 
 			 *
 			 * @param InInterface - New interface for node. 
diff --git a/Engine/Plugins/Runtime/Metasound/Source/MetasoundGraphCore/Public/MetasoundVertexData.h b/Engine/Plugins/Runtime/Metasound/Source/MetasoundGraphCore/Public/MetasoundVertexData.h
--- a/Engine/Plugins/Runtime/Metasound/Source/MetasoundGraphCore/Public/MetasoundVertexData.h
+++ b/Engine/Plugins/Runtime/Metasound/Source/MetasoundGraphCore/Public/MetasoundVertexData.h
@@ -246,6 +246,8 @@ namespace Metasound
 
 			const FInputDataVertex& GetVertex() const;
 
+			void SetDefaultLiteral(const FLiteral& InLiteral);
+
 			bool IsBound() const;
 
 			EDataReferenceAccessType GetAccessType() const;
@@ -725,6 +727,9 @@ namespace Metasound
 		/** Return the vertex associated with the vertex name. */
 		const FInputDataVertex& GetVertex(const FVertexName& InVertexName) const;
 
+		/** Set the default literal used to create default values. */
+		void SetDefaultLiteral(const FVertexName& InVertexName, const FLiteral& InLiteral);
+
 		/** Add a vertex. VertexInterfaceData must be unfrozen. */
 		void AddVertex(const FInputDataVertex& InVertex);
 
