Build fix for deprecation macro in MetaSounds

diff --git a/Engine/Plugins/Runtime/Metasound/Source/MetasoundGraphCore/Public/MetasoundGraph.h b/Engine/Plugins/Runtime/Metasound/Source/MetasoundGraphCore/Public/MetasoundGraph.h
--- a/Engine/Plugins/Runtime/Metasound/Source/MetasoundGraphCore/Public/MetasoundGraph.h
+++ b/Engine/Plugins/Runtime/Metasound/Source/MetasoundGraphCore/Public/MetasoundGraph.h
@@ -100,7 +100,7 @@ namespace Metasound
 			/** Removes all edges with connected to the node. */
 			void RemoveDataEdgesWithNode(const INode& InNode);
 
-			UE_DEPRECATED(5.6, TEXT("Use AddNode(...) which accepts a unique ptr"))
+			UE_DEPRECATED(5.6, "Use AddNode(...) which accepts a unique ptr")
 			void AddNode(const FGuid& InNodeID, TSharedPtr<const INode> InNode);
 
 			/** Store a node on this graph. 
@@ -124,7 +124,7 @@ namespace Metasound
 			 * @param InVertexName - The key for the graph input vertex.
 			 * @param InNode - A shared pointer to an input node. 
 			 */
-			UE_DEPRECATED(5.6, TEXT("Use AddInputNode(...) which accepts a unique ptr"))
+			UE_DEPRECATED(5.6, "Use AddInputNode(...) which accepts a unique ptr")
 			void AddInputNode(FGuid InNodeID, const FVertexName& InVertexName, TSharedPtr<const INode> InNode);
 
 			/** Add an input node to this graph.
@@ -141,7 +141,7 @@ namespace Metasound
 			 * @param InVertexName - The key for the graph output vertex.
 			 * @param InNode - A shared pointer to an output node. 
 			 */
-			UE_DEPRECATED(5.6, TEXT("Use AddOutputNode(...) which accepts a unique ptr"))
+			UE_DEPRECATED(5.6, "Use AddOutputNode(...) which accepts a unique ptr")
 			void AddOutputNode(FGuid InNodeID, const FVertexName& InVertexName, TSharedPtr<const INode> InNode);
 
 			/** Add an output node to this graph.
