Refactor MetaSound AssetTag Serialization Code Part 1
- Add and refactor tag data required interface data needed to determine node connection eligibility without loading the full MetaSound in Editor

#rnx

diff --git a/Engine/Plugins/Runtime/Metasound/Source/MetasoundEditor/Private/MetasoundAssetDefinitions.cpp b/Engine/Plugins/Runtime/Metasound/Source/MetasoundEditor/Private/MetasoundAssetDefinitions.cpp
--- a/Engine/Plugins/Runtime/Metasound/Source/MetasoundEditor/Private/MetasoundAssetDefinitions.cpp
+++ b/Engine/Plugins/Runtime/Metasound/Source/MetasoundEditor/Private/MetasoundAssetDefinitions.cpp
@@ -13,6 +13,7 @@
 #include "Metasound.h"
 #include "MetasoundAssetBase.h"
 #include "MetasoundAssetManager.h"
+#include "MetasoundAssetSubsystem.h"
 #include "MetasoundBuilderSubsystem.h"
 #include "MetasoundSource.h"
 #include "MetasoundEditor.h"
@@ -43,41 +44,20 @@ namespace Metasound::Editor
 {
 	namespace AssetDefinitionsPrivate
 	{
-		bool GetIsPreset(const FSoftObjectPath& InSourcePath)
+		const FSlateBrush* GetClassBrush(const FAssetData& InAssetData, FName InClassName, bool bIsThumbnail = false)
 		{
 			using namespace Frontend;
 
-			bool bIsPreset = false;
-			// If object in memory, try to resolve but not load
-			if (UObject* Object = InSourcePath.ResolveObject())
-			{
-				const FMetasoundAssetBase* MetaSoundAsset = IMetasoundUObjectRegistry::Get().GetObjectAsAssetBase(Object);
-				if (MetaSoundAsset)
-				{
-					bIsPreset = MetaSoundAsset->GetConstDocumentChecked().RootGraph.PresetOptions.bIsPreset;
-				}
-			}
-			// Otherwise, try to pull from asset registry, but avoid load as this call
-			// would then be slow and is called from the ContentBrowser many times
-			else
+			const FAssetClassQueryInfo QueryInfo(InAssetData);
+			if (!QueryInfo.bIsValid)
 			{
-				FAssetData IconAssetData;
-				const FAssetRegistryModule& AssetRegistryModule = FModuleManager::GetModuleChecked<FAssetRegistryModule>("AssetRegistry");
-				UE::AssetRegistry::EExists Exists = AssetRegistryModule.TryGetAssetByObjectPath(InSourcePath, IconAssetData);
-				if (Exists == UE::AssetRegistry::EExists::Exists)
-				{
-					IconAssetData.GetTagValue(AssetTags::IsPreset, bIsPreset);
-				}
+				UE_LOG(LogMetaSound, VeryVerbose,
+					TEXT("ClassBrush for asset '%s' may return incorrect preset icon. Asset requires reserialization."),
+					*InAssetData.GetObjectPathString());
 			}
 
-			return bIsPreset;
-		}
-
-		const FSlateBrush* GetClassBrush(const FAssetData& InAssetData, FName InClassName, bool bIsThumbnail = false)
-		{
-			const bool bIsPreset = GetIsPreset(InAssetData.ToSoftObjectPath());
 			FString BrushName = FString::Printf(TEXT("MetasoundEditor.%s"), *InClassName.ToString());
-			if (bIsPreset)
+			if (QueryInfo.bIsPreset)
 			{
 				BrushName += TEXT(".Preset");
 			}
diff --git a/Engine/Plugins/Runtime/Metasound/Source/MetasoundEditor/Private/MetasoundEditor.cpp b/Engine/Plugins/Runtime/Metasound/Source/MetasoundEditor/Private/MetasoundEditor.cpp
--- a/Engine/Plugins/Runtime/Metasound/Source/MetasoundEditor/Private/MetasoundEditor.cpp
+++ b/Engine/Plugins/Runtime/Metasound/Source/MetasoundEditor/Private/MetasoundEditor.cpp
@@ -106,6 +106,11 @@ namespace Metasound
 {
 	namespace Editor
 	{
+		namespace AssetEditorPrivate
+		{
+			static const FText LoadingAssetsFormat = LOCTEXT("MetaSoundScanInProgressNotificationText", "Loading {0} MetaSound Assets...");
+		} // namespace AssetEditorPrivate
+
 		bool IsPreviewingMetaSound(const UObject& InMetaSound)
 		{
 			using namespace Engine;
@@ -1118,11 +1123,7 @@ namespace Metasound
 			if (!MetaSoundEditorModule.IsRestrictedMode())
 			{
 				IMetasoundEngineModule& MetaSoundEngineModule = FModuleManager::GetModuleChecked<IMetasoundEngineModule>("MetaSoundEngine");
-				if (MetaSoundEngineModule.GetNodeClassRegistryPrimeStatus() < ENodeClassRegistryPrimeStatus::InProgress)
-				{
-					MetaSoundEngineModule.PrimeAssetRegistryAsync();
-				}
-				bPrimingRegistry = MetaSoundEngineModule.GetNodeClassRegistryPrimeStatus() <= ENodeClassRegistryPrimeStatus::InProgress;
+				MetaSoundEngineModule.PrimeAssetManager();
 			}
 
 			// Support undo/redo
@@ -1280,12 +1281,6 @@ namespace Metasound
 
 			FAssetEditorToolkit::InitAssetEditor(Mode, InitToolkitHost, TEXT("MetasoundEditorApp"), StandaloneDefaultLayout, bCreateDefaultStandaloneMenu, bCreateDefaultToolbar, ObjectToEdit, bToolbarFocusable, bUseSmallToolbarIcons);
 
-			// Has to be run after widgets are initialized to properly display
-			if (bPrimingRegistry)
-			{
-				NotifyAssetPrimeInProgress();
-			}
-
 			ExtendToolbarInternal();
 			RegenerateMenusAndToolbars();
 
@@ -1304,6 +1299,14 @@ namespace Metasound
 
 			RefreshExecVisibility(Builder->GetConstBuilder().GetBuildPageID());
 			FSlateApplication::Get().SetUserFocus(0, MetasoundGraphEditor);
+
+			FMetaSoundAssetManager& AssetManager = FMetaSoundAssetManager::GetChecked();
+			const bool bAssetScanComplete = AssetManager.IsInitialAssetScanComplete();
+			const int32 NumProcessingAssets = AssetManager.GetActiveAsyncLoadRequestCount();
+			if (!bAssetScanComplete || NumProcessingAssets > 0)
+			{
+				NotifyAssetLoadingInProgress(NumProcessingAssets);
+			}
 		}
 
 		UObject* FEditor::GetMetasoundObject() const
@@ -1516,7 +1519,7 @@ namespace Metasound
 			bRefreshGraph = true;
 		}
 
-		void FEditor::NotifyAssetPrimeInProgress()
+		void FEditor::NotifyAssetLoadingInProgress(int32 NumProcessing)
 		{
 			if (MetasoundGraphEditor.IsValid())
 			{
@@ -1524,14 +1527,18 @@ namespace Metasound
 
 				FSimpleDelegate OnCloseNotification = FSimpleDelegate::CreateLambda([this]()
 				{
-					if (NotificationPtr)
+					if (LoadingNotificationPtr)
 					{
-						NotificationPtr->Fadeout();
-						NotificationPtr.Reset();
+						LoadingNotificationPtr->Fadeout();
+						LoadingNotificationPtr.Reset();
 					}
 				});
 
-				FNotificationInfo Info(LOCTEXT("MetaSoundScanInProgressNotificationText", "Registering MetaSound Assets..."));
+				const FText InitialText = NumProcessing > 0
+					? FText::Format(AssetEditorPrivate::LoadingAssetsFormat, NumProcessing)
+					: LOCTEXT("AssetScanOrTagRetrievelInProgress", "Waiting for Asset Scan/AssetTag Registration to Finish...");
+
+				FNotificationInfo Info(InitialText);
 				Info.SubText = LOCTEXT("MetaSoundScanInProgressNotificationSubText", "Class selector results may be incomplete");
 				Info.bUseThrobber = true;
 				Info.bFireAndForget = false;
@@ -1539,23 +1546,23 @@ namespace Metasound
 				Info.FadeOutDuration = 1.0f;
 				Info.ButtonDetails.Add(FNotificationButtonInfo(CloseNotificationText, FText(), OnCloseNotification));
 
-				NotificationPtr = MetasoundGraphEditor->AddNotification(Info);
-				if (NotificationPtr.IsValid())
+				LoadingNotificationPtr = MetasoundGraphEditor->AddNotification(Info);
+				if (LoadingNotificationPtr.IsValid())
 				{
-					NotificationPtr->SetVisibility(EVisibility::Visible);
-					NotificationPtr->SetCompletionState(SNotificationItem::CS_Pending);
+					LoadingNotificationPtr->SetVisibility(EVisibility::Visible);
+					LoadingNotificationPtr->SetCompletionState(SNotificationItem::CS_Pending);
 				}
 			}
 		}
 
-		void FEditor::NotifyAssetPrimeComplete()
+		void FEditor::NotifyAssetLoadingComplete()
 		{
 			if (MetasoundGraphEditor.IsValid())
 			{
-				if (NotificationPtr.IsValid())
+				if (LoadingNotificationPtr.IsValid())
 				{
-					NotificationPtr->Fadeout();
-					NotificationPtr.Reset();
+					LoadingNotificationPtr->Fadeout();
+					LoadingNotificationPtr.Reset();
 				}
 
 				FNotificationInfo Info(LOCTEXT("MetaSoundScanInProgressNotification", "MetaSound Asset Registration Complete"));
@@ -3998,11 +4005,9 @@ PRAGMA_ENABLE_DEPRECATION_WARNINGS
 					FMetasoundFrontendVersionNumber HighestVersion = ExternalNode->FindHighestVersionInRegistry();
 					const bool bHasNewVersion = HighestVersion.IsValid() && HighestVersion > Metadata.GetVersion();
 
-					// Check for non-native classes
-					const FNodeRegistryKey RegistryKey = FNodeRegistryKey(Metadata);
-					const bool bIsClassNative = FMetasoundFrontendRegistryContainer::Get()->IsNodeNative(RegistryKey);
-
-					if (bHasNewVersion || !bIsClassNative)
+					// Check for asset classes
+					const bool bIsAssetClass = IMetaSoundAssetManager::GetChecked().IsAssetClass(Metadata);
+					if (bHasNewVersion || bIsAssetClass)
 					{
 						// These are ignored here when updating as the user is actively
 						// forcing an update.
@@ -4541,27 +4546,29 @@ PRAGMA_ENABLE_DEPRECATION_WARNINGS
 		void FEditor::Tick(float DeltaTime)
 		{
 			using namespace Metasound::Engine;
+
 			UObject* MetaSound = GetMetasoundObject();
 			if (!MetaSound)
 			{
 				return;
 			}
 
-			if (bPrimingRegistry)
+			RefreshEditorContext(*MetaSound);
+
+			GraphConnectionManager->Update(DeltaTime);
+
+			if (LoadingNotificationPtr.IsValid())
 			{
-				IMetasoundEngineModule& MetaSoundEngineModule = FModuleManager::GetModuleChecked<Metasound::Engine::IMetasoundEngineModule>("MetaSoundEngine");
-				ENodeClassRegistryPrimeStatus PrimeStatus = MetaSoundEngineModule.GetNodeClassRegistryPrimeStatus();
-				EAssetScanStatus ScanStatus = MetaSoundEngineModule.GetAssetRegistryScanStatus();
-				if (PrimeStatus == Metasound::Engine::ENodeClassRegistryPrimeStatus::Complete)
+				const int32 NumProcessing = FMetaSoundAssetManager::GetChecked().GetActiveAsyncLoadRequestCount();
+				if (NumProcessing > 0)
 				{
-					bPrimingRegistry = false;
-					NotifyAssetPrimeComplete();
+					LoadingNotificationPtr->SetText(FText::Format(AssetEditorPrivate::LoadingAssetsFormat, NumProcessing));
+				}
+				else
+				{
+					NotifyAssetLoadingComplete();
 				}
 			}
-
-			RefreshEditorContext(*MetaSound);
-
-			GraphConnectionManager->Update(DeltaTime);
 		}
 
 		void FEditor::RefreshEditorContext(UObject& MetaSound)
@@ -5169,7 +5176,7 @@ PRAGMA_ENABLE_DEPRECATION_WARNINGS
 
 				UEdGraphNode* OwningNode = TargetPin->GetOwningNode();
 				FVector2D Location = FVector2D(OwningNode->NodePosX, OwningNode->NodePosY);
-				Metasound::SchemaUtils::PromoteToInput(&Graph, TargetPin, Location - DisplayStyle::NodeLayout::DefaultOffsetX, /*bSelectNewNode=*/ true);
+				SchemaUtils::PromoteToInput(&Graph, TargetPin, Location - DisplayStyle::NodeLayout::DefaultOffsetX, /*bSelectNewNode=*/ true);
 			}
 		}
 
@@ -5188,7 +5195,7 @@ PRAGMA_ENABLE_DEPRECATION_WARNINGS
 
 		void FEditor::PromoteToOutput()
 		{
-			using namespace Metasound::Frontend;
+			using namespace Frontend;
 
 			if (MetasoundGraphEditor.IsValid())
 			{
@@ -5199,7 +5206,7 @@ PRAGMA_ENABLE_DEPRECATION_WARNINGS
 
 				UEdGraphNode* OwningNode = TargetPin->GetOwningNode();
 				FVector2D Location = FVector2D(OwningNode->NodePosX, OwningNode->NodePosY);
-				Metasound::SchemaUtils::PromoteToOutput(&Graph, TargetPin, Location + DisplayStyle::NodeLayout::DefaultOffsetX * 2.f, /*bSelectNewNode=*/ true);
+				SchemaUtils::PromoteToOutput(&Graph, TargetPin, Location + DisplayStyle::NodeLayout::DefaultOffsetX * 2.f, /*bSelectNewNode=*/ true);
 			}
 		}
 
@@ -5223,11 +5230,11 @@ PRAGMA_ENABLE_DEPRECATION_WARNINGS
 				FVector2D Location = FVector2D(OwningNode->NodePosX, OwningNode->NodePosY);
 				if (TargetPin->Direction == EGPD_Input)
 				{
-					Metasound::SchemaUtils::PromoteToVariable(&Graph, TargetPin, Location - DisplayStyle::NodeLayout::DefaultOffsetX, /*bSelectNewNode=*/ true);
+					SchemaUtils::PromoteToVariable(&Graph, TargetPin, Location - DisplayStyle::NodeLayout::DefaultOffsetX, /*bSelectNewNode=*/ true);
 				}
 				else
 				{
-					Metasound::SchemaUtils::PromoteToMutatorVariable(&Graph, TargetPin, Location + DisplayStyle::NodeLayout::DefaultOffsetX * 2.f, /*bSelectNewNode=*/ true);
+					SchemaUtils::PromoteToMutatorVariable(&Graph, TargetPin, Location + DisplayStyle::NodeLayout::DefaultOffsetX * 2.f, /*bSelectNewNode=*/ true);
 				}
 			}
 		}
@@ -5251,7 +5258,7 @@ PRAGMA_ENABLE_DEPRECATION_WARNINGS
 
 		void FEditor::PromoteToDeferredVariable()
 		{
-			using namespace Metasound::Frontend;
+			using namespace Frontend;
 
 			if (MetasoundGraphEditor.IsValid())
 			{
@@ -5262,7 +5269,7 @@ PRAGMA_ENABLE_DEPRECATION_WARNINGS
 
 				UEdGraphNode* OwningNode = TargetPin->GetOwningNode();
 				FVector2D Location = FVector2D(OwningNode->NodePosX, OwningNode->NodePosY);
-				Metasound::SchemaUtils::PromoteToDeferredVariable(&Graph, TargetPin, Location - DisplayStyle::NodeLayout::DefaultOffsetX, /*bSelectNewNode=*/ true);
+				SchemaUtils::PromoteToDeferredVariable(&Graph, TargetPin, Location - DisplayStyle::NodeLayout::DefaultOffsetX, /*bSelectNewNode=*/ true);
 			}
 		}
 
diff --git a/Engine/Plugins/Runtime/Metasound/Source/MetasoundEditor/Private/MetasoundEditor.h b/Engine/Plugins/Runtime/Metasound/Source/MetasoundEditor/Private/MetasoundEditor.h
--- a/Engine/Plugins/Runtime/Metasound/Source/MetasoundEditor/Private/MetasoundEditor.h
+++ b/Engine/Plugins/Runtime/Metasound/Source/MetasoundEditor/Private/MetasoundEditor.h
@@ -514,8 +514,8 @@ namespace Metasound::Editor
 
 		void EditObjectSettings();
 
-		void NotifyAssetPrimeComplete();
-		void NotifyAssetPrimeInProgress();
+		void NotifyAssetLoadingComplete();
+		void NotifyAssetLoadingInProgress(int32 NumProcessing);
 		void NotifyDocumentVersioned();
 		void NotifyNodePasteFailure_MultipleVariableSetters();
 		void NotifyNodePasteFailure_ReferenceLoop();
@@ -588,16 +588,13 @@ namespace Metasound::Editor
 			* being manually applied in code */
 		bool bManuallyClearingGraphSelection = false;
 
-		/** Whether or not registry is currently being primed due to assets still loading */
-		bool bPrimingRegistry = false;
-
 		/** Highest message severity set on last validation pass of graph. */
 		int32 HighestMessageSeverity = EMessageSeverity::Info;
 
 		/** If set, used to inform user of validation results on hover of play icon. */
 		FText GraphStatusDescriptionOverride;
 
-		TSharedPtr<SNotificationItem> NotificationPtr;
+		TSharedPtr<SNotificationItem> LoadingNotificationPtr;
 
 		bool bMemberRenameRequested = false;
 
diff --git a/Engine/Plugins/Runtime/Metasound/Source/MetasoundEditor/Private/MetasoundEditorGraphBuilder.cpp b/Engine/Plugins/Runtime/Metasound/Source/MetasoundEditor/Private/MetasoundEditorGraphBuilder.cpp
--- a/Engine/Plugins/Runtime/Metasound/Source/MetasoundEditor/Private/MetasoundEditorGraphBuilder.cpp
+++ b/Engine/Plugins/Runtime/Metasound/Source/MetasoundEditor/Private/MetasoundEditorGraphBuilder.cpp
@@ -274,21 +274,16 @@ namespace Metasound
 		{
 			using namespace Frontend;
 
-			FName Namespace;
-			FName ParameterName;
-			Audio::FParameterPath::SplitName(InNodeName, Namespace, ParameterName);
-
 			FText DisplayName;
 			auto GetAssetDisplayNameFromMetadata = [&DisplayName](const FMetasoundFrontendClassMetadata& Metadata)
 			{
 				DisplayName = Metadata.GetDisplayName();
 				if (DisplayName.IsEmptyOrWhitespace())
 				{
-					const FNodeRegistryKey RegistryKey = FNodeRegistryKey(Metadata);
-					bool bIsClassNative = FMetasoundFrontendRegistryContainer::Get()->IsNodeNative(RegistryKey);
-					if (!bIsClassNative)
+					const bool bIsAssetClass = IMetaSoundAssetManager::GetChecked().IsAssetClass(Metadata);
+					if (bIsAssetClass)
 					{
-						const FTopLevelAssetPath Path = IMetaSoundAssetManager::GetChecked().FindAssetPath(RegistryKey);
+						const FTopLevelAssetPath Path = IMetaSoundAssetManager::GetChecked().FindAssetPath(FAssetKey(Metadata));
 						if (Path.IsValid())
 						{
 							DisplayName = FText::FromName(Path.GetAssetName());
@@ -312,18 +307,25 @@ namespace Metasound
 				}
 			}
 
-			// 3. If that cannot be found, build a title from the cached node registry FName.
-			if (DisplayName.IsEmptyOrWhitespace())
+			if (DisplayName.IsEmptyOrWhitespace() || bInIncludeNamespace)
 			{
-				DisplayName = FText::FromString(ParameterName.ToString());
-			}
+				FName Namespace;
+				FName ParameterName;
+				Audio::FParameterPath::SplitName(InNodeName, Namespace, ParameterName);
 
-			// 4. Tack on the namespace if requested
-			if (bInIncludeNamespace)
-			{
-				if (!Namespace.IsNone())
+				// 3. If that cannot be found, build a title from the cached node registry FName.
+				if (DisplayName.IsEmptyOrWhitespace())
 				{
-					return FText::Format(LOCTEXT("ClassMetadataDisplayNameWithNamespaceFormat", "{0} ({1})"), DisplayName, FText::FromName(Namespace));
+					DisplayName = FText::FromString(ParameterName.ToString());
+				}
+
+				// 4. Tack on the namespace if requested
+				if (bInIncludeNamespace)
+				{
+					if (!Namespace.IsNone())
+					{
+						return FText::Format(LOCTEXT("ClassMetadataDisplayNameWithNamespaceFormat", "{0} ({1})"), DisplayName, FText::FromName(Namespace));
+					}
 				}
 			}
 
@@ -476,8 +478,8 @@ namespace Metasound
 			{
 				auto InitNodeFunc = [&InMetadata, &InNodeID](UMetasoundEditorGraph& MetasoundGraph, UMetasoundEditorGraphExternalNode& NewGraphNode)
 				{
-					const FNodeRegistryKey RegistryKey(InMetadata);
-					NewGraphNode.bIsClassNative = FMetasoundFrontendRegistryContainer::Get()->IsNodeNative(RegistryKey);
+					const bool bIsAssetClass = IMetaSoundAssetManager::GetChecked().IsAssetClass(InMetadata);
+					NewGraphNode.bIsClassNative = !bIsAssetClass;
 					NewGraphNode.NodeID = InNodeID;
 					NewGraphNode.ClassName = InMetadata.GetClassName();
 				};
diff --git a/Engine/Plugins/Runtime/Metasound/Source/MetasoundEditor/Private/MetasoundEditorGraphNode.cpp b/Engine/Plugins/Runtime/Metasound/Source/MetasoundEditor/Private/MetasoundEditorGraphNode.cpp
--- a/Engine/Plugins/Runtime/Metasound/Source/MetasoundEditor/Private/MetasoundEditorGraphNode.cpp
+++ b/Engine/Plugins/Runtime/Metasound/Source/MetasoundEditor/Private/MetasoundEditorGraphNode.cpp
@@ -1124,10 +1124,15 @@ void UMetasoundEditorGraphExternalNode::CacheBreadcrumb()
 {
 	using namespace Metasound::Frontend;
 	using namespace Metasound::Editor;
-	FConstNodeHandle NodeHandle = GetConstNodeHandle();
-	const FNodeRegistryKey RegistryKey(NodeHandle->GetClassMetadata());
-	Breadcrumb.bIsClassNative = FMetasoundFrontendRegistryContainer::Get()->IsNodeNative(RegistryKey);
-	Breadcrumb.ClassName = NodeHandle->GetClassMetadata().GetClassName();
+
+	if (const FMetasoundFrontendClass* FrontendClass = GetFrontendClass())
+	{
+		const FMetasoundFrontendClassMetadata& Metadata = FrontendClass->Metadata;
+		const bool bIsAssetClass = IMetaSoundAssetManager::GetChecked().IsAssetClass(Metadata);
+
+		Breadcrumb.bIsClassNative = !bIsAssetClass;
+		Breadcrumb.ClassName = Metadata.GetClassName();
+	}
 
 	// Cache template node generation parameters
 	if (const INodeTemplate* Template = INodeTemplateRegistry::Get().FindTemplate(Breadcrumb.ClassName))
@@ -1317,8 +1322,7 @@ void UMetasoundEditorGraphExternalNode::Validate(Metasound::Editor::FGraphNodeVa
 	}
 
 	// 4. Report if node was nativized
-	const FNodeRegistryKey RegistryKey(Metadata);
-	bool bNewIsClassNative = FMetasoundFrontendRegistryContainer::Get()->IsNodeNative(RegistryKey);
+	const bool bNewIsClassNative = !IMetaSoundAssetManager::GetChecked().IsAssetClass(Metadata);
 	if (bIsClassNative != bNewIsClassNative)
 	{
 		if (bNewIsClassNative)
@@ -1366,7 +1370,7 @@ void UMetasoundEditorGraphExternalNode::Validate(Metasound::Editor::FGraphNodeVa
 
 	// 7. Report if node class is deprecated
 	FMetasoundFrontendClass RegisteredClass;
-	if (FMetasoundFrontendRegistryContainer::Get()->GetFrontendClassFromRegistered(RegistryKey, RegisteredClass))
+	if (FMetasoundFrontendRegistryContainer::Get()->GetFrontendClassFromRegistered(FNodeRegistryKey(Metadata), RegisteredClass))
 	{
 		if (RegisteredClass.Metadata.GetIsDeprecated())
 		{
diff --git a/Engine/Plugins/Runtime/Metasound/Source/MetasoundEditor/Private/MetasoundEditorGraphSchema.cpp b/Engine/Plugins/Runtime/Metasound/Source/MetasoundEditor/Private/MetasoundEditorGraphSchema.cpp
--- a/Engine/Plugins/Runtime/Metasound/Source/MetasoundEditor/Private/MetasoundEditorGraphSchema.cpp
+++ b/Engine/Plugins/Runtime/Metasound/Source/MetasoundEditor/Private/MetasoundEditorGraphSchema.cpp
@@ -17,6 +17,7 @@
 #include "Metasound.h"
 #include "MetasoundAssetBase.h"
 #include "MetasoundAssetManager.h"
+#include "MetasoundAssetSubsystem.h"
 #include "MetasoundBuilderSubsystem.h"
 #include "MetasoundDataReference.h"
 #include "MetasoundDocumentBuilderRegistry.h"
@@ -36,6 +37,7 @@
 #include "MetasoundFrontendNodesCategories.h"
 #include "MetasoundFrontendNodeTemplateRegistry.h"
 #include "MetasoundFrontendRegistries.h"
+#include "MetasoundFrontendRegistryContainer.h"
 #include "MetasoundFrontendSearchEngine.h"
 #include "MetasoundLiteral.h"
 #include "MetasoundNodeInterface.h"
@@ -60,346 +62,615 @@
 
 #define LOCTEXT_NAMESPACE "MetaSoundEditor"
 
-static int32 EnableAllVersionsMetaSoundNodeClassCreationCVar = 0;
-FAutoConsoleVariableRef CVarEnableAllVersionsMetaSoundNodeClassCreation(
-	TEXT("au.MetaSound.EnableAllVersionsNodeClassCreation"),
-	EnableAllVersionsMetaSoundNodeClassCreationCVar,
-	TEXT("Enable creating nodes for major versions of deprecated MetaSound classes in the Editor.\n")
-	TEXT("0: Disabled (default), !0: Enabled"),
-	ECVF_Default);
 
-namespace Metasound
+namespace Metasound::Editor
 {
-	namespace Editor
+	const FText& GetContextGroupDisplayName(EPrimaryContextGroup InContextGroup)
 	{
-		const FText& GetContextGroupDisplayName(EPrimaryContextGroup InContextGroup)
+		switch (InContextGroup)
 		{
-			switch (InContextGroup)
-			{
-				case EPrimaryContextGroup::Inputs:
-					return NodeCategories::Inputs;
+			case EPrimaryContextGroup::Inputs:
+				return NodeCategories::Inputs;
 
-				case EPrimaryContextGroup::Outputs:
-					return NodeCategories::Outputs;
+			case EPrimaryContextGroup::Outputs:
+				return NodeCategories::Outputs;
 
-				case EPrimaryContextGroup::Graphs:
-					return NodeCategories::Graphs;
+			case EPrimaryContextGroup::Graphs:
+				return NodeCategories::Graphs;
 
-				case EPrimaryContextGroup::Functions:
-					return NodeCategories::Functions;
+			case EPrimaryContextGroup::Functions:
+				return NodeCategories::Functions;
 
-				case EPrimaryContextGroup::Conversions:
-					return NodeCategories::Conversions;
+			case EPrimaryContextGroup::Conversions:
+				return NodeCategories::Conversions;
 
-				case EPrimaryContextGroup::Variables:
-					return NodeCategories::Variables;
+			case EPrimaryContextGroup::Variables:
+				return NodeCategories::Variables;
 
-				case EPrimaryContextGroup::Common:
-				default:
-				{
-					return FText::GetEmpty();
-				}
+			case EPrimaryContextGroup::Common:
+			default:
+			{
+				return FText::GetEmpty();
 			}
 		}
+	}
 
-		namespace SchemaPrivate
-		{
-			static const FText CategoryDelim = LOCTEXT("MetaSoundActionsCategoryDelim", "|");
-			static const FText KeywordDelim = LOCTEXT("MetaSoundKeywordDelim", " ");
+	namespace SchemaPrivate
+	{
+		static int32 EnableAllVersionsMetaSoundNodeClassCreationCVar = 0;
+		FAutoConsoleVariableRef CVarEnableAllVersionsMetaSoundNodeClassCreation(
+			TEXT("au.MetaSound.EnableAllVersionsNodeClassCreation"),
+			EnableAllVersionsMetaSoundNodeClassCreationCVar,
+			TEXT("Enable creating nodes for major versions of deprecated MetaSound classes in the Editor.\n")
+			TEXT("0: Disabled (default), !0: Enabled"),
+			ECVF_Default);
+
+		static int32 ShowUnloadedAssetInBrowserCvar = 0;
+		FAutoConsoleVariableRef CVarShowUnloadedAssetInBrowser(
+			TEXT("au.MetaSound.Debug.ShowUnloadedAssetInBrowser"),
+			ShowUnloadedAssetInBrowserCvar,
+			TEXT("Shows a '*' in the MetaSound asset picker indicating the data displayed is from tags and not from a loaded asset.\n")
+			TEXT("1: Disabled (default), !0: Enabled"),
+			ECVF_Default);
+
+		static const FText CategoryDelim = LOCTEXT("MetaSoundActionsCategoryDelim", "|");
+		static const FText KeywordDelim = LOCTEXT("MetaSoundKeywordDelim", " ");
+
+		static const FText InputDisplayNameFormat = LOCTEXT("DisplayNameAddInputFormat", "Get {0}");
+		static const FText InputTooltipFormat = LOCTEXT("TooltipAddInputFormat", "Adds a getter for the input '{0}' to the graph.");
 
-			static const FText InputDisplayNameFormat = LOCTEXT("DisplayNameAddInputFormat", "Get {0}");
-			static const FText InputTooltipFormat = LOCTEXT("TooltipAddInputFormat", "Adds a getter for the input '{0}' to the graph.");
+		static const FText OutputDisplayNameFormat = LOCTEXT("DisplayNameAddOutputFormat", "Set {0}");
+		static const FText OutputTooltipFormat = LOCTEXT("TooltipAddOutputFormat", "Adds a setter for the output '{0}' to the graph.");
 
-			static const FText OutputDisplayNameFormat = LOCTEXT("DisplayNameAddOutputFormat", "Set {0}");
-			static const FText OutputTooltipFormat = LOCTEXT("TooltipAddOutputFormat", "Adds a setter for the output '{0}' to the graph.");
+		static const FText VariableAccessorDisplayNameFormat = LOCTEXT("DisplayNameAddVariableAccessorFormat", "Get {0}");
+		static const FText VariableAccessorTooltipFormat = LOCTEXT("TooltipAddVariableAccessorFormat", "Adds a getter for the variable '{0}' to the graph.");
 
-			static const FText VariableAccessorDisplayNameFormat = LOCTEXT("DisplayNameAddVariableAccessorFormat", "Get {0}");
-			static const FText VariableAccessorTooltipFormat = LOCTEXT("TooltipAddVariableAccessorFormat", "Adds a getter for the variable '{0}' to the graph.");
+		static const FText VariableDeferredAccessorDisplayNameFormat = LOCTEXT("DisplayNameAddVariableDeferredAccessorFormat", "Get Delayed {0}");
+		static const FText VariableDeferredAccessorTooltipFormat = LOCTEXT("TooltipAddVariableDeferredAccessorFormat", "Adds a delayed getter for the variable '{0}' to the graph.");
 
-			static const FText VariableDeferredAccessorDisplayNameFormat = LOCTEXT("DisplayNameAddVariableDeferredAccessorFormat", "Get Delayed {0}");
-			static const FText VariableDeferredAccessorTooltipFormat = LOCTEXT("TooltipAddVariableDeferredAccessorFormat", "Adds a delayed getter for the variable '{0}' to the graph.");
+		static const FText VariableMutatorDisplayNameFormat = LOCTEXT("DisplayNameAddVariableMutatorFormat", "Set {0}");
+		static const FText VariableMutatorTooltipFormat = LOCTEXT("TooltipAddVariableMutatorFormat", "Adds a setter for the variable '{0}' to the graph.");
 
-			static const FText VariableMutatorDisplayNameFormat = LOCTEXT("DisplayNameAddVariableMutatorFormat", "Set {0}");
-			static const FText VariableMutatorTooltipFormat = LOCTEXT("TooltipAddVariableMutatorFormat", "Adds a setter for the variable '{0}' to the graph.");
+		static const FText ClassDescriptionAndAuthorFormat = LOCTEXT("ClassDescriptionAndAuthorFormat", "{0}\nAuthor: {1}");
 
-			bool DataTypeSupportsAssetTypes(const Metasound::Frontend::FDataTypeRegistryInfo& InRegistryInfo, const TArray<FAssetData>& InAssets)
+		bool DataTypeSupportsAssetTypes(const Metasound::Frontend::FDataTypeRegistryInfo& InRegistryInfo, const TArray<FAssetData>& InAssets)
+		{
+			if (InRegistryInfo.PreferredLiteralType != Metasound::ELiteralType::UObjectProxy)
+			{
+				return false;
+			}
+
+			const IMetasoundEditorModule& EditorModule = FModuleManager::GetModuleChecked<IMetasoundEditorModule>("MetaSoundEditor");
+			return Algo::AnyOf(InAssets, [&EditorModule, &InRegistryInfo](const FAssetData& Asset)
 			{
-				if (InRegistryInfo.PreferredLiteralType != Metasound::ELiteralType::UObjectProxy)
+				if (InRegistryInfo.ProxyGeneratorClass)
 				{
-					return false;
+					if (const UClass* Class = Asset.GetClass())
+					{
+PRAGMA_DISABLE_DEPRECATION_WARNINGS
+						if (EditorModule.IsExplicitProxyClass(*InRegistryInfo.ProxyGeneratorClass))
+						{
+							return Class == InRegistryInfo.ProxyGeneratorClass;
+						}
+PRAGMA_ENABLE_DEPRECATION_WARNINGS
+						if (InRegistryInfo.bIsExplicit)
+						{
+							return Class == InRegistryInfo.ProxyGeneratorClass;
+						}
+						else
+						{
+							return Class->IsChildOf(InRegistryInfo.ProxyGeneratorClass);
+						}
+					}
 				}
 
-				const IMetasoundEditorModule& EditorModule = FModuleManager::GetModuleChecked<IMetasoundEditorModule>("MetaSoundEditor");
-				return Algo::AnyOf(InAssets, [&EditorModule, &InRegistryInfo](const FAssetData& Asset)
+				return false;
+			});
+		}
+
+		// Connects to first pin with the same DataType
+		bool TryConnectNewNodeToMatchingDataTypePin(UEdGraphNode& NewGraphNode, UEdGraphPin* FromPin)
+		{
+			using namespace Metasound::Frontend;
+
+			if (!FromPin)
+			{
+				return false;
+			}
+
+			if (FromPin->Direction == EGPD_Input)
+			{
+				FConstInputHandle InputHandle = FGraphBuilder::GetConstInputHandleFromPin(FromPin);
+				for (UEdGraphPin* Pin : NewGraphNode.Pins)
 				{
-					if (InRegistryInfo.ProxyGeneratorClass)
+					if (Pin->Direction == EGPD_Output)
 					{
-						if (const UClass* Class = Asset.GetClass())
+						FConstOutputHandle OutputHandle = FGraphBuilder::GetConstOutputHandleFromPin(Pin);
+						if (OutputHandle->IsValid() && 
+							InputHandle->CanConnectTo(*OutputHandle).Connectable == FConnectability::EConnectable::Yes)
 						{
-PRAGMA_DISABLE_DEPRECATION_WARNINGS
-							if (EditorModule.IsExplicitProxyClass(*InRegistryInfo.ProxyGeneratorClass))
-							{
-								return Class == InRegistryInfo.ProxyGeneratorClass;
-							}
-PRAGMA_ENABLE_DEPRECATION_WARNINGS
-							if (InRegistryInfo.bIsExplicit)
+							if (ensure(FGraphBuilder::ConnectNodes(*FromPin, *Pin, true /* bConnectEdPins */)))
 							{
-								return Class == InRegistryInfo.ProxyGeneratorClass;
+								return true;
 							}
-							else
+						}
+					}
+				}
+			}
+
+			if (FromPin->Direction == EGPD_Output)
+			{
+				FConstOutputHandle OutputHandle = FGraphBuilder::GetConstOutputHandleFromPin(FromPin);
+				for (UEdGraphPin* Pin : NewGraphNode.Pins)
+				{
+					if (Pin->Direction == EGPD_Input)
+					{
+						FConstInputHandle InputHandle = FGraphBuilder::GetConstInputHandleFromPin(Pin);
+						if (InputHandle->IsValid() && 
+							InputHandle->CanConnectTo(*OutputHandle).Connectable == FConnectability::EConnectable::Yes)
+						{
+							if (ensure(FGraphBuilder::ConnectNodes(*Pin, *FromPin, true /* bConnectEdPins */)))
 							{
-								return Class->IsChildOf(InRegistryInfo.ProxyGeneratorClass);
+								return true;
 							}
 						}
 					}
-
-					return false;
-				});
+				}
 			}
 
-			// Connects to first pin with the same DataType
-			bool TryConnectNewNodeToMatchingDataTypePin(UEdGraphNode& NewGraphNode, UEdGraphPin* FromPin)
-			{
-				using namespace Metasound::Frontend;
+			return false;
+		}
+
+		struct FDataTypeActionQuery
+		{
+			FGraphActionMenuBuilder& ActionMenuBuilder;
+			const TArray<Frontend::FConstNodeHandle>& NodeHandles;
+			FInterfaceNodeFilterFunction Filter;
+			EPrimaryContextGroup ContextGroup;
+			const FText& DisplayNameFormat;
+			const FText& TooltipFormat;
+			bool bShowSelectedActions = false;
+		};
+
+		template <typename TAction>
+		void GetDataTypeActions(const FDataTypeActionQuery& InQuery)
+		{
+			using namespace Editor;
+			using namespace Frontend;
 
-				if (!FromPin)
+			for (const FConstNodeHandle& NodeHandle : InQuery.NodeHandles)
+			{
+				if (!InQuery.Filter || InQuery.Filter(NodeHandle))
 				{
-					return false;
+					constexpr bool bIncludeNamespace = true;
+
+					const FText& GroupName = GetContextGroupDisplayName(InQuery.ContextGroup);
+					const FText NodeDisplayName = FGraphBuilder::GetDisplayName(*NodeHandle, bIncludeNamespace);
+					const FText Tooltip = FText::Format(InQuery.TooltipFormat, NodeDisplayName);
+					const FText DisplayName = FText::Format(InQuery.DisplayNameFormat, NodeDisplayName);
+					TSharedPtr<TAction> NewNodeAction = MakeShared<TAction>(GroupName, DisplayName, NodeHandle->GetID(), Tooltip, InQuery.ContextGroup);
+					InQuery.ActionMenuBuilder.AddAction(NewNodeAction);
 				}
+			}
+		}
+
+		void SelectNodeInEditor(UMetasoundEditorGraph& InMetaSoundGraph, UMetasoundEditorGraphNode& InNode)
+		{
+			TSharedPtr<FEditor> MetasoundEditor = FGraphBuilder::GetEditorForGraph(InMetaSoundGraph);
+			if (MetasoundEditor.IsValid())
+			{
+				MetasoundEditor->ClearSelectionAndSelectNode(&InNode);
+			}
+		}
+
+		void SelectNodeInEditorForRename(UMetasoundEditorGraph& InMetaSoundGraph, UMetasoundEditorGraphNode& InNode)
+		{
+			TSharedPtr<FEditor> MetasoundEditor = FGraphBuilder::GetEditorForGraph(InMetaSoundGraph);
+			if (MetasoundEditor.IsValid())
+			{
+				MetasoundEditor->ClearSelectionAndSelectNode(&InNode);
+				MetasoundEditor->SetDelayedRename();
+			}
+		}
+
+		UEdGraphNode* PromoteToVariable(const FName BaseName, UEdGraphPin& FromPin, const FName DataType, const FMetasoundFrontendClass& InVariableClass, const FVector2D& InLocation, bool bSelectNode)
+		{
+			using namespace Frontend;
+
+			UEdGraphNode* ConnectedNode = Cast<UEdGraphNode>(FromPin.GetOwningNode());
+			if (!ensure(ConnectedNode))
+			{
+				return nullptr;
+			}
 
-				if (FromPin->Direction == EGPD_Input)
+			const FMetasoundFrontendClassName& ClassName = InVariableClass.Metadata.GetClassName();
+
+			UMetasoundEditorGraph* MetaSoundGraph = CastChecked<UMetasoundEditorGraph>(ConnectedNode->GetGraph());
+			FMetaSoundFrontendDocumentBuilder& DocBuilder = MetaSoundGraph->GetBuilderChecked().GetBuilder();
+			const FName NodeName = FGraphBuilder::GenerateUniqueVariableName(DocBuilder, BaseName.ToString());
+
+			const FScopedTransaction Transaction(FText::Format(
+				LOCTEXT("PromoteNodeVertexToGraphVariableFormat", "Promote MetaSound Node {0} to {1}"),
+				FText::FromName(NodeName),
+				FText::FromName(ClassName.Namespace)));
+
+			UObject& ParentMetaSound = MetaSoundGraph->GetMetasoundChecked();
+			ParentMetaSound.Modify();
+			MetaSoundGraph->Modify();
+
+			// Cache the default literal from the pin if connecting to an input
+			FMetasoundFrontendLiteral DefaultValue;
+			if (FromPin.Direction == EGPD_Input)
+			{
+				FGraphBuilder::GetPinLiteral(FromPin, DefaultValue);
+			}
+
+			const FMetasoundFrontendVariable* FrontendVariable = DocBuilder.AddGraphVariable(NodeName, DataType);
+			if (ensure(FrontendVariable))
+			{
+				UMetasoundEditorGraphVariable* Variable = MetaSoundGraph->FindOrAddVariable(FrontendVariable->Name);
+				if (ensure(Variable))
 				{
-					FConstInputHandle InputHandle = FGraphBuilder::GetConstInputHandleFromPin(FromPin);
-					for (UEdGraphPin* Pin : NewGraphNode.Pins)
+					const FMetasoundFrontendNode* NewVariableNode = DocBuilder.AddGraphVariableNode(FrontendVariable->Name, InVariableClass.Metadata.GetType());
+					if (ensure(NewVariableNode))
 					{
-						if (Pin->Direction == EGPD_Output)
+						UMetasoundEditorGraphVariableNode* NewGraphNode = FGraphBuilder::AddVariableNode(ParentMetaSound, NewVariableNode->GetID());
+						if (ensure(NewGraphNode))
 						{
-							FConstOutputHandle OutputHandle = FGraphBuilder::GetConstOutputHandleFromPin(Pin);
-							if (OutputHandle->IsValid() && 
-								InputHandle->CanConnectTo(*OutputHandle).Connectable == FConnectability::EConnectable::Yes)
+							NewGraphNode->UpdateFrontendNodeLocation(InLocation);
+							NewGraphNode->SyncLocationFromFrontendNode();
+
+							// Set the literal using the new value if connecting to an input
+							if (FromPin.Direction == EGPD_Input)
+							{
+								UMetasoundEditorGraphMemberDefaultLiteral* DefaultLiteral = Variable->GetLiteral();
+								if (ensure(DefaultLiteral))
+								{
+									DefaultLiteral->SetFromLiteral(DefaultValue);
+								}
+
+								// Ensures the setter node value is synced with the editor literal value
+								constexpr bool bPostTransaction = false;
+								Variable->UpdateFrontendDefaultLiteral(bPostTransaction);
+							}
+
+							UEdGraphNode* EdGraphNode = CastChecked<UEdGraphNode>(NewGraphNode);
+							if (ensure(SchemaPrivate::TryConnectNewNodeToMatchingDataTypePin(*EdGraphNode, &FromPin)))
 							{
-								if (ensure(FGraphBuilder::ConnectNodes(*FromPin, *Pin, true /* bConnectEdPins */)))
+								FGraphBuilder::RegisterGraphWithFrontend(ParentMetaSound);
+								if (bSelectNode)
+								{
+									SchemaPrivate::SelectNodeInEditorForRename(*MetaSoundGraph, *NewGraphNode);
+								}
+								else
 								{
-									return true;
+									SchemaPrivate::SelectNodeInEditor(*MetaSoundGraph, *NewGraphNode);
 								}
+								return EdGraphNode;
 							}
 						}
 					}
 				}
+			}
+
+			return nullptr;
+		}
+
+		bool WillAddingVariableAccessorCauseLoop(const Frontend::IVariableController& InVariable, const Frontend::IInputController& InInput)
+		{
+			using namespace Metasound::Frontend;
+
+			// A variable mutator node must come before a variable accessor node,
+			// or else the nodes will create a loop from the hidden variable pin. 
+			// To determine if adding an accessor node will cause a loop (before actually
+			// adding an accessor node), we check whether an existing mutator can
+			// reach the node upstream which wants to connect to the accessor node.
+			//
+			// Example:
+			// Will cause loop:
+			// 	[AccessorNode]-->[DestinationNode]-->[Node]-->[MutatorNode] 
+			// 	       ^-------------------------------------------|
+			//
+			// Will not cause loop
+			//  [Node]-->[MutatorNode]-->[AccessorNode]-->[DestinationNode]
+			//       |                                        ^ 
+			//       |----------------------------------------|
+			FConstNodeHandle MutatorNode = InVariable.FindMutatorNode();
+			FConstNodeHandle DestinationNode = InInput.GetOwningNode();
+			return FGraphLinter::IsReachableUpstream(*MutatorNode, *DestinationNode);
+		}
+
+		bool WillAddingVariableMutatorCauseLoop(const Frontend::IVariableController& InVariable, const Frontend::IOutputController& InOutput)
+		{
+			using namespace Metasound::Frontend;
+
+			// A variable mutator node must come before a variable accessor node,
+			// or else the nodes will create a loop from the hidden variable pin. 
+			// To determine if adding a mutator node will cause a loop (before actually
+			// adding a mutator node), we check whether any existing accessor can
+			// reach the node downstream which wants to connect to the mutator node.
+			//
+			// Example:
+			// Will cause loop:
+			// 	[AccessorNode]-->[Node]-->[SourceNode]-->[MutatorNode] 
+			// 	     ^---------------------------------------|
+			//
+			// Will not cause loop
+			//  [SourceNode]-->[MutatorNode]-->[AccessorNode]-->[Node]
+			//       |                                            ^ 
+			//       |--------------------------------------------|
+			TArray<FConstNodeHandle> AccessorNodes = InVariable.FindAccessorNodes();
+			FConstNodeHandle SourceNode = InOutput.GetOwningNode();
+
+			auto IsSourceNodeReachableDownstream = [&SourceNode](const FConstNodeHandle& AccessorNode)
+			{
+				return FGraphLinter::IsReachableDownstream(*AccessorNode, *SourceNode);
+			};
+
+			return Algo::AnyOf(AccessorNodes, IsSourceNodeReachableDownstream);
+		}
 
-				if (FromPin->Direction == EGPD_Output)
+		class FAssetSchemaQueryResult : public ISchemaQueryResult
+		{
+		public:
+			FAssetSchemaQueryResult(Frontend::FAssetClassQueryInfo InTagData)
+				: QueryInfo(MoveTemp(InTagData))
+			{
+				bIsLoaded = FSoftObjectPath(QueryInfo.AssetPath).ResolveObject() != nullptr;
+			}
+
+			virtual ~FAssetSchemaQueryResult() = default;
+
+			const FMetasoundFrontendGraphClass* FindGraphClass() const
+			{
+				using namespace Metasound;
+
+				if (UObject* MetaSound = FSoftObjectPath(QueryInfo.AssetPath).TryLoad())
 				{
-					FConstOutputHandle OutputHandle = FGraphBuilder::GetConstOutputHandleFromPin(FromPin);
-					for (UEdGraphPin* Pin : NewGraphNode.Pins)
+					if (const FMetasoundAssetBase* MetaSoundAsset = IMetasoundUObjectRegistry::Get().GetObjectAsAssetBase(MetaSound))
 					{
-						if (Pin->Direction == EGPD_Input)
+						bIsLoaded = true;
+						if (!MetaSoundAsset->IsRegistered())
 						{
-							FConstInputHandle InputHandle = FGraphBuilder::GetConstInputHandleFromPin(Pin);
-							if (InputHandle->IsValid() && 
-								InputHandle->CanConnectTo(*OutputHandle).Connectable == FConnectability::EConnectable::Yes)
-							{
-								if (ensure(FGraphBuilder::ConnectNodes(*Pin, *FromPin, true /* bConnectEdPins */)))
-								{
-									return true;
-								}
-							}
+							FGraphBuilder::RegisterGraphWithFrontend(*MetaSound);
 						}
+						return &MetaSoundAsset->GetConstDocumentChecked().RootGraph;
 					}
 				}
 
-				return false;
+				bIsLoaded = false;
+				return nullptr;
 			}
 
-			struct FDataTypeActionQuery
+			virtual const FMetasoundFrontendClass* FindClass() const override
 			{
-				FGraphActionMenuBuilder& ActionMenuBuilder;
-				const TArray<Frontend::FConstNodeHandle>& NodeHandles;
-				FInterfaceNodeFilterFunction Filter;
-				EPrimaryContextGroup ContextGroup;
-				const FText& DisplayNameFormat;
-				const FText& TooltipFormat;
-				bool bShowSelectedActions = false;
-			};
+				return FindGraphClass();
+			}
 
-			template <typename TAction>
-			void GetDataTypeActions(const FDataTypeActionQuery& InQuery)
+			virtual bool CanConnectInputOfTypeAndAccess(FName InputTypeName, EMetasoundFrontendVertexAccessType InputAccessType) const override
 			{
-				using namespace Editor;
 				using namespace Frontend;
 
-				for (const FConstNodeHandle& NodeHandle : InQuery.NodeHandles)
+				const bool bCanConnect = Algo::AnyOf(QueryInfo.Outputs.ClassVertexInfo, [&InputTypeName, &InputAccessType](const FClassQueryVertexInfo& Output)
+				{
+					return InputTypeName == Output.TypeName && FMetasoundFrontendClassVertex::CanConnectVertexAccessTypes(Output.AccessType, InputAccessType);
+				});
+				return bCanConnect;
+			}
+
+			virtual bool CanConnectOutputOfTypeAndAccess(FName OutputTypeName, EMetasoundFrontendVertexAccessType OutputAccessType) const override
+			{
+				using namespace Engine;
+				const bool bCanConnect = Algo::AnyOf(QueryInfo.Inputs.ClassVertexInfo, [&OutputTypeName, &OutputAccessType](const FClassQueryVertexInfo& Input)
 				{
-					if (!InQuery.Filter || InQuery.Filter(NodeHandle))
+					return OutputTypeName == Input.TypeName && FMetasoundFrontendClassVertex::CanConnectVertexAccessTypes(OutputAccessType, Input.AccessType);
+				});
+				return bCanConnect;
+			}
+
+			virtual Frontend::FAssetKey GetAssetKey() const override
+			{
+				return Frontend::FAssetKey(QueryInfo.ClassName, QueryInfo.Version);
+			}
+
+			virtual const TArray<FText>& GetCategoryHierarchy() const override
+			{
+				if (bIsLoaded)
+				{
+					if (const FMetasoundFrontendClass* Class = FindClass())
 					{
-						constexpr bool bIncludeNamespace = true;
-
-						const FText& GroupName = GetContextGroupDisplayName(InQuery.ContextGroup);
-						const FText NodeDisplayName = FGraphBuilder::GetDisplayName(*NodeHandle, bIncludeNamespace);
-						const FText Tooltip = FText::Format(InQuery.TooltipFormat, NodeDisplayName);
-						const FText DisplayName = FText::Format(InQuery.DisplayNameFormat, NodeDisplayName);
-						TSharedPtr<TAction> NewNodeAction = MakeShared<TAction>(GroupName, DisplayName, NodeHandle->GetID(), Tooltip, InQuery.ContextGroup);
-						InQuery.ActionMenuBuilder.AddAction(NewNodeAction);
+						return Class->Metadata.GetCategoryHierarchy();
 					}
 				}
+
+				return QueryInfo.DisplayInfo.Hierarchy;
 			}
 
-			void SelectNodeInEditor(UMetasoundEditorGraph& InMetaSoundGraph, UMetasoundEditorGraphNode& InNode)
+			virtual FText GetDisplayName() const override
 			{
-				TSharedPtr<FEditor> MetasoundEditor = FGraphBuilder::GetEditorForGraph(InMetaSoundGraph);
-				if (MetasoundEditor.IsValid())
+				using namespace Frontend;
+
+				FText DisplayName = FText::FromName(QueryInfo.AssetPath.GetAssetName());
+
+				if (bIsLoaded)
+				{
+					if (const FMetasoundFrontendGraphClass* GraphClass = FindGraphClass())
+					{
+						FText DisplayNameOverride = GraphClass->Metadata.GetDisplayName();
+						if (!DisplayNameOverride.IsEmptyOrWhitespace())
+						{
+							DisplayName = DisplayNameOverride;
+						}
+					}
+				}
+				else if (!QueryInfo.DisplayInfo.ClassDisplayName.IsEmptyOrWhitespace())
+				{
+					DisplayName = QueryInfo.DisplayInfo.ClassDisplayName;
+				}
+
+				if (ShowUnloadedAssetInBrowserCvar)
 				{
-					MetasoundEditor->ClearSelectionAndSelectNode(&InNode);
+					return FText::Format(LOCTEXT("FileNodeLoadedNameHintFormat", "{0}*"), DisplayName);
 				}
+
+				return DisplayName;
 			}
 
-			void SelectNodeInEditorForRename(UMetasoundEditorGraph& InMetaSoundGraph, UMetasoundEditorGraphNode& InNode)
+			virtual const TArray<FText>& GetKeywords() const override
 			{
-				TSharedPtr<FEditor> MetasoundEditor = FGraphBuilder::GetEditorForGraph(InMetaSoundGraph);
-				if (MetasoundEditor.IsValid())
+				if (bIsLoaded)
 				{
-					MetasoundEditor->ClearSelectionAndSelectNode(&InNode);
-					MetasoundEditor->SetDelayedRename();
+					if (const FMetasoundFrontendClass* Class = FindClass())
+					{
+						return Class->Metadata.GetKeywords();
+					}
 				}
+
+				return QueryInfo.DisplayInfo.Keywords;
 			}
 
-			UEdGraphNode* PromoteToVariable(const FName BaseName, UEdGraphPin& FromPin, const FName DataType, const FMetasoundFrontendClass& InVariableClass, const FVector2D& InLocation, bool bSelectNode)
+			virtual EMetasoundFrontendClassType GetRegistryClassType() const override
 			{
-				using namespace Frontend;
+				return EMetasoundFrontendClassType::External;
+			}
 
-				UEdGraphNode* ConnectedNode = Cast<UEdGraphNode>(FromPin.GetOwningNode());
-				if (!ensure(ConnectedNode))
+			virtual FText GetTooltip() const override
+			{
+				if (bIsLoaded)
 				{
-					return nullptr;
+					if (const FMetasoundFrontendClass* Class = FindClass())
+					{
+						const FText Author = FText::FromString(Class->Metadata.GetAuthor());
+						const FText& Description = Class->Metadata.GetDescription();
+						if (Description.IsEmptyOrWhitespace())
+						{
+							return Author;
+						}
+						return Author.IsEmptyOrWhitespace() ? Description : FText::Format(ClassDescriptionAndAuthorFormat, Description, Author);
+					}
 				}
 
-				const FMetasoundFrontendClassName& ClassName = InVariableClass.Metadata.GetClassName();
+				return QueryInfo.DisplayInfo.ClassDescription;
+			}
+
+			virtual bool IsNative() const override
+			{
+				return false;
+			}
 
-				UMetasoundEditorGraph* MetaSoundGraph = CastChecked<UMetasoundEditorGraph>(ConnectedNode->GetGraph());
-				FMetaSoundFrontendDocumentBuilder& DocBuilder = MetaSoundGraph->GetBuilderChecked().GetBuilder();
-				const FName NodeName = FGraphBuilder::GenerateUniqueVariableName(DocBuilder, BaseName.ToString());
+		private:
+			Frontend::FAssetClassQueryInfo QueryInfo;
+			mutable bool bIsLoaded = false;
+		};
+
+		class FRegistrySchemaQueryResult : public ISchemaQueryResult
+		{
+			public:
+			FRegistrySchemaQueryResult(const Engine::FMetaSoundAssetManager& AssetManager, FMetasoundFrontendClass InClass)
+				: Class(MoveTemp(InClass))
+			{
+				bIsNative = !AssetManager.IsAssetClass(Class.Metadata);
+			}
 
-				const FScopedTransaction Transaction(FText::Format(
-					LOCTEXT("PromoteNodeVertexToGraphVariableFormat", "Promote MetaSound Node {0} to {1}"),
-					FText::FromName(NodeName),
-					FText::FromName(ClassName.Namespace)));
+			virtual ~FRegistrySchemaQueryResult() = default;
 
-				UObject& ParentMetaSound = MetaSoundGraph->GetMetasoundChecked();
-				ParentMetaSound.Modify();
-				MetaSoundGraph->Modify();
+			virtual const FMetasoundFrontendClass* FindClass() const override
+			{
+				return &Class;
+			}
 
-				// Cache the default literal from the pin if connecting to an input
-				FMetasoundFrontendLiteral DefaultValue;
-				if (FromPin.Direction == EGPD_Input)
+			virtual bool CanConnectInputOfTypeAndAccess(FName InputTypeName, EMetasoundFrontendVertexAccessType InputAccessType) const override
+			{
+				const bool bCanConnect = Algo::AnyOf(Class.Interface.Outputs, [&InputTypeName, &InputAccessType](const FMetasoundFrontendClassOutput& Output)
 				{
-					FGraphBuilder::GetPinLiteral(FromPin, DefaultValue);
-				}
+					return InputTypeName == Output.TypeName && FMetasoundFrontendClassVertex::CanConnectVertexAccessTypes(Output.AccessType, InputAccessType);
+				});
+				return bCanConnect;
+			}
 
-				const FMetasoundFrontendVariable* FrontendVariable = DocBuilder.AddGraphVariable(NodeName, DataType);
-				if (ensure(FrontendVariable))
+			virtual bool CanConnectOutputOfTypeAndAccess(FName OutputTypeName, EMetasoundFrontendVertexAccessType OutputAccessType) const override
+			{
+				const bool bCanConnect = Algo::AnyOf(Class.Interface.Inputs, [&OutputTypeName, &OutputAccessType](const FMetasoundFrontendClassInput& Input)
 				{
-					UMetasoundEditorGraphVariable* Variable = MetaSoundGraph->FindOrAddVariable(FrontendVariable->Name);
-					if (ensure(Variable))
+					return OutputTypeName == Input.TypeName && FMetasoundFrontendClassVertex::CanConnectVertexAccessTypes(OutputAccessType, Input.AccessType);
+				});
+				return bCanConnect;
+			}
+
+			virtual Frontend::FAssetKey GetAssetKey() const override
+			{
+				return Frontend::FAssetKey(Class.Metadata);
+			}
+
+			virtual const TArray<FText>& GetCategoryHierarchy() const override
+			{
+				return Class.Metadata.GetCategoryHierarchy();
+			}
+
+			virtual FText GetDisplayName() const override
+			{
+				using namespace Frontend;
+
+				FText DisplayName;
+				auto GetAssetDisplayNameFromMetadata = [&DisplayName](const FMetasoundFrontendClassMetadata& Metadata)
+				{
+					DisplayName = Metadata.GetDisplayName();
+					if (DisplayName.IsEmptyOrWhitespace())
 					{
-						const FMetasoundFrontendNode* NewVariableNode = DocBuilder.AddGraphVariableNode(FrontendVariable->Name, InVariableClass.Metadata.GetType());
-						if (ensure(NewVariableNode))
+						const FTopLevelAssetPath Path = IMetaSoundAssetManager::GetChecked().FindAssetPath(FAssetKey(Metadata));
+						if (Path.IsValid())
 						{
-							UMetasoundEditorGraphVariableNode* NewGraphNode = FGraphBuilder::AddVariableNode(ParentMetaSound, NewVariableNode->GetID());
-							if (ensure(NewGraphNode))
-							{
-								NewGraphNode->UpdateFrontendNodeLocation(InLocation);
-								NewGraphNode->SyncLocationFromFrontendNode();
+							DisplayName = FText::FromName(Path.GetAssetName());
+						}
+					}
+				};
 
-								// Set the literal using the new value if connecting to an input
-								if (FromPin.Direction == EGPD_Input)
-								{
-									UMetasoundEditorGraphMemberDefaultLiteral* DefaultLiteral = Variable->GetLiteral();
-									if (ensure(DefaultLiteral))
-									{
-										DefaultLiteral->SetFromLiteral(DefaultValue);
-									}
-
-									// Ensures the setter node value is synced with the editor literal value
-									constexpr bool bPostTransaction = false;
-									Variable->UpdateFrontendDefaultLiteral(bPostTransaction);
-								}
+				// 1. Try to get display name from metadata or asset if one can be found from the asset manager
+				GetAssetDisplayNameFromMetadata(Class.Metadata);
 
-								UEdGraphNode* EdGraphNode = CastChecked<UEdGraphNode>(NewGraphNode);
-								if (ensure(SchemaPrivate::TryConnectNewNodeToMatchingDataTypePin(*EdGraphNode, &FromPin)))
-								{
-									FGraphBuilder::RegisterGraphWithFrontend(ParentMetaSound);
-									if (bSelectNode)
-									{
-										SchemaPrivate::SelectNodeInEditorForRename(*MetaSoundGraph, *NewGraphNode);
-									}
-									else
-									{
-										SchemaPrivate::SelectNodeInEditor(*MetaSoundGraph, *NewGraphNode);
-									}
-									return EdGraphNode;
-								}
-							}
-						}
+				// 2. If version is missing from the registry or from asset system, then this node
+				// will not provide a useful DisplayName.  In that case, attempt to find the next highest
+				// class & associated DisplayName.
+				if (DisplayName.IsEmptyOrWhitespace())
+				{
+					FMetasoundFrontendClass ClassWithHighestVersion;
+					if (ISearchEngine::Get().FindClassWithHighestVersion(Class.Metadata.GetClassName(), ClassWithHighestVersion))
+					{
+						GetAssetDisplayNameFromMetadata(ClassWithHighestVersion.Metadata);
 					}
 				}
 
-				return nullptr;
+				return DisplayName;
 			}
 
-			bool WillAddingVariableAccessorCauseLoop(const Frontend::IVariableController& InVariable, const Frontend::IInputController& InInput)
+			virtual const TArray<FText>& GetKeywords() const override
 			{
-				using namespace Metasound::Frontend;
+				return Class.Metadata.GetKeywords();
+			}
 
-				// A variable mutator node must come before a variable accessor node,
-				// or else the nodes will create a loop from the hidden variable pin. 
-				// To determine if adding an accessor node will cause a loop (before actually
-				// adding an accessor node), we check whether an existing mutator can
-				// reach the node upstream which wants to connect to the accessor node.
-				//
-				// Example:
-				// Will cause loop:
-				// 	[AccessorNode]-->[DestinationNode]-->[Node]-->[MutatorNode] 
-				// 	       ^-------------------------------------------|
-				//
-				// Will not cause loop
-				//  [Node]-->[MutatorNode]-->[AccessorNode]-->[DestinationNode]
-				//       |                                        ^ 
-				//       |----------------------------------------|
-				FConstNodeHandle MutatorNode = InVariable.FindMutatorNode();
-				FConstNodeHandle DestinationNode = InInput.GetOwningNode();
-				return FGraphLinter::IsReachableUpstream(*MutatorNode, *DestinationNode);
-			}
-
-			bool WillAddingVariableMutatorCauseLoop(const Frontend::IVariableController& InVariable, const Frontend::IOutputController& InOutput)
+			virtual FText GetTooltip() const
 			{
-				using namespace Metasound::Frontend;
+				const FString& Author = Class.Metadata.GetAuthor();
+				return Author.IsEmpty()
+					? FText::Format(ClassDescriptionAndAuthorFormat, Class.Metadata.GetDescription(), FText::FromString(Author))
+					: Class.Metadata.GetDescription();
+			}
 
-				// A variable mutator node must come before a variable accessor node,
-				// or else the nodes will create a loop from the hidden variable pin. 
-				// To determine if adding a mutator node will cause a loop (before actually
-				// adding a mutator node), we check whether any existing accessor can
-				// reach the node downstream which wants to connect to the mutator node.
-				//
-				// Example:
-				// Will cause loop:
-				// 	[AccessorNode]-->[Node]-->[SourceNode]-->[MutatorNode] 
-				// 	     ^---------------------------------------|
-				//
-				// Will not cause loop
-				//  [SourceNode]-->[MutatorNode]-->[AccessorNode]-->[Node]
-				//       |                                            ^ 
-				//       |--------------------------------------------|
-				TArray<FConstNodeHandle> AccessorNodes = InVariable.FindAccessorNodes();
-				FConstNodeHandle SourceNode = InOutput.GetOwningNode();
-
-				auto IsSourceNodeReachableDownstream = [&SourceNode](const FConstNodeHandle& AccessorNode)
-				{
-					return FGraphLinter::IsReachableDownstream(*AccessorNode, *SourceNode);
-				};
+			virtual EMetasoundFrontendClassType GetRegistryClassType() const override
+			{
+				return Class.Metadata.GetType();
+			}
 
-				return Algo::AnyOf(AccessorNodes, IsSourceNodeReachableDownstream);
+			virtual bool IsNative() const
+			{
+				return bIsNative;
 			}
-		} // namespace SchemaPrivate
-	} // namespace Editor
+
+			private:
+			FMetasoundFrontendClass Class;
+			bool bIsNative = false;
+		};
+	} // namespace SchemaPrivate
 
 	namespace SchemaUtils
 	{
@@ -651,13 +922,13 @@ PRAGMA_ENABLE_DEPRECATION_WARNINGS
 			return nullptr;
 		}
 	} // namespace SchemaUtils
-} // namespace Metasound
+} // namespace Metasound::Editor
 
 UEdGraphNode* FMetasoundGraphSchemaAction_NodeWithMultipleOutputs::PerformAction(class UEdGraph* ParentGraph, TArray<UEdGraphPin*>& FromPins, const FVector2D Location, bool bSelectNewNode/* = true*/)
 {
 	using namespace Metasound::Editor;
 	
-	UEdGraphNode* ResultNode = NULL;
+	UEdGraphNode* ResultNode = nullptr;
 
 	if (FromPins.Num() > 0)
 	{
@@ -674,9 +945,10 @@ UEdGraphNode* FMetasoundGraphSchemaAction_NodeWithMultipleOutputs::PerformAction
 	}
 	else
 	{
-		ResultNode = PerformAction(ParentGraph, NULL, Location, bSelectNewNode);
+		ResultNode = PerformAction(ParentGraph, nullptr, Location, bSelectNewNode);
 	}
 
+	bSelectNewNode &= ResultNode != nullptr;
 	TSharedPtr<FEditor> MetasoundEditor = FGraphBuilder::GetEditorForGraph(*ParentGraph);
 	if (MetasoundEditor.IsValid() && bSelectNewNode)
 	{
@@ -690,8 +962,7 @@ const FSlateBrush* FMetasoundGraphSchemaAction_NewNode::GetIconBrush() const
 {
 	using namespace Metasound::Frontend;
 
-	const FNodeRegistryKey RegistryKey = FNodeRegistryKey(ClassMetadata);
-	const bool bIsClassNative = FMetasoundFrontendRegistryContainer::Get()->IsNodeNative(RegistryKey);
+	const bool bIsClassNative = QueryResult.IsValid() && QueryResult->IsNative();
 	if (bIsClassNative)
 	{
 		return &Metasound::Editor::Style::GetSlateBrushSafe("MetasoundEditor.Graph.Node.Class.Native");
@@ -706,9 +977,7 @@ const FLinearColor& FMetasoundGraphSchemaAction_NewNode::GetIconColor() const
 
 	if (const UMetasoundEditorSettings* EditorSettings = GetDefault<UMetasoundEditorSettings>())
 	{
-		const FNodeRegistryKey RegistryKey = FNodeRegistryKey(ClassMetadata);
-		const bool bIsClassNative = FMetasoundFrontendRegistryContainer::Get()->IsNodeNative(RegistryKey);
-		if (bIsClassNative)
+		if (QueryResult.IsValid() && QueryResult->IsNative())
 		{
 			return EditorSettings->NativeNodeTitleColor;
 		}
@@ -724,13 +993,21 @@ UEdGraphNode* FMetasoundGraphSchemaAction_NewNode::PerformAction(UEdGraph* Paren
 	using namespace Metasound::Editor;
 	using namespace Metasound::Frontend;
 
+	const FMetasoundFrontendClass* Class = QueryResult->FindClass();
+	if (!Class)
+	{
+		return nullptr;
+	}
+
 	const FScopedTransaction Transaction(LOCTEXT("AddNewNode", "Add New MetaSound Node"));
 	UMetasoundEditorGraph* MetaSoundGraph = CastChecked<UMetasoundEditorGraph>(ParentGraph);
 	UObject& ParentMetasound = MetaSoundGraph->GetMetasoundChecked();
 	ParentMetasound.Modify();
 	ParentGraph->Modify();
 
-	if (UMetasoundEditorGraphExternalNode* NewGraphNode = FGraphBuilder::AddExternalNode(ParentMetasound, ClassMetadata, bSelectNewNode))
+	FMetasoundFrontendClassMetadata Metadata = Class->Metadata;
+	Metadata.SetType(EMetasoundFrontendClassType::External);
+	if (UMetasoundEditorGraphExternalNode* NewGraphNode = FGraphBuilder::AddExternalNode(ParentMetasound, Metadata, bSelectNewNode))
 	{
 		NewGraphNode->Modify();
 		NewGraphNode->UpdateFrontendNodeLocation(Location);
@@ -819,7 +1096,7 @@ FMetasoundGraphSchemaAction_PromoteToInput::FMetasoundGraphSchemaAction_PromoteT
 
 UEdGraphNode* FMetasoundGraphSchemaAction_PromoteToInput::PerformAction(UEdGraph* ParentGraph, UEdGraphPin* FromPin, const FVector2D InLocation, bool bSelectNewNode /* = true */)
 {
-	using namespace Metasound;
+	using namespace Metasound::Editor;
 	return SchemaUtils::PromoteToInput(ParentGraph, FromPin, InLocation, bSelectNewNode);
 }
 
@@ -834,7 +1111,7 @@ FMetasoundGraphSchemaAction_PromoteToVariable_AccessorNode::FMetasoundGraphSchem
 
 UEdGraphNode* FMetasoundGraphSchemaAction_PromoteToVariable_AccessorNode::PerformAction(UEdGraph* ParentGraph, UEdGraphPin* FromPin, const FVector2D InLocation, bool bSelectNewNode /* = true */)
 {
-	using namespace Metasound;
+	using namespace Metasound::Editor;
 	return SchemaUtils::PromoteToVariable(ParentGraph, FromPin, InLocation, bSelectNewNode);
 }
 
@@ -849,7 +1126,7 @@ FMetasoundGraphSchemaAction_PromoteToVariable_DeferredAccessorNode::FMetasoundGr
 
 UEdGraphNode* FMetasoundGraphSchemaAction_PromoteToVariable_DeferredAccessorNode::PerformAction(UEdGraph* ParentGraph, UEdGraphPin* FromPin, const FVector2D InLocation, bool bSelectNewNode /* = true */)
 {
-	using namespace Metasound;
+	using namespace Metasound::Editor;
 	return SchemaUtils::PromoteToDeferredVariable(ParentGraph, FromPin, InLocation, bSelectNewNode);
 }
 
@@ -864,7 +1141,7 @@ FMetasoundGraphSchemaAction_PromoteToVariable_MutatorNode::FMetasoundGraphSchema
 
 UEdGraphNode* FMetasoundGraphSchemaAction_PromoteToVariable_MutatorNode::PerformAction(UEdGraph* ParentGraph, UEdGraphPin* FromPin, const FVector2D InLocation, bool bSelectNewNode /* = true */)
 {
-	using namespace Metasound;
+	using namespace Metasound::Editor;
 	return SchemaUtils::PromoteToMutatorVariable(ParentGraph, FromPin, InLocation, bSelectNewNode);
 }
 
@@ -930,7 +1207,7 @@ FMetasoundGraphSchemaAction_PromoteToOutput::FMetasoundGraphSchemaAction_Promote
 
 UEdGraphNode* FMetasoundGraphSchemaAction_PromoteToOutput::PerformAction(UEdGraph* ParentGraph, UEdGraphPin* FromPin, const FVector2D InLocation, bool bSelectNewNode /* = true */)
 {
-	using namespace Metasound;
+	using namespace Metasound::Editor;
 	return SchemaUtils::PromoteToOutput(ParentGraph, FromPin, InLocation, bSelectNewNode);
 }
 
@@ -1324,7 +1601,7 @@ void UMetasoundEditorGraphSchema::GetGraphContextActions(FGraphContextMenuBuilde
 	using namespace Metasound::Editor;
 	using namespace Metasound::Frontend;
 
-	FActionClassFilters ClassFilters;
+	FActionVertexFilters VertexFilters;
 	FConstGraphHandle GraphHandle = IGraphController::GetInvalidHandle();
 	EMetasoundFrontendVertexAccessType OutputAccessType = EMetasoundFrontendVertexAccessType::Unset;
 	if (const UEdGraphPin* FromPin = ContextMenuBuilder.FromPin)
@@ -1334,15 +1611,12 @@ void UMetasoundEditorGraphSchema::GetGraphContextActions(FGraphContextMenuBuilde
 			FConstInputHandle InputHandle = FGraphBuilder::GetConstInputHandleFromPin(ContextMenuBuilder.FromPin);
 			OutputAccessType = InputHandle->GetVertexAccessType();
 
-			ClassFilters.OutputFilterFunction = [InputHandle](const FMetasoundFrontendClassOutput& InOutput)
-			{
-				return InOutput.TypeName == InputHandle->GetDataType() && 
-					FMetasoundFrontendClassVertex::CanConnectVertexAccessTypes(InOutput.AccessType, InputHandle->GetVertexAccessType());
-			};
+			VertexFilters.OutputTypeName = InputHandle->GetDataType();
+			VertexFilters.OutputAccessType = InputHandle->GetVertexAccessType();
 
 			// Show only input nodes as output nodes can only connected if FromPin is input
 			GraphHandle = InputHandle->GetOwningNode()->GetOwningGraph();
-			GetDataTypeInputNodeActions(ContextMenuBuilder, GraphHandle, [InputHandle](FConstNodeHandle NodeHandle)
+			GetDataTypeInputNodeActions(ContextMenuBuilder, GraphHandle, [&InputHandle](FConstNodeHandle NodeHandle)
 			{
 				bool bHasConnectableOutput = false;
 				NodeHandle->IterateConstOutputs([&](FConstOutputHandle PotentialOutputHandle)
@@ -1369,11 +1643,8 @@ void UMetasoundEditorGraphSchema::GetGraphContextActions(FGraphContextMenuBuilde
 		if (FromPin->Direction == EGPD_Output)
 		{
 			FConstOutputHandle OutputHandle = FGraphBuilder::GetConstOutputHandleFromPin(FromPin);
-			ClassFilters.InputFilterFunction = [OutputHandle](const FMetasoundFrontendClassInput& InInput)
-			{
-				return InInput.TypeName == OutputHandle->GetDataType() &&
-					FMetasoundFrontendClassVertex::CanConnectVertexAccessTypes(OutputHandle->GetVertexAccessType(), InInput.AccessType);
-			};
+			VertexFilters.InputTypeName = OutputHandle->GetDataType();
+			VertexFilters.InputAccessType = OutputHandle->GetVertexAccessType();
 
 			// Show only output nodes as input nodes can only connected if FromPin is output
 			GraphHandle = OutputHandle->GetOwningNode()->GetOwningGraph();
@@ -1422,13 +1693,13 @@ void UMetasoundEditorGraphSchema::GetGraphContextActions(FGraphContextMenuBuilde
 		}
 	}
 
-	GetFunctionActions(ContextMenuBuilder, ClassFilters, true /* bShowSelectedActions */, GraphHandle);
+	GetFunctionActions(ContextMenuBuilder, VertexFilters, true /* bShowSelectedActions */, GraphHandle);
 
 	// Variable and conversion actions are always by reference so are incompatible with constructor outputs 
 	if (OutputAccessType != EMetasoundFrontendVertexAccessType::Value)
 	{
-		GetVariableActions(ContextMenuBuilder, ClassFilters, true /* bShowSelectedActions */, GraphHandle);
-		GetConversionActions(ContextMenuBuilder, ClassFilters);
+		GetVariableActions(ContextMenuBuilder, VertexFilters, true /* bShowSelectedActions */, GraphHandle);
+		GetConversionActions(ContextMenuBuilder, VertexFilters);
 	}
 }
 
@@ -1463,15 +1734,13 @@ void UMetasoundEditorGraphSchema::GetContextMenuActions(class UToolMenu* Menu, c
 			const FMetasoundFrontendClassMetadata& Metadata = NodeHandle->GetClassMetadata();
 			const bool bHasNewVersion = HighestVersion.IsValid() && HighestVersion > Metadata.GetVersion();
 
-			const FNodeRegistryKey RegistryKey = FNodeRegistryKey(Metadata);
-			const bool bIsClassNative = FMetasoundFrontendRegistryContainer::Get()->IsNodeNative(RegistryKey);
-
-			if (bHasNewVersion || !bIsClassNative)
+			const bool bIsAssetClass = IMetaSoundAssetManager::GetChecked().IsAssetClass(Metadata);
+			if (bHasNewVersion || bIsAssetClass)
 			{
 				Section.AddMenuEntry(FEditorCommands::Get().UpdateNodeClass);
 			}
 
-			FMetasoundFrontendNodeStyle Style = NodeHandle->GetNodeStyle();		
+			FMetasoundFrontendNodeStyle Style = NodeHandle->GetNodeStyle();
 			if (Style.bUnconnectedPinsHidden)
 			{
 				OrganizationSection.AddMenuEntry(FGraphEditorCommands::Get().ShowAllPins, LOCTEXT("ShowUnconnectedPins", "Show Unconnected Pins"),
@@ -2278,49 +2547,52 @@ void UMetasoundEditorGraphSchema::DroppedAssetsOnPin(const TArray<FAssetData>& A
 	}
 }
 
-void UMetasoundEditorGraphSchema::GetConversionActions(FGraphActionMenuBuilder& ActionMenuBuilder, Metasound::Editor::FActionClassFilters InFilters, bool bShowSelectedActions) const
+void UMetasoundEditorGraphSchema::GetConversionActions(FGraphActionMenuBuilder& ActionMenuBuilder, Metasound::Editor::FActionVertexFilters AccessFilters, bool bShowSelectedActions) const
 {
 	using namespace Metasound;
 	using namespace Metasound::Editor;
 
-	const bool bIncludeAllVersions = static_cast<bool>(EnableAllVersionsMetaSoundNodeClassCreationCVar);
-	const TArray<FMetasoundFrontendClass> FrontendClasses = Frontend::ISearchEngine::Get().FindAllClasses(bIncludeAllVersions);
-	for (const FMetasoundFrontendClass& FrontendClass : FrontendClasses)
+	constexpr bool bScanAssetTags = false; // All conversion actions are natively defined, so no reason to list as conversion
+	QueryNodeClasses([&](TUniquePtr<ISchemaQueryResult>&& Result)
 	{
-		if (InFilters.InputFilterFunction && !FrontendClass.Interface.Inputs.ContainsByPredicate(InFilters.InputFilterFunction))
+		if (Result->GetRegistryClassType() != EMetasoundFrontendClassType::External)
 		{
-			continue;
+			return;
 		}
 
-		if (InFilters.OutputFilterFunction && !FrontendClass.Interface.Outputs.ContainsByPredicate(InFilters.OutputFilterFunction))
+		if (AccessFilters.HasInputFilters())
 		{
-			continue;
+			if (!Result->CanConnectInputOfTypeAndAccess(AccessFilters.InputTypeName, AccessFilters.InputAccessType))
+			{
+				return;
+			}
 		}
 
-		FMetasoundFrontendClassMetadata Metadata = FrontendClass.Metadata;
-		const FText Tooltip = Metadata.GetAuthor().IsEmpty()
-			? Metadata.GetDescription()
-			: FText::Format(LOCTEXT("MetasoundTooltipAuthorFormat1", "{0}\nAuthor: {1}"), Metadata.GetDescription(), FText::FromString(Metadata.GetAuthor()));
-		if (!Metadata.GetCategoryHierarchy().IsEmpty() && !Metadata.GetCategoryHierarchy()[0].CompareTo(NodeCategories::Conversions))
+		if (AccessFilters.HasOutputFilters())
 		{
-			FText KeywordsText = FText::Join(SchemaPrivate::KeywordDelim, Metadata.GetKeywords());
-			const FText CategoryText = FText::Join(SchemaPrivate::CategoryDelim, Metadata.GetCategoryHierarchy());
+			if (!Result->CanConnectOutputOfTypeAndAccess(AccessFilters.OutputTypeName, AccessFilters.OutputAccessType))
+			{
+				return;
+			}
+		}
 
-			constexpr bool bShowNamespace = false; // TODO: Make this an option to display
+		const TArray<FText>& CategoryHierarchy = Result->GetCategoryHierarchy();
+		if (!CategoryHierarchy.IsEmpty() && !CategoryHierarchy[0].CompareTo(NodeCategories::Conversions))
+		{
+			const FText Tooltip = Result->GetTooltip();
 			TSharedPtr<FMetasoundGraphSchemaAction_NewNode> NewNodeAction = MakeShared<FMetasoundGraphSchemaAction_NewNode>
 			(
-				CategoryText,
-				FGraphBuilder::GetDisplayName(Metadata, FName(), bShowNamespace),
+				FText::Join(SchemaPrivate::CategoryDelim, CategoryHierarchy),
+				Result->GetDisplayName(),
 				Tooltip,
 				EPrimaryContextGroup::Conversions,
-				KeywordsText
+				FText::Join(SchemaPrivate::KeywordDelim, Result->GetKeywords())
 			);
 
-			Metadata.SetType(EMetasoundFrontendClassType::External);
-			NewNodeAction->ClassMetadata = MoveTemp(Metadata);
+			NewNodeAction->QueryResult = TSharedPtr<ISchemaQueryResult>(Result.Release());
 			ActionMenuBuilder.AddAction(NewNodeAction);
 		}
-	}
+	}, bScanAssetTags);
 }
 
 void UMetasoundEditorGraphSchema::GetDataTypeInputNodeActions(FGraphContextMenuBuilder& ActionMenuBuilder, Metasound::Frontend::FConstGraphHandle InGraphHandle, Metasound::Editor::FInterfaceNodeFilterFunction InFilter, bool bShowSelectedActions) const
@@ -2379,101 +2651,74 @@ void UMetasoundEditorGraphSchema::GetDataTypeOutputNodeActions(FGraphContextMenu
 	SchemaPrivate::GetDataTypeActions<FMetasoundGraphSchemaAction_NewOutput>(ActionQuery);
 }
 
-void UMetasoundEditorGraphSchema::GetFunctionActions(FGraphActionMenuBuilder& ActionMenuBuilder, Metasound::Editor::FActionClassFilters InFilters, bool bShowSelectedActions, Metasound::Frontend::FConstGraphHandle InGraphHandle) const
+void UMetasoundEditorGraphSchema::GetFunctionActions(FGraphActionMenuBuilder& ActionMenuBuilder, Metasound::Editor::FActionVertexFilters AccessFilters, bool bShowSelectedActions, Metasound::Frontend::FConstGraphHandle InGraphHandle) const
 {
 	using namespace Metasound::Editor;
-	using namespace Metasound::Editor::SchemaPrivate;
+	using namespace Metasound::Engine;
 	using namespace Metasound::Frontend;
 
-	const IMetaSoundAssetManager& AssetManager = IMetaSoundAssetManager::GetChecked();
-	FMetasoundAssetBase* ParentAsset = nullptr;
-	if (InGraphHandle->IsValid())
+	const FMetaSoundAssetManager& AssetManager = FMetaSoundAssetManager::GetChecked();
+	const FAssetKey ParentAssetKey(InGraphHandle->GetGraphMetadata());
+	QueryNodeClasses([&](TUniquePtr<ISchemaQueryResult>&& Result)
 	{
-		const FAssetKey AssetKey(InGraphHandle->GetGraphMetadata());
-		ParentAsset = AssetManager.TryLoadAssetFromKey(AssetKey);
-	}
-
-	const bool bIncludeAllVersions = static_cast<bool>(EnableAllVersionsMetaSoundNodeClassCreationCVar);
-	const TArray<FMetasoundFrontendClass> FrontendClasses = ISearchEngine::Get().FindAllClasses(bIncludeAllVersions);
-	for (const FMetasoundFrontendClass& FrontendClass : FrontendClasses)
-	{
-		if (InFilters.InputFilterFunction && !FrontendClass.Interface.Inputs.ContainsByPredicate(InFilters.InputFilterFunction))
+		if (Result->GetRegistryClassType() != EMetasoundFrontendClassType::External)
 		{
-			continue;
+			return;
 		}
 
-		if (InFilters.OutputFilterFunction && !FrontendClass.Interface.Outputs.ContainsByPredicate(InFilters.OutputFilterFunction))
+		if (AccessFilters.HasInputFilters())
 		{
-			continue;
+			if (!Result->CanConnectInputOfTypeAndAccess(AccessFilters.InputTypeName, AccessFilters.InputAccessType))
+			{
+				return;
+			}
 		}
 
-		if (FrontendClass.Metadata.GetType() == EMetasoundFrontendClassType::Template)
+		if (AccessFilters.HasOutputFilters())
 		{
-			continue;
+			if (!Result->CanConnectOutputOfTypeAndAccess(AccessFilters.OutputTypeName, AccessFilters.OutputAccessType))
+			{
+				return;
+			}
 		}
 
-		if (FrontendClass.Metadata.GetType() != EMetasoundFrontendClassType::External)
+		bool bCausesLoop = false;
+		auto IsMatchingKey = [&ParentAssetKey, &bCausesLoop](const FAssetKey& AssetKey)
 		{
-			continue;
-		}
-
-		bool bIsAsset = false;
-		if (ParentAsset)
+			bCausesLoop = ParentAssetKey == AssetKey;
+		};
+		AssetManager.IterateReferences(Result->GetAssetKey(), IsMatchingKey);
+		if (bCausesLoop)
 		{
-			const FNodeRegistryKey RegistryKey = FNodeRegistryKey(FrontendClass.Metadata);
-			if (FMetasoundAssetBase* NodeClassAsset = AssetManager.FindAsset(RegistryKey))
-			{
-				// The path can be not an asset if it was transiently created (eg. in PIE if generated by a runtime MetaSoundDocumentBuilder)
-				// Ignore these assets and do not include as this would create a dependency on a transient MetaSound from a MetaSound asset.
-				// (MetasoundAssetBase is a misnomer; it is not always an asset now that the builder API can generate transient MetaSounds
-				// at runtime).
-				UObject* OwningAsset = NodeClassAsset->GetOwningAsset();
-				check(OwningAsset);
-				if (!OwningAsset->IsAsset())
-				{
-					continue;
-				}
-
-				if (ParentAsset->AddingReferenceCausesLoop(*NodeClassAsset))
-				{
-					continue;
-				}
-
-				bIsAsset = true;
-			}
+			return;
 		}
 
-		FMetasoundFrontendClassMetadata Metadata = FrontendClass.Metadata;
-		const FText Tooltip = Metadata.GetAuthor().IsEmpty()
-			? Metadata.GetDescription()
-			: FText::Format(LOCTEXT("MetasoundTooltipAuthorFormat2", "{0}\nAuthor: {1}"), Metadata.GetDescription(), FText::FromString(Metadata.GetAuthor()));
-
-		if (Metadata.GetCategoryHierarchy().IsEmpty() || Metadata.GetCategoryHierarchy()[0].CompareTo(Metasound::NodeCategories::Conversions))
+		const TArray<FText>& CategoryHierarchy = Result->GetCategoryHierarchy();
+		if (!CategoryHierarchy.IsEmpty() && !CategoryHierarchy[0].CompareTo(Metasound::NodeCategories::Conversions))
 		{
-			const EPrimaryContextGroup ContextGroup = bIsAsset ? EPrimaryContextGroup::Graphs : EPrimaryContextGroup::Functions;
-			TArray<FText> CategoryHierarchy { GetContextGroupDisplayName(ContextGroup) };
-			CategoryHierarchy.Append(Metadata.GetCategoryHierarchy());
-			const FText KeywordsText = FText::Join(KeywordDelim, Metadata.GetKeywords());
-			const FText CategoryText = FText::Join(CategoryDelim, CategoryHierarchy);
+			return;
+		}
 
-			constexpr bool bShowNamespace = false; // TODO: Make this an option to display
-			TSharedPtr<FMetasoundGraphSchemaAction_NewNode> NewNodeAction = MakeShared<FMetasoundGraphSchemaAction_NewNode>
-			(
-				CategoryText,
-				FGraphBuilder::GetDisplayName(Metadata, FName(), bShowNamespace),
-				Tooltip,
-				ContextGroup,
-				KeywordsText
-			);
+		const FText Tooltip = Result->GetTooltip();
+		const EPrimaryContextGroup ContextGroup = Result->IsNative() ? EPrimaryContextGroup::Functions : EPrimaryContextGroup::Graphs;
+		TArray<FText> TextHierarchy { GetContextGroupDisplayName(ContextGroup) };
+		TextHierarchy.Append(CategoryHierarchy);
+
+		TSharedPtr<FMetasoundGraphSchemaAction_NewNode> NewNodeAction = MakeShared<FMetasoundGraphSchemaAction_NewNode>
+		(
+			FText::Join(SchemaPrivate::CategoryDelim, TextHierarchy),
+			Result->GetDisplayName(),
+			Tooltip,
+			ContextGroup,
+			FText::Join(SchemaPrivate::KeywordDelim, Result->GetKeywords())
+		);
 
-			Metadata.SetType(EMetasoundFrontendClassType::External);
-			NewNodeAction->ClassMetadata = Metadata;
-			ActionMenuBuilder.AddAction(NewNodeAction);
-		}
-	}
+		NewNodeAction->QueryResult = TSharedPtr<ISchemaQueryResult>(Result.Release());
+		ActionMenuBuilder.AddAction(NewNodeAction);
+	});
 }
 
-void UMetasoundEditorGraphSchema::GetVariableActions(FGraphActionMenuBuilder& ActionMenuBuilder, Metasound::Editor::FActionClassFilters InFilters, bool bShowSelectedActions, Metasound::Frontend::FConstGraphHandle InGraphHandle) const
+void UMetasoundEditorGraphSchema::GetVariableActions(FGraphActionMenuBuilder& ActionMenuBuilder, Metasound::Editor::FActionVertexFilters, bool bShowSelectedActions, Metasound::Frontend::FConstGraphHandle InGraphHandle) const
 {
 	using namespace Metasound::Frontend;
 	using namespace Metasound::Editor;
@@ -2612,5 +2857,50 @@ void UMetasoundEditorGraphSchema::SetNodePosition(UEdGraphNode* Node, const FVec
 		UEdGraphSchema::SetNodePosition(Node, Position);
 	}
 }
+
+void UMetasoundEditorGraphSchema::QueryNodeClasses(TFunctionRef<void(TUniquePtr<Metasound::Editor::ISchemaQueryResult>&&)> OnClassFound, bool bScanAssetTags) const
+{
+	using namespace Metasound::Editor;
+	using namespace Metasound::Engine;
+	using namespace Metasound::Frontend;
+
+	TArray<TUniquePtr<ISchemaQueryResult>> Results;
+	const FMetaSoundAssetManager& AssetManager = FMetaSoundAssetManager::GetChecked();
+
+	// Gather all loaded classes
+	{
+		TSet<FNodeRegistryKey> FoundKeys;
+		{
+			const bool bIncludeAllVersions = static_cast<bool>(SchemaPrivate::EnableAllVersionsMetaSoundNodeClassCreationCVar);
+			TArray<FMetasoundFrontendClass> FrontendClasses = ISearchEngine::Get().FindAllClasses(bIncludeAllVersions);
+			for (const FMetasoundFrontendClass& Class : FrontendClasses)
+			{
+				FoundKeys.Add(FNodeRegistryKey(Class.Metadata));
+				Results.Add(MakeUnique<SchemaPrivate::FRegistrySchemaQueryResult>(AssetManager, Class));
+			}
+		}
+
+		// Append all unloaded asset classes
+		if (bScanAssetTags)
+		{
+			AssetManager.IterateAssetTagData([&FoundKeys, &Results](FAssetClassQueryInfo QueryInfo)
+			{
+				using namespace Metasound::Frontend;
+
+				const FAssetKey AssetKey(QueryInfo.ClassName, QueryInfo.Version);
+				const FNodeRegistryKey RegKey(AssetKey);
+				if (!FoundKeys.Contains(RegKey))
+				{
+					Results.Add(MakeUnique<SchemaPrivate::FAssetSchemaQueryResult>(MoveTemp(QueryInfo)));
+				}
+			});
+		}
+	}
+
+	for (TUniquePtr<ISchemaQueryResult>& Result : Results)
+	{
+		OnClassFound(MoveTemp(Result));
+	}
+}
 #undef LOCTEXT_NAMESPACE
 
diff --git a/Engine/Plugins/Runtime/Metasound/Source/MetasoundEditor/Private/MetasoundEditorModule.cpp b/Engine/Plugins/Runtime/Metasound/Source/MetasoundEditor/Private/MetasoundEditorModule.cpp
--- a/Engine/Plugins/Runtime/Metasound/Source/MetasoundEditor/Private/MetasoundEditorModule.cpp
+++ b/Engine/Plugins/Runtime/Metasound/Source/MetasoundEditor/Private/MetasoundEditorModule.cpp
@@ -85,6 +85,24 @@ DEFINE_LOG_CATEGORY(LogMetasoundEditor);
 
 namespace Metasound::Editor
 {
+	EAssetPrimeStatus IMetasoundEditorModule::GetAssetRegistryPrimeStatus() const
+	{
+		return EAssetPrimeStatus::NotRequested;
+	}
+
+	EAssetScanStatus IMetasoundEditorModule::GetAssetRegistryScanStatus() const
+	{
+		const bool bIsScanComplete = Engine::FMetaSoundAssetManager::GetChecked().IsInitialAssetScanComplete();
+		if (bIsScanComplete)
+		{
+			return EAssetScanStatus::Complete;
+		}
+
+		// No longer returns whether or not requested vs InProgress, but function is deprecated and should just
+		// use AssetManager to know whether or not scan is in progress.
+		return EAssetScanStatus::InProgress;
+	}
+
 	namespace ModulePrivate
 	{
 		static const FLazyName AssetToolName { "AssetTools" };
@@ -714,18 +732,6 @@ namespace Metasound::Editor
 
 		}
 
-		virtual void PrimeAssetRegistryAsync() override {}
-
-		virtual EAssetPrimeStatus GetAssetRegistryPrimeStatus() const override
-		{
-			return EAssetPrimeStatus::NotRequested;
-		}
-
-		virtual EAssetScanStatus GetAssetRegistryScanStatus() const override
-		{
-			return EAssetScanStatus::NotRequested;
-		}
-
 		virtual void RegisterExplicitProxyClass(const UClass& InClass) override
 		{
 			using namespace Metasound::Frontend;
@@ -925,7 +931,7 @@ namespace Metasound::Editor
 			// Metasound Frontend. The frontend registration must occur before
 			// the Metasound Editor registration of a USoundWave.
 			IMetasoundEngineModule& MetaSoundEngineModule = FModuleManager::LoadModuleChecked<IMetasoundEngineModule>("MetasoundEngine");
-				
+
 			// Bind delegates for MetaSound registration in the asset registry
 			MetaSoundEngineModule.GetOnGraphRegisteredDelegate().BindLambda([](UObject& InMetaSound, Engine::ERegistrationAssetContext AssetContext)
 			{
@@ -1060,7 +1066,7 @@ namespace Metasound::Editor
 					OldToNewReferenceKeys.FindOrAdd(MoveTemp(OldRegistryKey)) = MoveTemp(NewRegistryKey);
 
 					UObject& MetaSoundObject = MetaSoundBuilder.CastDocumentObjectChecked<UObject>();
-					AssetManager.AddOrUpdateAsset(MetaSoundObject);
+					AssetManager.AddOrUpdateFromObject(MetaSoundObject);
 				}
 
 				// Fix up dependencies
diff --git a/Engine/Plugins/Runtime/Metasound/Source/MetasoundEditor/Public/MetasoundEditorGraphSchema.h b/Engine/Plugins/Runtime/Metasound/Source/MetasoundEditor/Public/MetasoundEditorGraphSchema.h
--- a/Engine/Plugins/Runtime/Metasound/Source/MetasoundEditor/Public/MetasoundEditorGraphSchema.h
+++ b/Engine/Plugins/Runtime/Metasound/Source/MetasoundEditor/Public/MetasoundEditorGraphSchema.h
@@ -21,35 +21,62 @@ class UMetaSoundPatch;
 class UMetasoundEditorGraphNode;
 
 
-namespace Metasound
+namespace Metasound::Editor
 {
-	namespace Editor
+	// Ordered such that highest priority is highest value
+	enum class EPrimaryContextGroup : uint8
 	{
-		// Ordered such that highest priority is highest value
-		enum class EPrimaryContextGroup : uint8
-		{
-			Conversions = 0,
-			Graphs,
-			Functions,
-			Outputs,
-			Inputs,
-			Variables,
+		Conversions = 0,
+		Graphs,
+		Functions,
+		Outputs,
+		Inputs,
+		Variables,
+
+		Common
+	};
+
+	const FText& GetContextGroupDisplayName(EPrimaryContextGroup InContextGroup);
+
+	class METASOUNDEDITOR_API ISchemaQueryResult
+	{
+	public:
+		virtual ~ISchemaQueryResult() = default;
+
+		virtual const FMetasoundFrontendClass* FindClass() const = 0;
+		virtual bool CanConnectInputOfTypeAndAccess(FName TypeName, EMetasoundFrontendVertexAccessType AccessType) const = 0;
+		virtual bool CanConnectOutputOfTypeAndAccess(FName TypeName, EMetasoundFrontendVertexAccessType AccessType) const = 0;
 
-			Common
-		};
+		virtual Frontend::FAssetKey GetAssetKey() const = 0;
+		virtual const TArray<FText>& GetCategoryHierarchy() const = 0;
+		virtual FText GetDisplayName() const = 0;
+		virtual const TArray<FText>& GetKeywords() const = 0;
+		virtual EMetasoundFrontendClassType GetRegistryClassType() const = 0;
+		virtual FText GetTooltip() const = 0;
 
-		const FText& GetContextGroupDisplayName(EPrimaryContextGroup InContextGroup);
+		virtual bool IsNative() const = 0;
+	};
 
-		using FInputFilterFunction = TFunction<bool(const FMetasoundFrontendClassInput&)>;
-		using FOutputFilterFunction = TFunction<bool(const FMetasoundFrontendClassOutput&)>;
-		using FInterfaceNodeFilterFunction = TFunction<bool(Frontend::FConstNodeHandle)>;
+	using FInterfaceNodeFilterFunction = TFunction<bool(Frontend::FConstNodeHandle)>;
 
-		struct FActionClassFilters
+	struct FActionVertexFilters
+	{
+		FName InputTypeName = { };
+		EMetasoundFrontendVertexAccessType InputAccessType = EMetasoundFrontendVertexAccessType::Unset;
+
+		bool HasInputFilters() const
+		{
+			return !InputTypeName.IsNone() && InputAccessType != EMetasoundFrontendVertexAccessType::Unset;
+		}
+
+		FName OutputTypeName = { };
+		EMetasoundFrontendVertexAccessType OutputAccessType = EMetasoundFrontendVertexAccessType::Unset;
+
+		bool HasOutputFilters() const
 		{
-			FInputFilterFunction InputFilterFunction;
-			FOutputFilterFunction OutputFilterFunction;
-		};
-	} // namespace Editor
+			return !OutputTypeName.IsNone() && OutputAccessType != EMetasoundFrontendVertexAccessType::Unset;
+		}
+	};
 
 	namespace SchemaUtils
 	{
@@ -59,7 +86,7 @@ namespace Metasound
 		UEdGraphNode* PromoteToDeferredVariable(UEdGraph* ParentGraph, UEdGraphPin* FromPin, const FVector2D InLocation, bool bSelectNewNode);
 		UEdGraphNode* PromoteToMutatorVariable(UEdGraph* ParentGraph, UEdGraphPin* FromPin, const FVector2D InLocation, bool bSelectNewNode);
 	} // namespace SchemaUtils
-} // namespace Metasound
+} // namespace Metasound::Editor
 
 
 USTRUCT()
@@ -351,8 +378,7 @@ struct METASOUNDEDITOR_API FMetasoundGraphSchemaAction_NewNode : public FMetasou
 {
 	GENERATED_USTRUCT_BODY();
 
-	/** Class Metadata of node to create */
-	FMetasoundFrontendClassMetadata ClassMetadata;
+	TSharedPtr<Metasound::Editor::ISchemaQueryResult> QueryResult;
 
 	FMetasoundGraphSchemaAction_NewNode() 
 		: FMetasoundGraphSchemaAction_NodeWithMultipleOutputs()
@@ -524,10 +550,12 @@ public:
 	void BreakNodeLinks(UEdGraphNode& TargetNode, bool bShouldActuallyTransact) const;
 
 private:
+	void QueryNodeClasses(TFunctionRef<void(TUniquePtr<Metasound::Editor::ISchemaQueryResult>&&)> OnClassFound, bool bScanAssetTags = true) const;
+
 	/** Adds actions for creating actions associated with graph DataTypes */
-	void GetConversionActions(FGraphActionMenuBuilder& ActionMenuBuilder, Metasound::Editor::FActionClassFilters InFilters = Metasound::Editor::FActionClassFilters(), bool bShowSelectedActions = true) const;
-	void GetFunctionActions(FGraphActionMenuBuilder& ActionMenuBuilder, Metasound::Editor::FActionClassFilters InFilters = Metasound::Editor::FActionClassFilters(), bool bShowSelectedActions = true, Metasound::Frontend::FConstGraphHandle InGraphHandle = Metasound::Frontend::IGraphController::GetInvalidHandle()) const;
-	void GetVariableActions(FGraphActionMenuBuilder& ActionMenuBuilder, Metasound::Editor::FActionClassFilters InFilters = Metasound::Editor::FActionClassFilters(), bool bShowSelectedActions = true, Metasound::Frontend::FConstGraphHandle InGraphHandle = Metasound::Frontend::IGraphController::GetInvalidHandle()) const;
+	void GetConversionActions(FGraphActionMenuBuilder& ActionMenuBuilder, Metasound::Editor::FActionVertexFilters InFilters = Metasound::Editor::FActionVertexFilters(), bool bShowSelectedActions = true) const;
+	void GetFunctionActions(FGraphActionMenuBuilder& ActionMenuBuilder, Metasound::Editor::FActionVertexFilters InFilters = Metasound::Editor::FActionVertexFilters(), bool bShowSelectedActions = true, Metasound::Frontend::FConstGraphHandle InGraphHandle = Metasound::Frontend::IGraphController::GetInvalidHandle()) const;
+	void GetVariableActions(FGraphActionMenuBuilder& ActionMenuBuilder, Metasound::Editor::FActionVertexFilters InFilters = Metasound::Editor::FActionVertexFilters(), bool bShowSelectedActions = true, Metasound::Frontend::FConstGraphHandle InGraphHandle = Metasound::Frontend::IGraphController::GetInvalidHandle()) const;
 
 	void BreakSinglePinLink(UEdGraphPin* SourcePin, UEdGraphPin* TargetPin, bool bShouldTransact) const;
 
diff --git a/Engine/Plugins/Runtime/Metasound/Source/MetasoundEditor/Public/MetasoundEditorModule.h b/Engine/Plugins/Runtime/Metasound/Source/MetasoundEditor/Public/MetasoundEditorModule.h
--- a/Engine/Plugins/Runtime/Metasound/Source/MetasoundEditor/Public/MetasoundEditorModule.h
+++ b/Engine/Plugins/Runtime/Metasound/Source/MetasoundEditor/Public/MetasoundEditorModule.h
@@ -83,14 +83,14 @@ namespace Metasound
 			UE_DEPRECATED(5.3, "IsMetaSoundAssetClass is deprecated, use IMetasoundUObjectRegistry::IsRegisteredClass")
 			virtual bool IsMetaSoundAssetClass(const FTopLevelAssetPath& InClassName) const = 0;
 
-			UE_DEPRECATED(5.5, "Use PrimeAssetRegistryAsync in MetaSoundEngineModule.")
-			virtual void PrimeAssetRegistryAsync() = 0;
+			UE_DEPRECATED(5.5, "Prime no longer supported (nor recommended as AssetTag updates now enable querying MetaSound assets for interface information without loading full asset.")
+			virtual void PrimeAssetRegistryAsync() { }
 
-			UE_DEPRECATED(5.5, "Use GetNodeClassRegistryPrimeStatus in MetaSoundEngineModule.")
-			virtual EAssetPrimeStatus GetAssetRegistryPrimeStatus() const = 0;
+			UE_DEPRECATED(5.5, "Prime no longer supported (nor recommended as AssetTag updates now enable querying MetaSound assets for interface information without loading full asset.")
+			virtual EAssetPrimeStatus GetAssetRegistryPrimeStatus() const;
 
-			UE_DEPRECATED(5.5, "Use the same function in MetaSoundEngineModule.")
-			virtual EAssetScanStatus GetAssetRegistryScanStatus() const = 0;
+			UE_DEPRECATED(5.5, "Use FMetaSoundAssetManager::IsInitialAssetScanComplete instead")
+			virtual EAssetScanStatus GetAssetRegistryScanStatus() const;
 
 			virtual TUniquePtr<FMetasoundDefaultLiteralCustomizationBase> CreateMemberDefaultLiteralCustomization(UClass& InClass, IDetailCategoryBuilder& DefaultCategoryBuilder) const = 0;
 
diff --git a/Engine/Plugins/Runtime/Metasound/Source/MetasoundEngine/Private/Metasound.cpp b/Engine/Plugins/Runtime/Metasound/Source/MetasoundEngine/Private/Metasound.cpp
--- a/Engine/Plugins/Runtime/Metasound/Source/MetasoundEngine/Private/Metasound.cpp
+++ b/Engine/Plugins/Runtime/Metasound/Source/MetasoundEngine/Private/Metasound.cpp
@@ -28,6 +28,7 @@
 #include "MetasoundReceiveNode.h"
 #include "MetasoundTrigger.h"
 #include "MetasoundUObjectRegistry.h"
+#include "UObject/AssetRegistryTagsContext.h"
 #include "UObject/ObjectSaveContext.h"
 
 #include UE_INLINE_GENERATED_CPP_BY_NAME(Metasound)
@@ -116,7 +117,8 @@ void UMetaSoundPatch::PreDuplicate(FObjectDuplicationParameters& DupParams)
 void UMetaSoundPatch::PostDuplicate(EDuplicateMode::Type InDuplicateMode)
 {
 	Super::PostDuplicate(InDuplicateMode);
-	Metasound::Engine::FAssetHelper::PostDuplicate(this, InDuplicateMode, AssetClassID);
+
+	Metasound::Engine::FAssetHelper::PostDuplicate(this, InDuplicateMode);
 }
 
 void UMetaSoundPatch::PostEditUndo()
@@ -178,14 +180,13 @@ FText UMetaSoundPatch::GetDisplayName() const
 	FString TypeName = UMetaSoundPatch::StaticClass()->GetName();
 	return FMetasoundAssetBase::GetDisplayName(MoveTemp(TypeName));
 }
+#endif // WITH_EDITORONLY_DATA
 
-
-void UMetaSoundPatch::SetRegistryAssetClassInfo(const Metasound::Frontend::FNodeClassInfo& InNodeInfo)
+void UMetaSoundPatch::GetAssetRegistryTags(FAssetRegistryTagsContext Context) const
 {
-	Metasound::Engine::FAssetHelper::SetMetaSoundRegistryAssetClassInfo(*this, InNodeInfo);
+	Super::GetAssetRegistryTags(Context);
+	Metasound::Engine::FAssetHelper::GetAssetRegistryTags(this, Context);
 }
-#endif // WITH_EDITORONLY_DATA
-
 
 FTopLevelAssetPath UMetaSoundPatch::GetAssetPathChecked() const
 {
diff --git a/Engine/Plugins/Runtime/Metasound/Source/MetasoundEngine/Private/MetasoundAssetSubsystem.cpp b/Engine/Plugins/Runtime/Metasound/Source/MetasoundEngine/Private/MetasoundAssetSubsystem.cpp
--- a/Engine/Plugins/Runtime/Metasound/Source/MetasoundEngine/Private/MetasoundAssetSubsystem.cpp
+++ b/Engine/Plugins/Runtime/Metasound/Source/MetasoundEngine/Private/MetasoundAssetSubsystem.cpp
@@ -9,6 +9,7 @@
 #include "HAL/CriticalSection.h"
 #include "Metasound.h"
 #include "MetasoundAssetBase.h"
+#include "MetasoundAssetManager.h"
 #include "MetasoundBuilderSubsystem.h"
 #include "MetasoundDocumentBuilderRegistry.h"
 #include "MetasoundEngineAsset.h"
@@ -23,7 +24,9 @@
 #include "MetasoundUObjectRegistry.h"
 #include "Misc/CoreDelegates.h"
 #include "Modules/ModuleManager.h"
+#include "UObject/AssetRegistryTagsContext.h"
 #include "UObject/NoExportTypes.h"
+#include "UObject/Object.h"
 
 
 #include UE_INLINE_GENERATED_CPP_BY_NAME(MetasoundAssetSubsystem)
@@ -33,80 +36,31 @@ namespace Metasound::Engine
 {
 	namespace AssetSubsystemPrivate
 	{
-		bool GetAssetClassInfo(const FAssetData& InAssetData, Frontend::FNodeClassInfo& OutInfo)
-		{
-			using namespace Metasound;
-			using namespace Metasound::Frontend;
-
-			bool bSuccess = true;
-
-			OutInfo.Type = EMetasoundFrontendClassType::External;
-			OutInfo.AssetPath = FTopLevelAssetPath(InAssetData.PackageName, InAssetData.AssetName);
-			FString AssetClassID;
-			bSuccess &= InAssetData.GetTagValue(AssetTags::AssetClassID, AssetClassID);
-			OutInfo.AssetClassID = FGuid(AssetClassID);
-			OutInfo.ClassName = FMetasoundFrontendClassName(FName(), *AssetClassID, FName());
-
-#if WITH_EDITORONLY_DATA
-			InAssetData.GetTagValue(AssetTags::IsPreset, OutInfo.bIsPreset);
-#endif // WITH_EDITORONLY_DATA
-
-			int32 RegistryVersionMajor = 0;
-			bSuccess &= InAssetData.GetTagValue(AssetTags::RegistryVersionMajor, RegistryVersionMajor);
-			OutInfo.Version.Major = RegistryVersionMajor;
-
-			int32 RegistryVersionMinor = 0;
-			bSuccess &= InAssetData.GetTagValue(AssetTags::RegistryVersionMinor, RegistryVersionMinor);
-			OutInfo.Version.Minor = RegistryVersionMinor;
-
-#if WITH_EDITORONLY_DATA
-			auto ParseTypesString = [&](const FName AssetTag, TSet<FName>& OutTypes)
-			{
-				FString TypesString;
-				if (InAssetData.GetTagValue(AssetTag, TypesString))
-				{
-					TArray<FString> DataTypeStrings;
-					TypesString.ParseIntoArray(DataTypeStrings, *AssetTags::ArrayDelim);
-					Algo::Transform(DataTypeStrings, OutTypes, [](const FString& DataType) { return *DataType; });
-					return true;
-				}
-
-				return false;
-			};
-
-			// These values are optional and not necessary to return successfully as MetaSounds
-			// don't require inputs or outputs for asset tags to be valid (ex. a new MetaSound,
-			// non-source asset has no inputs or outputs)
-			OutInfo.InputTypes.Reset();
-			ParseTypesString(AssetTags::RegistryInputTypes, OutInfo.InputTypes);
-
-			OutInfo.OutputTypes.Reset();
-			ParseTypesString(AssetTags::RegistryOutputTypes, OutInfo.OutputTypes);
-#endif // WITH_EDITORONLY_DATA
+		using FAssetClassInfoMap = FMetaSoundAssetManager::FAssetClassInfoMap;
 
-			return bSuccess;
-		}
+		std::atomic<bool> InitialAssetScanComplete = false;
 
-		bool RemovePath(FCriticalSection* MapCritSec, TMap<Frontend::FAssetKey, TArray<FTopLevelAssetPath>>& Map, const Frontend::FAssetKey& AssetKey, const FTopLevelAssetPath& AssetPath)
+		bool RemoveAssetTagData(FCriticalSection& MapCritSec, FAssetClassInfoMap& Map, const Frontend::FAssetKey& AssetKey, const FTopLevelAssetPath& InAssetPath)
 		{
-			check(MapCritSec);
-			FScopeLock Lock(MapCritSec);
-			if (TArray<FTopLevelAssetPath>* MapAssetPaths = Map.Find(AssetKey))
+			using namespace Frontend;
+
+			FScopeLock Lock(&MapCritSec);
+			if (TArray<FAssetClassQueryInfo>* AssetInfoArray = Map.Find(AssetKey))
 			{
-				auto ComparePaths = [&AssetPath](const FTopLevelAssetPath& Path) 
+				auto ComparePaths = [&InAssetPath](const FAssetClassQueryInfo& QueryInfo)
 				{
 					// Compare full paths if valid
-					if (Path.IsValid() && AssetPath.IsValid())
+					if (QueryInfo.AssetPath.IsValid() && InAssetPath.IsValid())
 					{
-						return Path == AssetPath;
+						return QueryInfo.AssetPath == InAssetPath;
 					}
 					// Package names are stripped on destruction, so only asset name is reliable
-					return Path.GetAssetName() == AssetPath.GetAssetName(); 
+					return QueryInfo.AssetPath.GetAssetName() == InAssetPath.GetAssetName();
 				};
 
-				if (MapAssetPaths->RemoveAllSwap(ComparePaths, EAllowShrinking::No) > 0)
+				if (AssetInfoArray->RemoveAllSwap(ComparePaths, EAllowShrinking::No) > 0)
 				{
-					if (MapAssetPaths->IsEmpty())
+					if (AssetInfoArray->IsEmpty())
 					{
 						Map.Remove(AssetKey);
 					}
@@ -117,139 +71,60 @@ namespace Metasound::Engine
 			return false;
 		}
 
-		void AddPath(FCriticalSection* MapCritSec, TMap<Frontend::FAssetKey, TArray<FTopLevelAssetPath>>& Map, const Frontend::FAssetKey& AssetKey, const FTopLevelAssetPath& AssetPath)
+		Frontend::FAssetKey AddAssetTagData(FCriticalSection& MapCritSec, FAssetClassInfoMap& Map, Frontend::FAssetClassQueryInfo&& QueryInfo)
 		{
-			check(MapCritSec);
-			FScopeLock Lock(MapCritSec);
-			TArray<FTopLevelAssetPath>& Paths = Map.FindOrAdd(AssetKey);
-			Paths.AddUnique(AssetPath);
+			using namespace Frontend;
+
+			const Frontend::FAssetKey AssetKey(QueryInfo.ClassName, QueryInfo.Version);
+			if (!AssetKey.IsValid())
+			{
+				return AssetKey;
+			}
+
+			FScopeLock Lock(&MapCritSec);
+			TArray<FAssetClassQueryInfo>& TagDatas = Map.FindOrAdd(AssetKey);
+			auto IsTagData = [&QueryInfo](const FAssetClassQueryInfo& Iter) { return Iter.AssetPath == QueryInfo.AssetPath; };
+			TagDatas.RemoveAllSwap(IsTagData);
+			TagDatas.Add(MoveTemp(QueryInfo));
+
 #if !NO_LOGGING
-			if (Paths.Num() > 1)
+			if (TagDatas.Num() > 1)
 			{
 				TArray<FString> PathStrings;
-				Algo::Transform(Paths, PathStrings, [](const FTopLevelAssetPath& Path) { return Path.ToString(); });
+				Algo::Transform(TagDatas, PathStrings, [](const FAssetClassQueryInfo& QueryInfo) { return QueryInfo.AssetPath.ToString(); });
 				UE_LOG(LogMetaSound, Warning,
 					TEXT("MetaSoundAssetManager has registered multiple assets with key '%s':\n%s\n"),
 					*AssetKey.ToString(),
 					*FString::Join(PathStrings, TEXT("\n")));
 			}
 #endif // !NO_LOGGING
-		}
-	}
-
-	class FMetaSoundAssetManager :
-		public Frontend::IMetaSoundAssetManager,
-		public FGCObject
-	{
-	public:
-		using FAssetInfo = Frontend::IMetaSoundAssetManager::FAssetInfo;
-		using FAssetKey = Frontend::FAssetKey;
-		using FNodeRegistryKey = Frontend::FNodeRegistryKey;
 
-#if WITH_EDITORONLY_DATA
-		using FVersionDocumentResults = Frontend::IMetaSoundAssetManager::FVersionDocumentResults;
-#endif // WITH_EDITORONLY_DATA
-
-		FMetaSoundAssetManager() = default;
-		virtual ~FMetaSoundAssetManager();
-
-		static FMetaSoundAssetManager& GetChecked()
-		{
-			using namespace Frontend;
-			return static_cast<FMetaSoundAssetManager&>(IMetaSoundAssetManager::GetChecked());
+			return AssetKey;
 		}
-
-		void RebuildDenyListCache(const UAssetManager& InAssetManager);
-		void RegisterAssetClassesInDirectories(const TArray<FMetaSoundAssetDirectory>& Directories);
-
-#if WITH_EDITOR
-		bool ReplaceReferencesInDirectory(const TArray<FMetaSoundAssetDirectory>& InDirectories, const FNodeRegistryKey& OldClassKey, const FNodeRegistryKey& NewClassKey) const;
-#endif // WITH_EDITOR
-		void RequestAsyncLoadReferencedAssets(FMetasoundAssetBase& InAssetBase);
-		void OnAssetScanComplete();
-		void SearchAndIterateDirectoryAssets(const TArray<FDirectoryPath>& InDirectories, TFunctionRef<void(const FAssetData&)> InFunction) const;
-		FMetasoundAssetBase* TryLoadAsset(const FSoftObjectPath& InObjectPath) const;
-		void UnregisterAssetClassesInDirectories(const TArray<FMetaSoundAssetDirectory>& Directories);
-
-		/* IMetaSoundAssetManager Implementation */
-#if WITH_EDITORONLY_DATA
-		virtual bool AddAssetReferences(FMetasoundAssetBase& InAssetBase) override;
-#endif // WITH_EDITORONLY_DATA
-		virtual FAssetKey AddOrUpdateAsset(const FAssetData& InAssetData) override;
-		virtual FAssetKey AddOrUpdateAsset(const UObject& InObject) override;
-		virtual bool CanAutoUpdate(const FMetasoundFrontendClassName& InClassName) const override;
-		virtual bool ContainsKey(const FAssetKey& InKey) const override;
-		virtual FMetasoundAssetBase* FindAsset(const FAssetKey& InKey) const override;
-		virtual TScriptInterface<IMetaSoundDocumentInterface> FindAssetAsDocumentInterface(const FAssetKey& InKey) const override;
-		virtual FTopLevelAssetPath FindAssetPath(const FAssetKey& InKey) const override;
-		virtual TArray<FTopLevelAssetPath> FindAssetPaths(const FAssetKey& InKey) const override;
-		virtual FMetasoundAssetBase* GetAsAsset(UObject& InObject) const override;
-		virtual const FMetasoundAssetBase* GetAsAsset(const UObject& InObject) const override;
-#if WITH_EDITOR
-		virtual TSet<FAssetInfo> GetReferencedAssetClasses(const FMetasoundAssetBase& InAssetBase) const override;
-		virtual bool ReassignClassName(TScriptInterface<IMetaSoundDocumentInterface> DocInterface) override;
-#endif // WITH_EDITOR
-		virtual void IterateAssets(TFunctionRef<void(const FAssetKey, const TArray<FTopLevelAssetPath>&)> Iter) const override;
-		virtual void ReloadMetaSoundAssets() const override;
-		virtual void RemoveAsset(const UObject& InObject) override;
-		virtual void RemoveAsset(const FAssetData& InAssetData) override;
-		virtual void RenameAsset(const FAssetData& InAssetData, const FString& InOldObjectPath) override;
-		virtual void SetLogActiveAssetsOnShutdown(bool bInLogActiveAssetsOnShutdown) override;
-		virtual FMetasoundAssetBase* TryLoadAssetFromKey(const FAssetKey& InKey) const override;
-		virtual bool TryGetAssetIDFromClassName(const FMetasoundFrontendClassName& InClassName, FGuid& OutGuid) const override;
-		virtual bool TryLoadReferencedAssets(const FMetasoundAssetBase& InAssetBase, TArray<FMetasoundAssetBase*>& OutReferencedAssets) const override;
-#if WITH_EDITORONLY_DATA
-		virtual FVersionDocumentResults VersionDocumentsInFolders(const TArray<FString>& FolderPaths) const override;
-#endif // WITH_EDITORONLY_DATA
-
-		virtual void WaitUntilAsyncLoadReferencedAssetsComplete(FMetasoundAssetBase& InAssetBase) override;
-
-		/* FGCObject */
-		virtual void AddReferencedObjects(FReferenceCollector& Collector) override;
-		virtual FString GetReferencerName() const override { return TEXT("FMetaSoundAssetManager"); }
-
-	private:
-		TArray<FMetaSoundAsyncAssetDependencies> LoadingDependencies;
-
-		FMetaSoundAsyncAssetDependencies* FindLoadingDependencies(const UObject* InParentAsset);
-		FMetaSoundAsyncAssetDependencies* FindLoadingDependencies(int32 InLoadID);
-		void RemoveLoadingDependencies(int32 InLoadID);
-		void OnAssetsLoaded(int32 InLoadID);
-
-		FStreamableManager StreamableManager;
-		int32 AsyncLoadIDCounter = 0;
-
-		int32 AutoUpdateDenyListChangeID = INDEX_NONE;
-		TSet<FName> AutoUpdateDenyListCache;
-		std::atomic<bool> bIsInitialAssetScanComplete = false;
-		TMap<FAssetKey, TArray<FTopLevelAssetPath>> PathMap;
-
-		// Critical section primarily for allowing safe access of path map during async loading of MetaSound assets.
-		mutable FCriticalSection PathMapCriticalSection;
-
-		bool bLogActiveAssetsOnShutdown = true;
-	};
+	} // AssetSubsystemPrivate
 
 	FMetaSoundAssetManager::~FMetaSoundAssetManager()
 	{
 #if !NO_LOGGING
+		using namespace Frontend;
+
 		if (bLogActiveAssetsOnShutdown)
 		{
-			TMap<FAssetKey, TArray<FTopLevelAssetPath>> PathsOnShutdown;
+			TMap<FAssetKey, TArray<FAssetClassQueryInfo>> InfoOnShutdown;
 			{
-				FScopeLock Lock(&PathMapCriticalSection);
-				PathsOnShutdown = MoveTemp(PathMap);
-				PathMap.Reset();
+				FScopeLock Lock(&TagDataMapCriticalSection);
+				InfoOnShutdown = MoveTemp(ClassInfoMap);
+				ClassInfoMap.Reset();
 			}
 
-			if (!PathsOnShutdown.IsEmpty())
+			if (!InfoOnShutdown.IsEmpty())
 			{
-				UE_LOG(LogMetaSound, Display, TEXT("AssetManager is shutting down with the following %i assets active:"), PathsOnShutdown.Num());
-				for (const TPair<FAssetKey, TArray<FTopLevelAssetPath>>& Pair : PathsOnShutdown)
+				UE_LOG(LogMetaSound, Display, TEXT("AssetManager is shutting down with the following %i assets active:"), InfoOnShutdown.Num());
+				for (const TPair<FAssetKey, TArray<FAssetClassQueryInfo>>& Pair : InfoOnShutdown)
 				{
-					for (const FTopLevelAssetPath& Path : Pair.Value)
+					for (const FAssetClassQueryInfo& QueryInfo : Pair.Value)
 					{
-						UE_LOG(LogMetaSound, Display, TEXT("- %s"), *Path.ToString());
+						UE_LOG(LogMetaSound, Display, TEXT("- %s"), *QueryInfo.AssetPath.ToString());
 					}
 				}
 			}
@@ -276,7 +151,7 @@ namespace Metasound::Engine
 			const FAssetKey AssetKey(Document.RootGraph.Metadata);
 			if (!ContainsKey(AssetKey))
 			{
-				AddOrUpdateAsset(*InAssetBase.GetOwningAsset());
+				AddOrUpdateFromObject(*InAssetBase.GetOwningAsset());
 				UE_LOG(LogMetaSound, Verbose, TEXT("Adding asset '%s' to MetaSoundAsset registry."), *InAssetBase.GetOwningAssetName());
 			}
 		}
@@ -287,7 +162,7 @@ namespace Metasound::Engine
 		{
 			FNodeRegistryKey RegistryKey;
 			const bool bIsKey = FNodeRegistryKey::Parse(KeyString, RegistryKey);
-			if (!bIsKey || !ContainsKey(RegistryKey))
+			if (!bIsKey || !ContainsKey(FAssetKey(RegistryKey)))
 			{
 				UE_LOG(LogMetaSound, Verbose, TEXT("Missing referenced class '%s' asset entry."), *KeyString);
 				bAddFromReferencedAssets = true;
@@ -319,7 +194,7 @@ namespace Metasound::Engine
 					UObject* MetaSoundObject = Asset->GetOwningAsset();
 					if (ensureAlways(MetaSoundObject))
 					{
-						AddOrUpdateAsset(*MetaSoundObject);
+						AddOrUpdateFromObject(*MetaSoundObject);
 					}
 				}
 			}
@@ -333,79 +208,109 @@ namespace Metasound::Engine
 	}
 #endif // WITH_EDITORONLY_DATA
 
-	Frontend::FAssetKey FMetaSoundAssetManager::AddOrUpdateAsset(const UObject& InObject)
+	Frontend::FAssetKey FMetaSoundAssetManager::AddOrUpdateFromObject(const UObject& InObject)
 	{
 		using namespace AssetSubsystemPrivate;
 		using namespace Frontend;
 
-		METASOUND_TRACE_CPUPROFILER_EVENT_SCOPE(FMetaSoundAssetManager::AddOrUpdateAsset_UObject);
+		METASOUND_TRACE_CPUPROFILER_EVENT_SCOPE(FMetaSoundAssetManager::AddOrUpdateFromObject);
 
 		const FMetasoundAssetBase* MetaSoundAsset = Metasound::IMetasoundUObjectRegistry::Get().GetObjectAsAssetBase(&InObject);
 		check(MetaSoundAsset);
 
-		const FMetasoundFrontendDocument& Document = MetaSoundAsset->GetConstDocumentChecked();
-		const FAssetKey AssetKey = FAssetKey(Document.RootGraph);
+		TScriptInterface<const IMetaSoundDocumentInterface> ScriptInterface(&InObject);
+		check(ScriptInterface.GetObject() != nullptr);
+		const IMetaSoundDocumentInterface& DocInterface = *ScriptInterface;
+		return AssetSubsystemPrivate::AddAssetTagData(TagDataMapCriticalSection, ClassInfoMap, FAssetClassQueryInfo(DocInterface));
+	}
+
+	void  FMetaSoundAssetManager::AddOrLoadAndUpdateFromObjectAsync(const FAssetData& InAssetData, TUniqueFunction<void(FAssetKey, UObject&)>&& OnAssetRegistered)
+	{
+		using namespace Frontend;
+
+		METASOUND_TRACE_CPUPROFILER_EVENT_SCOPE(UMetaSoundAssetSubsystem::AddOrLoadAndUpdateFromObjectAsync);
+
+		// Don't add temporary assets used for diffing
+		if (InAssetData.HasAnyPackageFlags(PKG_ForDiffing))
+		{
+			OnAssetRegistered.Reset();
+			return;
+		}
 
-		if (AssetKey.IsValid())
+		const FSoftObjectPath Path = InAssetData.ToSoftObjectPath();
+		if (UObject* Object = Path.ResolveObject())
 		{
-			AssetSubsystemPrivate::AddPath(&PathMapCriticalSection, PathMap, AssetKey, FTopLevelAssetPath(&InObject));
+			UE_LOG(LogMetaSound, VeryVerbose, TEXT("FMetaSoundAssetManager::AddOrLoadAndUpdateFromObjectAsync called, but object ''%s' already loaded. Updating entry using in-memory object definition."), *Path.ToString());
+			const FAssetKey AssetKey = AddOrUpdateFromObject(*Object);
+			OnAssetRegistered(AssetKey, *Object);
+			OnAssetRegistered.Reset();
+			return;
 		}
 
-		return AssetKey;
+#if WITH_EDITORONLY_DATA
+		ActiveAsyncAssetLoadRequests++;
+#endif // WITH_EDITORONLY_DATA
+
+		UE_LOG(LogMetaSound, VeryVerbose, TEXT("MetaSoundAssetManager requested aync loading asset '%s'..."), *InAssetData.GetSoftObjectPath().ToString());
+		LoadPackageAsync(Path.GetLongPackageName(), FLoadPackageAsyncDelegate::CreateLambda([](const FName& PackageName, UPackage* Package, EAsyncLoadingResult::Type Result)
+			{
+				FPackageLoadedArgs Args;
+				Args.bMarkAsVersioned = true;
+				Args.Package = Package;
+				Args.PackageName = PackageName;
+				Args.Result = Result;
+				FMetaSoundAssetManager::OnPackageLoaded(Args);
+			})
+		);
 	}
 
-	Frontend::FAssetKey FMetaSoundAssetManager::AddOrUpdateAsset(const FAssetData& InAssetData)
+	void FMetaSoundAssetManager::AddOrUpdateFromAssetTags(const FAssetData& InAssetData)
 	{
 		using namespace Frontend;
 
 		METASOUND_TRACE_CPUPROFILER_EVENT_SCOPE(UMetaSoundAssetSubsystem::AddOrUpdateAsset_AssetData);
 
-		// Invalid ClassID means the node could not be registered.
-		// Let caller report or ensure as necessary.
-		FAssetKey AssetKey = FAssetKey::GetInvalid();
-
 		// Don't add temporary assets used for diffing
 		if (InAssetData.HasAnyPackageFlags(PKG_ForDiffing))
 		{
-			return AssetKey;
+			return;
 		}
 
-		FNodeClassInfo ClassInfo;
-		bool bClassInfoFound = AssetSubsystemPrivate::GetAssetClassInfo(InAssetData, ClassInfo);
-		if (!bClassInfoFound)
+		// If object is loaded, always use data in memory instead of asset tag data as it may be more up-to-date.
+		const FSoftObjectPath Path = InAssetData.ToSoftObjectPath();
+		if (UObject* Object = Path.ResolveObject())
 		{
-			UObject* Object = nullptr;
-
-			FSoftObjectPath Path = InAssetData.ToSoftObjectPath();
-			if (!FPackageName::GetPackageMountPoint(InAssetData.GetObjectPathString()).IsNone())
-			{
-				if (InAssetData.IsAssetLoaded())
-				{
-					Object = Path.ResolveObject();
-					UE_LOG(LogMetaSound, Verbose, TEXT("Adding loaded asset '%s' to MetaSoundAsset registry."), *Object->GetName());
-				}
-				else
-				{
-					Object = Path.TryLoad();
-					UE_LOG(LogMetaSound, Verbose, TEXT("Loaded asset '%s' and adding to MetaSoundAsset registry."), *Object->GetName());
-				}
-			}
-
-			if (Object)
-			{
-				return AddOrUpdateAsset(*Object);
-			}
+			UE_LOG(LogMetaSound, VeryVerbose, TEXT("FMetaSoundAssetManager::AddOrUpdateFromAssetTags, but object '%s' already loaded. Updating entry using in-memory object definition."), *Path.ToString());
+			AddOrUpdateFromObject(*Object);
 		}
-
-		if (ClassInfo.AssetClassID.IsValid())
+		else
 		{
-			AssetKey = FAssetKey(ClassInfo.ClassName, ClassInfo.Version);
-			if (AssetKey.IsValid())
+			// Older document versions did not serialize asset tags. In these cases, a load of the object is required
+			// in order to retrieve the appropriate asset tag data translated to asset info cached in the MetaSound Asset Manager.
+			FAssetClassQueryInfo QueryInfo(InAssetData);
+			if (QueryInfo.bIsValid)
+			{
+				AssetSubsystemPrivate::AddAssetTagData(TagDataMapCriticalSection, ClassInfoMap, MoveTemp(QueryInfo));
+			}
+			else
 			{
-				AssetSubsystemPrivate::AddPath(&PathMapCriticalSection, PathMap, AssetKey, ClassInfo.AssetPath);
+#if WITH_EDITORONLY_DATA
+				ActiveAsyncAssetLoadRequests++;
+#endif // WITH_EDITORONLY_DATA
+		
+				UE_LOG(LogMetaSound, Verbose, TEXT("MetaSound tags require updating: MetaSoundAssetManager aync loading asset '%s' to access tags/update entry..."), *InAssetData.GetSoftObjectPath().ToString());
+				LoadPackageAsync(Path.GetLongPackageName(), FLoadPackageAsyncDelegate::CreateLambda([](const FName& PackageName, UPackage* Package, EAsyncLoadingResult::Type Result)
+					{
+						FPackageLoadedArgs Args;
+						Args.bMarkAsVersioned = true;
+						Args.Package = Package;
+						Args.PackageName = PackageName;
+						Args.Result = Result;
+						FMetaSoundAssetManager::OnPackageLoaded(Args);
+					})
+				);
 			}
 		}
-		return AssetKey;
 	}
 
 	bool FMetaSoundAssetManager::CanAutoUpdate(const FMetasoundFrontendClassName& InClassName) const
@@ -421,8 +326,36 @@ namespace Metasound::Engine
 
 	bool FMetaSoundAssetManager::ContainsKey(const Frontend::FAssetKey& InKey) const
 	{
-		FScopeLock Lock(&PathMapCriticalSection);
-		return PathMap.Contains(InKey);
+		FScopeLock Lock(&TagDataMapCriticalSection);
+		return ClassInfoMap.Contains(InKey);
+	}
+
+	bool FMetaSoundAssetManager::ContainsKey(const Frontend::FNodeRegistryKey& InRegistryKey) const
+	{
+		if (Frontend::FAssetKey::IsValidType(InRegistryKey.Type))
+		{
+			return ContainsKey(Frontend::FAssetKey(InRegistryKey));
+		}
+
+		return false;
+	}
+
+	void FMetaSoundAssetManager::DepthFirstVisit_AssetKey(const FAssetKey& InKey, TFunctionRef<TSet<FAssetKey>(const FAssetKey&)> VisitFunction)
+	{
+		// Non recursive depth first traversal.
+		TArray<FAssetKey> Stack({ InKey });
+		TSet<FAssetKey> Visited;
+
+		while (!Stack.IsEmpty())
+		{
+			const FAssetKey CurrentKey = Stack.Pop();
+			if (!Visited.Contains(CurrentKey))
+			{
+				TArray<FAssetKey> Children = VisitFunction(CurrentKey).Array();
+				Stack.Append(MoveTemp(Children));
+				Visited.Add(CurrentKey);
+			}
+		}
 	}
 
 	FMetaSoundAsyncAssetDependencies* FMetaSoundAssetManager::FindLoadingDependencies(const UObject* InParentAsset)
@@ -459,6 +392,12 @@ namespace Metasound::Engine
 		return nullptr;
 	}
 
+	TArray<Frontend::FAssetClassQueryInfo> FMetaSoundAssetManager::FindAssetTagData(const Frontend::FAssetKey& InKey) const
+	{
+		FScopeLock Lock(&TagDataMapCriticalSection);
+		return ClassInfoMap.FindRef(InKey);
+	}
+
 	TScriptInterface<IMetaSoundDocumentInterface> FMetaSoundAssetManager::FindAssetAsDocumentInterface(const Frontend::FAssetKey& InKey) const
 	{
 		const FTopLevelAssetPath AssetPath = FindAssetPath(InKey);
@@ -475,12 +414,14 @@ namespace Metasound::Engine
 
 	FTopLevelAssetPath FMetaSoundAssetManager::FindAssetPath(const Frontend::FAssetKey& InKey) const
 	{
-		FScopeLock Lock(&PathMapCriticalSection);
-		if (const TArray<FTopLevelAssetPath>* Paths = PathMap.Find(InKey))
+		using namespace Frontend;
+
+		FScopeLock Lock(&TagDataMapCriticalSection);
+		if (const TArray<FAssetClassQueryInfo>* TagDatas = ClassInfoMap.Find(InKey))
 		{
-			if (!Paths->IsEmpty())
+			if (!TagDatas->IsEmpty())
 			{
-				return Paths->Last();
+				return TagDatas->Last().AssetPath;
 			}
 		}
 
@@ -489,15 +430,26 @@ namespace Metasound::Engine
 
 	TArray<FTopLevelAssetPath> FMetaSoundAssetManager::FindAssetPaths(const Frontend::FAssetKey& InKey) const
 	{
-		FScopeLock Lock(&PathMapCriticalSection);
-		if (const TArray<FTopLevelAssetPath>* Paths = PathMap.Find(InKey))
+		using namespace Frontend;
+
+		FScopeLock Lock(&TagDataMapCriticalSection);
+		if (const TArray<FAssetClassQueryInfo>* TagDatas = ClassInfoMap.Find(InKey))
 		{
-			return *Paths;
+			TArray<FTopLevelAssetPath> Paths;
+			Algo::Transform(*TagDatas, Paths, [](const FAssetClassQueryInfo& QueryInfo) { return QueryInfo.AssetPath; });
+			return Paths;
 		}
 
 		return { };
 	}
 
+#if WITH_EDITOR
+	int32 FMetaSoundAssetManager::GetActiveAsyncLoadRequestCount() const
+	{
+		return ActiveAsyncAssetLoadRequests;
+	}
+#endif // WITH_EDITOR
+
 	FMetasoundAssetBase* FMetaSoundAssetManager::GetAsAsset(UObject& InObject) const
 	{
 		return IMetasoundUObjectRegistry::Get().GetObjectAsAssetBase(&InObject);
@@ -508,7 +460,42 @@ namespace Metasound::Engine
 		return IMetasoundUObjectRegistry::Get().GetObjectAsAssetBase(&InObject);
 	}
 
-	void FMetaSoundAssetManager::OnAssetsLoaded(int32 InLoadID)
+	Frontend::FAssetKey FMetaSoundAssetManager::GetAssetKey(const FSoftObjectPath& InObjectPath) const
+	{
+		using namespace Frontend;
+
+		FAssetKey AssetKey;
+
+		if (const UObject* Object = InObjectPath.ResolveObject())
+		{
+			TScriptInterface<const IMetaSoundDocumentInterface> DocInterface(Object);
+			if (DocInterface.GetObject())
+			{
+				return FAssetKey(DocInterface->GetConstDocument().RootGraph);
+			}
+		}
+
+		FAssetData AssetData;
+		const FAssetRegistryModule& AssetRegistryModule = FModuleManager::GetModuleChecked<FAssetRegistryModule>("AssetRegistry");
+		UE::AssetRegistry::EExists Exists = AssetRegistryModule.TryGetAssetByObjectPath(InObjectPath, AssetData);
+		if (Exists == UE::AssetRegistry::EExists::Exists)
+		{
+			const FAssetClassQueryInfo QueryInfo(AssetData);
+			if (QueryInfo.bIsValid)
+			{
+				return FAssetKey(QueryInfo.ClassName, QueryInfo.Version);
+			}
+		}
+
+		return { };
+	}
+
+	bool FMetaSoundAssetManager::IsInitialAssetScanComplete() const
+	{
+		return AssetSubsystemPrivate::InitialAssetScanComplete;
+	}
+
+	void FMetaSoundAssetManager::OnReferencedAssetsLoaded(int32 InLoadID)
 	{
 		FMetaSoundAsyncAssetDependencies* LoadedDependencies = FindLoadingDependencies(InLoadID);
 		if (ensureMsgf(LoadedDependencies, TEXT("Call to async asset load complete with invalid IDs %d"), InLoadID))
@@ -550,9 +537,74 @@ namespace Metasound::Engine
 		}
 	}
 
-	void FMetaSoundAssetManager::OnAssetScanComplete()
+	void FMetaSoundAssetManager::OnPackageLoaded(const FPackageLoadedArgs& PackageLoadedArgs)
 	{
-		bIsInitialAssetScanComplete = true;
+		FMetaSoundAssetManager* AssetManager = FMetaSoundAssetManager::Get();
+		if (!AssetManager)
+		{
+			return; // Likely shutting down
+		}
+
+#if WITH_EDITOR
+		AssetManager->ActiveAsyncAssetLoadRequests--;
+
+		if (AssetManager->bNotifyTagDataScanComplete)
+		{
+			if (AssetManager->IsInitialAssetScanComplete() && AssetManager->ActiveAsyncAssetLoadRequests == 0)
+			{
+				AssetManager->bNotifyTagDataScanComplete = false;
+				UE_LOG(LogMetaSound, Display, TEXT("Async MetaSound Load/Asset Tag Prime Complete"));
+			}
+		}
+#endif // WITH_EDITOR
+
+		switch (PackageLoadedArgs.Result)
+		{
+			case EAsyncLoadingResult::Succeeded:
+			{
+				check(PackageLoadedArgs.Package);
+				if (UObject* MetaSoundObj = PackageLoadedArgs.Package->FindAssetInPackage())
+				{
+					// Tags had to be versioned if asset tags were not loaded properly so mark asset as versioned on load.
+					// This flags version scripts to properly resave even if the document version remained the same.
+					{
+#if WITH_EDITORONLY_DATA
+						if (PackageLoadedArgs.bMarkAsVersioned)
+						{
+							if (FMetasoundAssetBase* AssetBase = AssetManager->GetAsAsset(*MetaSoundObj))
+							{
+								AssetBase->SetVersionedOnLoad();
+							}
+						}
+#endif // WITH_EDITORONLY_DATA
+
+						const FAssetKey AssetKey = AssetManager->AddOrUpdateFromObject(*MetaSoundObj);
+						if (PackageLoadedArgs.OnAssetRegistered)
+						{
+							PackageLoadedArgs.OnAssetRegistered(AssetKey, *MetaSoundObj);
+						}
+					}
+				}
+				break;
+			}
+
+			case EAsyncLoadingResult::Canceled:
+			{
+				UE_LOG(LogMetaSound, VeryVerbose, TEXT("MetaSoundAssetManager request to aync load asset '%s' canceled"), *PackageLoadedArgs.PackageName.ToString());
+				break;
+			}
+
+			case EAsyncLoadingResult::Failed:
+			{
+				UE_LOG(LogMetaSound, Warning, TEXT("MetaSoundAssetManager request to aync load asset '%s' failed"), *PackageLoadedArgs.PackageName.ToString());
+				break;
+			}
+
+			default:
+			{
+				checkNoEntry();
+			}
+		}
 	}
 
 #if WITH_EDITOR
@@ -579,20 +631,12 @@ namespace Metasound::Engine
 			}
 			else
 			{
-				FNodeRegistryKey RegistryKey(Class.Metadata);
+				const FNodeRegistryKey RegistryKey(Class.Metadata);
 				const FMetasoundFrontendRegistryContainer* Registry = FMetasoundFrontendRegistryContainer::Get();
 				check(Registry);
-				const bool bIsRegistered = Registry->IsNodeRegistered(RegistryKey);
 
-				bool bReportFail = false;
-				if (bIsRegistered)
-				{
-					if (!Registry->IsNodeNative(RegistryKey))
-					{
-						bReportFail = true;
-					}
-				}
-				else
+				const bool bIsRegistered = Registry->IsNodeRegistered(RegistryKey);
+				if (!bIsRegistered)
 				{
 					// Don't report failure if a matching class with a matching major version and higher minor version exists (it will be autoupdated) 
 					FMetasoundFrontendClass FrontendClass;
@@ -601,19 +645,14 @@ namespace Metasound::Engine
 						AssetKey.Version.Major == FrontendClass.Metadata.GetVersion().Major &&
 						AssetKey.Version.Minor < FrontendClass.Metadata.GetVersion().Minor))
 					{
-						bReportFail = true;
-					}
-				}
-
-				if (bReportFail)
-				{
-					if (bIsInitialAssetScanComplete)
-					{
-						UE_LOG(LogMetaSound, Warning, TEXT("MetaSound Node Class with registry key '%s' not registered when gathering referenced asset classes from '%s': Retrieving all asset classes may not be comprehensive."), *AssetKey.ToString(), *InAssetBase.GetOwningAssetName());
-					}
-					else
-					{
-						UE_LOG(LogMetaSound, Warning, TEXT("Attempt to get registered dependent asset with key '%s' from MetaSound asset '%s' before asset scan has completed: Asset class cannot be provided"), *AssetKey.ToString(), *InAssetBase.GetOwningAssetName());
+						if (IsInitialAssetScanComplete())
+						{
+							UE_LOG(LogMetaSound, Warning, TEXT("MetaSound Node Class with registry key '%s' not registered when gathering referenced asset classes from '%s': Retrieving all asset classes may not be comprehensive."), *AssetKey.ToString(), *InAssetBase.GetOwningAssetName());
+						}
+						else
+						{
+							UE_LOG(LogMetaSound, Warning, TEXT("Attempt to get registered dependent asset with key '%s' from MetaSound asset '%s' before asset scan has completed: Asset class cannot be provided"), *AssetKey.ToString(), *InAssetBase.GetOwningAssetName());
+						}
 					}
 				}
 			}
@@ -623,6 +662,8 @@ namespace Metasound::Engine
 
 	bool FMetaSoundAssetManager::ReassignClassName(TScriptInterface<IMetaSoundDocumentInterface> DocInterface)
 	{
+		using namespace Frontend;
+
 #if WITH_EDITORONLY_DATA
 		UObject* MetaSoundObject = DocInterface.GetObject();
 		if (!MetaSoundObject)
@@ -645,14 +686,13 @@ namespace Metasound::Engine
 
 		{
 			const FAssetKey OldAssetKey(ClassMetadata.GetClassName(), ClassMetadata.GetVersion());
-			AssetSubsystemPrivate::RemovePath(&PathMapCriticalSection, PathMap, OldAssetKey, Path);
+			AssetSubsystemPrivate::RemoveAssetTagData(TagDataMapCriticalSection, ClassInfoMap, OldAssetKey, Path);
 		}
 
 		Builder.GenerateNewClassName();
 
 		{
-			const FAssetKey NewAssetKey(ClassMetadata.GetClassName(), ClassMetadata.GetVersion());
-			AssetSubsystemPrivate::AddPath(&PathMapCriticalSection, PathMap, NewAssetKey, Path);
+			AssetSubsystemPrivate::AddAssetTagData(TagDataMapCriticalSection, ClassInfoMap, FAssetClassQueryInfo(*DocInterface));
 		}
 
 		AssetBase->UpdateAndRegisterForExecution();
@@ -664,14 +704,111 @@ namespace Metasound::Engine
 	}
 #endif // WITH_EDITOR
 
-	void FMetaSoundAssetManager::IterateAssets(TFunctionRef<void(const FAssetKey, const TArray<FTopLevelAssetPath>&)> Iter) const
+	bool FMetaSoundAssetManager::IsAssetClass(const FMetasoundFrontendClassMetadata& ClassMetadata) const
 	{
-		for (const TPair<FAssetKey, TArray<FTopLevelAssetPath>>& Pair : PathMap)
+		const EMetasoundFrontendClassType& ClassType = ClassMetadata.GetType();
+		if (ClassType != EMetasoundFrontendClassType::External && ClassType != EMetasoundFrontendClassType::Graph)
 		{
-			Iter(Pair.Key, Pair.Value);
+			return false;
+		}
+
+		return ContainsKey(FAssetKey(ClassMetadata));
+	}
+
+#if WITH_EDITOR
+	void FMetaSoundAssetManager::IterateAssetTagData(TFunctionRef<void(Frontend::FAssetClassQueryInfo)> Iter, bool bIterateDuplicates) const
+	{
+		using namespace Frontend;
+
+		TArray<TArray<FAssetClassQueryInfo>> TagDataMatrix;
+		{
+			FScopeLock Lock(&TagDataMapCriticalSection);
+			ClassInfoMap.GenerateValueArray(TagDataMatrix);
+		}
+
+		if (bIterateDuplicates)
+		{
+			for (TArray<FAssetClassQueryInfo>& TagDataArray : TagDataMatrix)
+			{
+				for (FAssetClassQueryInfo& QueryInfo : TagDataArray)
+				{
+					Iter(MoveTemp(QueryInfo));
+				}
+			}
+		}
+		else
+		{
+			for (TArray<FAssetClassQueryInfo>& TagDataArray : TagDataMatrix)
+			{
+				if (!TagDataArray.IsEmpty())
+				{
+					Iter(MoveTemp(TagDataArray.Last()));
+				}
+			}
 		}
 	}
 
+	void FMetaSoundAssetManager::IterateReferences(const FAssetKey& InKey, TFunctionRef<void(const FAssetKey&)> VisitFunction) const
+	{
+		DepthFirstVisit_AssetKey(InKey, [&](const FAssetKey& ReferencedKey)
+		{
+			using namespace Frontend;
+
+			TArray<FAssetClassQueryInfo> TagDatas;
+			{
+				FScopeLock Lock(&TagDataMapCriticalSection);
+				if (const TArray<FAssetClassQueryInfo>* TagDataArray = ClassInfoMap.Find(ReferencedKey))
+				{
+					TagDatas = *TagDataArray;
+				}
+			}
+
+			TSet<FAssetKey> Children;
+			for (FAssetClassQueryInfo& QueryInfo : TagDatas)
+			{
+				// If loaded, this is the "freshest" reference list, so use that.
+				if (const FMetasoundAssetBase* LoadedAsset = FindAsset(ReferencedKey))
+				{
+					const FMetasoundFrontendDocument& Document = LoadedAsset->GetConstDocumentChecked();
+					for (const FMetasoundFrontendClass& Class : Document.Dependencies)
+					{
+						if (Class.Metadata.GetType() == EMetasoundFrontendClassType::External)
+						{
+							bool bIsAsset = false;
+							
+							{
+								FScopeLock Lock(&TagDataMapCriticalSection);
+								bIsAsset = ClassInfoMap.Contains(Class.Metadata);
+							}
+
+							if (bIsAsset)
+							{
+								const FAssetKey RefKey = FAssetKey(Class.Metadata);
+								VisitFunction(RefKey);
+								Children.Add(RefKey);
+							}
+						}
+					}
+
+				}
+				// Otherwise, use provided reference list from the asset tag data
+				else
+				{
+					for (const FAssetKey& RefKey : QueryInfo.ReferencedAssetKeys)
+					{
+						VisitFunction(RefKey);
+					}
+					Children.Append(MoveTemp(QueryInfo.ReferencedAssetKeys));
+				}
+			}
+
+			VisitFunction(ReferencedKey);
+			Children.Add(ReferencedKey);
+			return Children;
+		});
+	}
+#endif // WITH_EDITOR
+
 	void FMetaSoundAssetManager::RebuildDenyListCache(const UAssetManager& InAssetManager)
 	{
 		using namespace Frontend;
@@ -696,11 +833,10 @@ namespace Metasound::Engine
 			FAssetData AssetData;
 			if (AssetManager.GetAssetDataForPath(UpdateSettings.MetaSound, AssetData))
 			{
-				FString AssetClassID;
-				if (AssetData.GetTagValue(AssetTags::AssetClassID, AssetClassID))
+				const FAssetClassQueryInfo QueryInfo(AssetData);
+				if (QueryInfo.bIsValid)
 				{
-					const FMetasoundFrontendClassName ClassName = { FName(), *AssetClassID, FName() };
-					AutoUpdateDenyListCache.Add(ClassName.GetFullName());
+					AutoUpdateDenyListCache.Add(QueryInfo.ClassName.GetFullName());
 				}
 			}
 		}
@@ -715,16 +851,18 @@ namespace Metasound::Engine
 
 		SearchAndIterateDirectoryAssets(Directories, [this](const FAssetData& AssetData)
 		{
-			AddOrUpdateAsset(AssetData);
-			FMetasoundAssetBase* MetaSoundAsset = Metasound::IMetasoundUObjectRegistry::Get().GetObjectAsAssetBase(AssetData.GetAsset());
-			check(MetaSoundAsset);
-
-			Frontend::FMetaSoundAssetRegistrationOptions RegOptions;
-			if (const UMetaSoundSettings* Settings = GetDefault<UMetaSoundSettings>())
+			AddOrLoadAndUpdateFromObjectAsync(AssetData, [](FAssetKey, UObject& AssetObject)
 			{
-				RegOptions.bAutoUpdateLogWarningOnDroppedConnection = Settings->bAutoUpdateLogWarningOnDroppedConnection;
-			}
-			MetaSoundAsset->UpdateAndRegisterForExecution(RegOptions);
+				FMetasoundAssetBase* MetaSoundAsset = Metasound::IMetasoundUObjectRegistry::Get().GetObjectAsAssetBase(&AssetObject);
+				check(MetaSoundAsset);
+
+				Frontend::FMetaSoundAssetRegistrationOptions RegOptions;
+				if (const UMetaSoundSettings* Settings = GetDefault<UMetaSoundSettings>())
+				{
+					RegOptions.bAutoUpdateLogWarningOnDroppedConnection = Settings->bAutoUpdateLogWarningOnDroppedConnection;
+				}
+				MetaSoundAsset->UpdateAndRegisterForExecution(RegOptions);
+			});
 		});
 	}
 
@@ -745,25 +883,24 @@ namespace Metasound::Engine
 		}
 
 		const FAssetKey AssetKey(Metadata.GetClassName(), Metadata.GetVersion());
-		AssetSubsystemPrivate::RemovePath(&PathMapCriticalSection, PathMap, AssetKey, AssetPath);
+		AssetSubsystemPrivate::RemoveAssetTagData(TagDataMapCriticalSection, ClassInfoMap, AssetKey, AssetPath);
 	}
 
 	void FMetaSoundAssetManager::RemoveAsset(const FAssetData& InAssetData)
 	{
 		using namespace Frontend;
 
-		FNodeClassInfo ClassInfo;
-		const FTopLevelAssetPath AssetPath(InAssetData.PackageName, InAssetData.AssetName);
-		if (ensureAlways(AssetSubsystemPrivate::GetAssetClassInfo(InAssetData, ClassInfo)))
+		const FAssetClassQueryInfo QueryInfo(InAssetData);
+		if (ensureAlways(QueryInfo.bIsValid))
 		{
 			if (IDocumentBuilderRegistry* BuilderRegistry = IDocumentBuilderRegistry::Get())
 			{
 				constexpr bool bForceUnregister = true;
-				BuilderRegistry->FinishBuilding(ClassInfo.ClassName, AssetPath, bForceUnregister);
+				BuilderRegistry->FinishBuilding(QueryInfo.ClassName, QueryInfo.AssetPath, bForceUnregister);
 			}
 
-			const FAssetKey AssetKey(ClassInfo.ClassName, ClassInfo.Version);
-			AssetSubsystemPrivate::RemovePath(&PathMapCriticalSection, PathMap, AssetKey, AssetPath);
+			const FAssetKey AssetKey(QueryInfo.ClassName, QueryInfo.Version);
+			AssetSubsystemPrivate::RemoveAssetTagData(TagDataMapCriticalSection, ClassInfoMap, AssetKey, QueryInfo.AssetPath);
 		}
 	}
 
@@ -783,20 +920,13 @@ namespace Metasound::Engine
 		FMetasoundAssetBase* MetaSoundAsset = GetAsAsset(*InAssetData.GetAsset());
 		check(MetaSoundAsset);
 
-		FNodeClassInfo ClassInfo;
-		if (ensureAlways(AssetSubsystemPrivate::GetAssetClassInfo(InAssetData, ClassInfo)))
+		FAssetClassQueryInfo QueryInfo(InAssetData);
+		if (ensureAlways(QueryInfo.bIsValid))
 		{
-			const FAssetKey AssetKey(ClassInfo.ClassName, ClassInfo.Version);
+			const FAssetKey AssetKey(QueryInfo.ClassName, QueryInfo.Version);
 			const FTopLevelAssetPath OldPath(InOldObjectPath);
-			AssetSubsystemPrivate::RemovePath(&PathMapCriticalSection, PathMap, AssetKey, OldPath);
-
-			if (ClassInfo.AssetClassID.IsValid())
-			{
-				if (AssetKey.IsValid())
-				{
-					AssetSubsystemPrivate::AddPath(&PathMapCriticalSection, PathMap, AssetKey, ClassInfo.AssetPath);
-				}
-			}
+			AssetSubsystemPrivate::RemoveAssetTagData(TagDataMapCriticalSection, ClassInfoMap, AssetKey, OldPath);
+			AssetSubsystemPrivate::AddAssetTagData(TagDataMapCriticalSection, ClassInfoMap, MoveTemp(QueryInfo));
 		}
 	}
 
@@ -873,10 +1003,10 @@ namespace Metasound::Engine
 
 				auto AssetsLoadedDelegate = [this, AsyncLoadID]()
 				{
-					this->OnAssetsLoaded(AsyncLoadID);
+					this->OnReferencedAssetsLoaded(AsyncLoadID);
 				};
 
-				// Store async loading data for use when async load is complete. 
+				// Store async loading data for use when async load is complete.
 				FMetaSoundAsyncAssetDependencies& AsyncDependencies = LoadingDependencies.AddDefaulted_GetRef();
 
 				AsyncDependencies.LoadID = AsyncLoadID;
@@ -894,30 +1024,34 @@ namespace Metasound::Engine
 
 	void FMetaSoundAssetManager::ReloadMetaSoundAssets() const
 	{
+		using namespace Frontend;
+
+		FScopeLock Lock(&TagDataMapCriticalSection);
+
+		IMetasoundUObjectRegistry& ObjectRegistry = IMetasoundUObjectRegistry::Get();
 		TSet<FMetasoundAssetBase*> ToReregister;
-		IterateAssets([&ToReregister](const FAssetKey& AssetKey, const TArray<FTopLevelAssetPath>& Paths)
+		for (const TPair<FAssetKey, TArray<FAssetClassQueryInfo>>& Pair : ClassInfoMap)
 		{
-			if (FMetasoundAssetBase* Asset = IMetaSoundAssetManager::GetChecked().FindAsset(AssetKey))
+			if (!Pair.Value.IsEmpty())
 			{
-				if (Asset->IsRegistered())
+				const FAssetClassQueryInfo& QueryInfo = Pair.Value.Last();
+				if (UObject* Object = FSoftObjectPath(QueryInfo.AssetPath, { }).ResolveObject())
 				{
-					ToReregister.Add(Asset);
-					Asset->UnregisterGraphWithFrontend();
+					if (FMetasoundAssetBase* Asset = ObjectRegistry.GetObjectAsAssetBase(Object); Asset && Asset->IsRegistered())
+					{
+						ToReregister.Add(Asset);
+						Asset->UnregisterGraphWithFrontend();
+					}
 				}
 			}
-		});
+		}
 
 		// Handled in second loop to avoid re-registering referenced graphs more than once
-		IterateAssets([&ToReregister](const FAssetKey& AssetKey, const TArray<FTopLevelAssetPath>& Paths)
+		for (FMetasoundAssetBase* AssetToReregister : ToReregister)
 		{
-			if (FMetasoundAssetBase* Asset = IMetaSoundAssetManager::GetChecked().FindAsset(AssetKey))
-			{
-				if (ToReregister.Contains(Asset))
-				{
-					Asset->UpdateAndRegisterForExecution();
-				}
-			}
-		});
+			check(AssetToReregister);
+			AssetToReregister->UpdateAndRegisterForExecution();
+		}
 	}
 
 	void FMetaSoundAssetManager::SearchAndIterateDirectoryAssets(const TArray<FDirectoryPath>& InDirectories, TFunctionRef<void(const FAssetData&)> InFunction) const
@@ -947,6 +1081,13 @@ namespace Metasound::Engine
 		});
 	}
 	
+#if WITH_EDITOR
+	void FMetaSoundAssetManager::SetCanNotifyAssetTagScanComplete()
+	{
+		bNotifyTagDataScanComplete = true;
+	}
+#endif // WITH_EDITOR
+
 	void FMetaSoundAssetManager::SetLogActiveAssetsOnShutdown(bool bInLogActiveAssetsOnShutdown)
 	{
 		bLogActiveAssetsOnShutdown = bInLogActiveAssetsOnShutdown;
@@ -1008,7 +1149,6 @@ namespace Metasound::Engine
 
 		SearchAndIterateDirectoryAssets(Directories, [this](const FAssetData& AssetData)
 		{
-			using namespace Metasound;
 			using namespace Frontend;
 
 			if (AssetData.IsAssetLoaded())
@@ -1021,17 +1161,17 @@ namespace Metasound::Engine
 			}
 			else
 			{
-				FNodeClassInfo AssetClassInfo;
-				if (ensureAlways(AssetSubsystemPrivate::GetAssetClassInfo(AssetData, AssetClassInfo)))
+				const FAssetClassQueryInfo ClassQueryInfo(AssetData);
+				if (ensureAlways(ClassQueryInfo.bIsValid))
 				{
-					const FNodeRegistryKey RegistryKey = FNodeRegistryKey(AssetClassInfo);
+					const FAssetKey AssetKey(ClassQueryInfo.ClassName, ClassQueryInfo.Version);
+					const FNodeRegistryKey RegistryKey = FNodeRegistryKey(AssetKey);
 					const bool bIsRegistered = FMetasoundFrontendRegistryContainer::Get()->IsNodeRegistered(RegistryKey);
 					if (bIsRegistered)
 					{
 						FMetasoundFrontendRegistryContainer::Get()->UnregisterNode(RegistryKey);
 						const FTopLevelAssetPath AssetPath(AssetData.PackageName, AssetData.AssetName);
-						const FAssetKey AssetKey(AssetClassInfo.ClassName, AssetClassInfo.Version);
-						AssetSubsystemPrivate::RemovePath(&PathMapCriticalSection, PathMap, AssetKey, AssetPath);
+						AssetSubsystemPrivate::RemoveAssetTagData(TagDataMapCriticalSection, ClassInfoMap, AssetKey, AssetPath);
 					}
 				}
 			}
@@ -1126,7 +1266,7 @@ namespace Metasound::Engine
 					else
 					{
 						// This will remove the loading dependencies from internal storage
-						OnAssetsLoaded(LoadingDependency->LoadID);
+						OnReferencedAssetsLoaded(LoadingDependency->LoadID);
 					}
 
 					// This will prevent OnAssetsLoaded from being called via the streamables
@@ -1176,6 +1316,7 @@ void UMetaSoundAssetSubsystem::PostEngineInitInternal()
 void UMetaSoundAssetSubsystem::PostInitAssetScanInternal()
 {
 	using namespace Metasound::Engine;
+	using namespace Metasound::Frontend;
 
 	METASOUND_TRACE_CPUPROFILER_EVENT_SCOPE(UMetaSoundAssetSubsystem::PostInitAssetScanInternal);
 
@@ -1185,9 +1326,22 @@ void UMetaSoundAssetSubsystem::PostInitAssetScanInternal()
 		FMetaSoundAssetManager& Manager = FMetaSoundAssetManager::GetChecked();
 		Manager.SearchAndIterateDirectoryAssets(Settings->DirectoriesToRegister, [&Manager](const FAssetData& AssetData)
 		{
-			Manager.AddOrUpdateAsset(AssetData);
+			using namespace Metasound;
+
+			Manager.AddOrLoadAndUpdateFromObjectAsync(AssetData, [](FAssetKey, UObject& AssetObject)
+			{
+				FMetasoundAssetBase* MetaSoundAsset = IMetasoundUObjectRegistry::Get().GetObjectAsAssetBase(&AssetObject);
+				check(MetaSoundAsset);
+
+				Frontend::FMetaSoundAssetRegistrationOptions RegOptions;
+				if (const UMetaSoundSettings* Settings = GetDefault<UMetaSoundSettings>())
+				{
+					RegOptions.bAutoUpdateLogWarningOnDroppedConnection = Settings->bAutoUpdateLogWarningOnDroppedConnection;
+				}
+				MetaSoundAsset->UpdateAndRegisterForExecution(RegOptions);
+			});
 		});
-		Manager.OnAssetScanComplete();
+		AssetSubsystemPrivate::InitialAssetScanComplete.store(true);
 	}
 }
 
@@ -1202,13 +1356,14 @@ void UMetaSoundAssetSubsystem::AddAssetReferences(FMetasoundAssetBase& InAssetBa
 Metasound::Frontend::FNodeRegistryKey UMetaSoundAssetSubsystem::AddOrUpdateAsset(const UObject& InObject)
 {
 	using namespace Metasound::Frontend;
-	return IMetaSoundAssetManager::GetChecked().AddOrUpdateAsset(InObject);
+	return IMetaSoundAssetManager::GetChecked().AddOrUpdateFromObject(InObject);
 }
 
 Metasound::Frontend::FNodeRegistryKey UMetaSoundAssetSubsystem::AddOrUpdateAsset(const FAssetData& InAssetData)
 {
 	using namespace Metasound::Frontend;
-	return IMetaSoundAssetManager::GetChecked().AddOrUpdateAsset(InAssetData);
+	IMetaSoundAssetManager::GetChecked().AddOrUpdateFromAssetTags(InAssetData);
+	return { };
 }
 
 bool UMetaSoundAssetSubsystem::CanAutoUpdate(const FMetasoundFrontendClassName& InClassName) const
diff --git a/Engine/Plugins/Runtime/Metasound/Source/MetasoundEngine/Private/MetasoundEngineAsset.h b/Engine/Plugins/Runtime/Metasound/Source/MetasoundEngine/Private/MetasoundEngineAsset.h
--- a/Engine/Plugins/Runtime/Metasound/Source/MetasoundEngine/Private/MetasoundEngineAsset.h
+++ b/Engine/Plugins/Runtime/Metasound/Source/MetasoundEngine/Private/MetasoundEngineAsset.h
@@ -4,9 +4,11 @@
 
 #include "Metasound.h"
 #include "MetasoundAssetManager.h"
+#include "MetasoundAssetSubsystem.h"
 #include "MetasoundDocumentBuilderRegistry.h"
 #include "MetasoundEngineModule.h"
 #include "MetasoundFrontendDocumentIdGenerator.h"
+#include "MetasoundFrontendQuery.h"
 #include "MetasoundFrontendRegistryKey.h"
 #include "MetasoundGlobals.h"
 #include "MetasoundSettings.h"
@@ -14,19 +16,23 @@
 #include "Misc/App.h"
 #include "Modules/ModuleManager.h"
 #include "Serialization/Archive.h"
+#include "UObject/AssetRegistryTagsContext.h"
 
 #if WITH_EDITORONLY_DATA
 #include "Algo/Transform.h"
 #include "Interfaces/ITargetPlatform.h"
 #include "MetasoundFrontendRegistryContainer.h"
 #include "Misc/DataValidation.h"
+#include "Serialization/JsonWriter.h"
 #include "UObject/GarbageCollection.h"
 #include "UObject/ObjectMacros.h"
+#include "UObject/Object.h"
 #include "UObject/StrongObjectPtrTemplates.h"
 #endif // WITH_EDITORONLY_DATA
 
 #define LOCTEXT_NAMESPACE "MetasoundEngine"
 
+
 namespace Metasound::Engine
 {
 	/** MetaSound Engine Asset helper provides routines for UObject based MetaSound assets. 
@@ -46,7 +52,7 @@ namespace Metasound::Engine
 			FDocumentBuilderRegistry::GetChecked().SetEventLogVerbosity(FDocumentBuilderRegistry::ELogEvent::DuplicateEntries, ELogVerbosity::NoLogging);
 		}
 
-		static void PostDuplicate(TScriptInterface<IMetaSoundDocumentInterface> MetaSound, EDuplicateMode::Type InDuplicateMode, FGuid& OutAssetClassID)
+		static void PostDuplicate(TScriptInterface<IMetaSoundDocumentInterface> MetaSound, EDuplicateMode::Type InDuplicateMode)
 		{
 			using namespace Engine;
 			using namespace Frontend;
@@ -61,8 +67,7 @@ namespace Metasound::Engine
 
 				FMetaSoundFrontendDocumentBuilder& DocBuilder = DuplicateBuilder.GetBuilder();
 				const FMetasoundFrontendClassName DuplicateName = DocBuilder.GetConstDocumentChecked().RootGraph.Metadata.GetClassName();
-				const FMetasoundFrontendClassName NewName = DocBuilder.GenerateNewClassName();
-				ensureAlwaysMsgf(IMetaSoundAssetManager::GetChecked().TryGetAssetIDFromClassName(NewName, OutAssetClassID), TEXT("Failed to retrieve newly duplicated MetaSoundClassName AssetID"));
+				DocBuilder.GenerateNewClassName();
 
 				constexpr bool bForceUnregisterNodeClass = true;
 				BuilderRegistry.FinishBuilding(DuplicateName, MetaSound->GetAssetPathChecked(), bForceUnregisterNodeClass);
@@ -112,23 +117,6 @@ namespace Metasound::Engine
 			using namespace Metasound::Engine;
 
 			EDataValidationResult Result = EDataValidationResult::Valid;
-
-			// Need to prime asset registry to look for duplicate class names
-			IMetasoundEngineModule& MetaSoundEngineModule = FModuleManager::GetModuleChecked<IMetasoundEngineModule>("MetaSoundEngine");
-			// Checking for duplicate class names only requires the asset manager to be primed, but not for assets to be loaded. 
-			if (!MetaSoundEngineModule.IsAssetManagerPrimed())
-			{
-				MetaSoundEngineModule.PrimeAssetManager();
-				// Check again, as priming relies on the asset registry being loaded so may not be complete
-				if (!MetaSoundEngineModule.IsAssetManagerPrimed())
-				{
-					Result = EDataValidationResult::Invalid;
-					InOutContext.AddError(LOCTEXT("UniqueClassNameAssetManagerNotReady",
-						"MetaSound Asset Manager was unable to be primed to check for unique class names. This may be because the asset registry has not finished loading assets. Please try again later."));
-					return Result;
-				}
-			}
-
 			IMetaSoundAssetManager& AssetManager = IMetaSoundAssetManager::GetChecked();
 			// Validation has added assets to the asset manager
 			// and we don't remove them immediately after validation to optimize possible subsequent validation
@@ -155,10 +143,10 @@ namespace Metasound::Engine
 
 		static EDataValidationResult IsDataValid(const UObject& MetaSound, const FMetasoundFrontendDocument& Document, FDataValidationContext& InOutContext)
 		{
-			using namespace Metasound::Engine;
+			using namespace Metasound;
 
 			EDataValidationResult Result = EDataValidationResult::Valid;
-			if (MetasoundEngineModulePrivate::EnableMetaSoundEditorAssetValidation)
+			if (Engine::GetEditorAssetValidationEnabled())
 			{
 				// We cannot rely on the asset registry scan being complete during the call
 				// to IsDataValid(...) while running a cook commandlet. The IMetasoundAssetManager
@@ -210,6 +198,21 @@ namespace Metasound::Engine
 
 #endif // WITH_EDITOR
 
+		static void GetAssetRegistryTags(TScriptInterface<const IMetaSoundDocumentInterface> DocInterface, FAssetRegistryTagsContext& Context)
+		{
+			using namespace Frontend;
+
+			const UObject* MetaSound = DocInterface.GetObject();
+			check(MetaSound);
+			if (MetaSound->GetFlags() & (RF_Transient | RF_ClassDefaultObject))
+			{
+				return;
+			}
+
+			const FAssetClassQueryInfo QueryInfo(*DocInterface);
+			QueryInfo.ExportToContext(Context);
+		}
+
 		template <typename TMetaSoundObject>
 		static FTopLevelAssetPath GetAssetPathChecked(TMetaSoundObject& InMetaSound)
 		{
@@ -347,49 +350,6 @@ namespace Metasound::Engine
 				}
 			}
 		}
-
-#if WITH_EDITORONLY_DATA
-		template <typename TMetaSoundObject>
-		static void SetMetaSoundRegistryAssetClassInfo(TMetaSoundObject& InMetaSound, const Frontend::FNodeClassInfo& InClassInfo)
-		{
-			using namespace Frontend;
-
-			check(AssetTags::AssetClassID == GET_MEMBER_NAME_CHECKED(TMetaSoundObject, AssetClassID));
-			check(AssetTags::IsPreset == GET_MEMBER_NAME_CHECKED(TMetaSoundObject, bIsPreset));
-			check(AssetTags::RegistryInputTypes == GET_MEMBER_NAME_CHECKED(TMetaSoundObject, RegistryInputTypes));
-			check(AssetTags::RegistryOutputTypes == GET_MEMBER_NAME_CHECKED(TMetaSoundObject, RegistryOutputTypes));
-			check(AssetTags::RegistryVersionMajor == GET_MEMBER_NAME_CHECKED(TMetaSoundObject, RegistryVersionMajor));
-			check(AssetTags::RegistryVersionMinor == GET_MEMBER_NAME_CHECKED(TMetaSoundObject, RegistryVersionMinor));
-
-			bool bMarkDirty = InMetaSound.AssetClassID != InClassInfo.AssetClassID;
-			bMarkDirty |= InMetaSound.RegistryVersionMajor != InClassInfo.Version.Major;
-			bMarkDirty |= InMetaSound.RegistryVersionMinor != InClassInfo.Version.Minor;
-			bMarkDirty |= InMetaSound.bIsPreset != InClassInfo.bIsPreset;
-
-			InMetaSound.AssetClassID = InClassInfo.AssetClassID;
-			InMetaSound.RegistryVersionMajor = InClassInfo.Version.Major;
-			InMetaSound.RegistryVersionMinor = InClassInfo.Version.Minor;
-			InMetaSound.bIsPreset = InClassInfo.bIsPreset;
-
-			{
-				TArray<FString> InputTypes;
-				Algo::Transform(InClassInfo.InputTypes, InputTypes, [](const FName& Name) { return Name.ToString(); });
-
-				const FString TypeString = FString::Join(InputTypes, *AssetTags::ArrayDelim);
-				bMarkDirty |= InMetaSound.RegistryInputTypes != TypeString;
-				InMetaSound.RegistryInputTypes = TypeString;
-			}
-
-			{
-				TArray<FString> OutputTypes;
-				Algo::Transform(InClassInfo.OutputTypes, OutputTypes, [](const FName& Name) { return Name.ToString(); });
-
-				const FString TypeString = FString::Join(OutputTypes, *AssetTags::ArrayDelim);
-				bMarkDirty |= InMetaSound.RegistryOutputTypes != TypeString;
-				InMetaSound.RegistryOutputTypes = TypeString;
-			}
-		}
-#endif // WITH_EDITORONLY_DATA
 	};
 } // namespace Metasound::Engine
 #undef LOCTEXT_NAMESPACE // MetasoundEngine
diff --git a/Engine/Plugins/Runtime/Metasound/Source/MetasoundEngine/Private/MetasoundEngineModule.cpp b/Engine/Plugins/Runtime/Metasound/Source/MetasoundEngine/Private/MetasoundEngineModule.cpp
--- a/Engine/Plugins/Runtime/Metasound/Source/MetasoundEngine/Private/MetasoundEngineModule.cpp
+++ b/Engine/Plugins/Runtime/Metasound/Source/MetasoundEngine/Private/MetasoundEngineModule.cpp
@@ -54,9 +54,10 @@ REGISTER_METASOUND_DATATYPE(Metasound::FWaveTableBankAsset, "WaveTableBankAsset"
 namespace Metasound::Engine
 {
 #if WITH_EDITOR
-	namespace MetasoundEngineModulePrivate
+	namespace ModulePrivate
 	{
 		int32 EnableMetaSoundEditorAssetValidation = 1;
+		int32 EnableMetaSoundEditorAssetAutoLoadAndRegister = 0;
 
 		FAutoConsoleVariableRef CVarEnableMetaSoundEditorAssetValidation(
 			TEXT("au.MetaSound.Editor.EnableAssetValidation"),
@@ -64,7 +65,19 @@ namespace Metasound::Engine
 			TEXT("Enables MetaSound specific asset validation.\n")
 			TEXT("Default: 1 (Enabled)"),
 			ECVF_Default);
-	} // namespace MetasoundEngineModulePrivate
+
+		FAutoConsoleVariableRef CVarEnableMetaSoundAutoLoadingAndRegisteringAssets(
+			TEXT("au.MetaSound.Editor.EnableAutoLoadAndRegisterOnAssetScan"),
+			EnableMetaSoundEditorAssetAutoLoadAndRegister,
+			TEXT("Enables auto-loading and registration of assets. Not recommended as it is slow, but can be useful for debugging load issues with serialized MetaSound assets. \n")
+			TEXT("Default: 0 (Disabled)"),
+			ECVF_Default);
+	} // namespace ModulePrivate
+
+	bool GetEditorAssetValidationEnabled()
+	{
+		return ModulePrivate::EnableMetaSoundEditorAssetValidation != 0;
+	}
 #endif // WITH_EDITOR
 
 	void FModule::FModule::StartupModule() 
@@ -145,12 +158,10 @@ namespace Metasound::Engine
 		FAssetRegistryModule& AssetRegistryModule = FModuleManager::LoadModuleChecked<FAssetRegistryModule>("AssetRegistry");
 		if (AssetRegistryModule.Get().IsLoadingAssets())
 		{
-			AssetScanStatus = EAssetScanStatus::InProgress;
 			AssetRegistryModule.Get().OnFilesLoaded().AddRaw(this, &FModule::OnAssetScanFinished);
 		}
 		else
 		{
-			AssetScanStatus = EAssetScanStatus::Complete;
 		}
 #endif // WITH_EDITOR
 
@@ -160,11 +171,6 @@ namespace Metasound::Engine
 	void FModule::FModule::ShutdownModule()
 	{
 #if WITH_EDITOR
-		if (NodeClassRegistryPrimeStatus != ENodeClassRegistryPrimeStatus::NotRequested && NodeClassRegistryPrimeStatus != ENodeClassRegistryPrimeStatus::Complete)
-		{
-			NodeClassRegistryPrimeStatus = ENodeClassRegistryPrimeStatus::Canceled;
-		}
-
 		ShutdownAssetClassRegistry();
 #endif // WITH_EDITOR
 		DeinitializeAssetManager();
@@ -172,68 +178,41 @@ namespace Metasound::Engine
 	}
 
 #if WITH_EDITOR
-	void FModule::LoadAndRegisterAsset(const FAssetData& InAssetData)
+	void FModule::AddClassRegistryAsset(const FAssetData& InAssetData)
 	{
-		// Ignore requests if graphs cannot be executed, as registration that results in IGraph generation is not supported.
-		if (!Metasound::CanEverExecuteGraph())
+		using namespace Frontend;
+
+		// Don't add temporary assets used for diffing
+		if (InAssetData.HasAnyPackageFlags(PKG_ForDiffing))
 		{
 			return;
 		}
 
-		Frontend::FMetaSoundAssetRegistrationOptions RegOptions;
-		RegOptions.bForceReregister = false;
-		if (const UMetaSoundSettings* Settings = GetDefault<UMetaSoundSettings>())
+		// If an object's class could not be found, ignore this asset.  This can hit for non-MetaSound assets
+		// and it is up to the system in charge of interacting with that asset or the loading behavior to
+		// report the failed load of the class.
+		const UClass* AssetClass = InAssetData.GetClass();
+		if (!AssetClass || !IMetasoundUObjectRegistry::Get().IsRegisteredClass(*AssetClass))
 		{
-			RegOptions.bAutoUpdateLogWarningOnDroppedConnection = Settings->bAutoUpdateLogWarningOnDroppedConnection;
+			return;
 		}
 
-		if (InAssetData.IsAssetLoaded())
+		if (ModulePrivate::EnableMetaSoundEditorAssetAutoLoadAndRegister)
 		{
-			if (UObject* AssetObject = InAssetData.GetAsset())
+			IMetaSoundAssetManager::GetChecked().AddOrLoadAndUpdateFromObjectAsync(InAssetData,
+			[](FAssetKey, UObject& AssetObject)
 			{
-				OnGraphRegister.ExecuteIfBound(*AssetObject, ERegistrationAssetContext::None);
-			}
+				FModule& ThisModule = FModuleManager::GetModuleChecked<Metasound::Engine::FModule>("MetaSoundEngine");
+				ThisModule.GetOnGraphRegisteredDelegate().ExecuteIfBound(AssetObject, ERegistrationAssetContext::None);
+			});
 		}
 		else
 		{
-			if (NodeClassRegistryPrimeStatus == ENodeClassRegistryPrimeStatus::NotRequested || NodeClassRegistryPrimeStatus == ENodeClassRegistryPrimeStatus::Canceled)
-			{
-				return;
-			}
-
-			ActiveAsyncAssetLoadRequests++;
-
-			FSoftObjectPath AssetPath = InAssetData.ToSoftObjectPath();
-			auto LoadAndRegister = [this, ObjectPath = AssetPath, RegOptions](const FName& PackageName, UPackage* Package, EAsyncLoadingResult::Type Result)
-			{
-				if (NodeClassRegistryPrimeStatus == ENodeClassRegistryPrimeStatus::Canceled)
-				{
-					return;
-				}
-
-				if (Result == EAsyncLoadingResult::Succeeded)
-				{
-					UObject* MetaSoundObj = ObjectPath.ResolveObject();
-					FMetasoundAssetBase* MetaSoundAsset = IMetasoundUObjectRegistry::Get().GetObjectAsAssetBase(MetaSoundObj);
-					check(MetaSoundAsset);
-					if (!MetaSoundAsset->IsRegistered())
-					{
-						OnGraphRegister.ExecuteIfBound(*MetaSoundObj, ERegistrationAssetContext::None);
-					}
-				}
-
-				ActiveAsyncAssetLoadRequests--;
-				if (NodeClassRegistryPrimeStatus == ENodeClassRegistryPrimeStatus::InProgress && ActiveAsyncAssetLoadRequests == 0)
-				{
-					NodeClassRegistryPrimeStatus = ENodeClassRegistryPrimeStatus::Complete;
-				}
-			};
-
-			LoadPackageAsync(AssetPath.GetLongPackageName(), FLoadPackageAsyncDelegate::CreateLambda(LoadAndRegister));
+			IMetaSoundAssetManager::GetChecked().AddOrUpdateFromAssetTags(InAssetData);
 		}
 	}
-
-	void FModule::AddClassRegistryAsset(const FAssetData& InAssetData)
+	
+	void FModule::UpdateClassRegistryAsset(const FAssetData& InAssetData)
 	{
 		using namespace Frontend;
 
@@ -242,50 +221,26 @@ namespace Metasound::Engine
 		// report the failed load of the class.
 		if (const UClass* AssetClass = InAssetData.GetClass())
 		{
-			// Don't add temporary assets used for diffing
-			if (InAssetData.HasAnyPackageFlags(PKG_ForDiffing))
-			{
-				return;
-			}
 			const bool bIsRegisteredClass = IMetasoundUObjectRegistry::Get().IsRegisteredClass(*AssetClass);
 			if (bIsRegisteredClass)
 			{
-				const FNodeRegistryKey RegistryKey = IMetaSoundAssetManager::GetChecked().AddOrUpdateAsset(InAssetData);
-
-				// Can be invalid if being called for the first time on an asset before its class name is generated
-				if (RegistryKey.IsValid())
+				if (ModulePrivate::EnableMetaSoundEditorAssetAutoLoadAndRegister)
 				{
-					const bool bPrimeRequested = NodeClassRegistryPrimeStatus > ENodeClassRegistryPrimeStatus::NotRequested;
-					const bool bIsRegistered = FMetasoundFrontendRegistryContainer::Get()->IsNodeRegistered(RegistryKey);
-					if (bPrimeRequested && !bIsRegistered)
+					IMetaSoundAssetManager::GetChecked().AddOrLoadAndUpdateFromObjectAsync(InAssetData,
+					[](FAssetKey AssetKey, UObject& AssetObject)
 					{
-						LoadAndRegisterAsset(InAssetData);
-					}
+						// Have to re-register to avoid registry desync.
+						const bool bIsRegistered = FMetasoundFrontendRegistryContainer::Get()->IsNodeRegistered(FNodeRegistryKey(AssetKey));
+						if (bIsRegistered || ModulePrivate::EnableMetaSoundEditorAssetAutoLoadAndRegister)
+						{
+							FModule& ThisModule = FModuleManager::GetModuleChecked<Metasound::Engine::FModule>("MetaSoundEngine");
+							ThisModule.GetOnGraphRegisteredDelegate().ExecuteIfBound(AssetObject, ERegistrationAssetContext::None);
+						}
+					});
 				}
-			}
-		}
-	}
-	
-	void FModule::UpdateClassRegistryAsset(const FAssetData& InAssetData)
-	{
-		using namespace Frontend;
-
-		// If an object's class could not be found, ignore this asset.  This can hit for non-MetaSound assets
-		// and it is up to the system in charge of interacting with that asset or the loading behavior to
-		// report the failed load of the class.
-		if (const UClass* AssetClass = InAssetData.GetClass())
-		{
-			const bool bIsRegisteredClass = IMetasoundUObjectRegistry::Get().IsRegisteredClass(*AssetClass);
-			if (bIsRegisteredClass)
-			{
-				const FNodeRegistryKey RegistryKey = IMetaSoundAssetManager::GetChecked().AddOrUpdateAsset(InAssetData);
-				const bool bPrimeRequested = NodeClassRegistryPrimeStatus > ENodeClassRegistryPrimeStatus::NotRequested;
-				const bool bIsRegistered = FMetasoundFrontendRegistryContainer::Get()->IsNodeRegistered(RegistryKey);
-
-				// Have to re-register even if prime was not requested to avoid registry desync.
-				if (bPrimeRequested || bIsRegistered)
+				else
 				{
-					LoadAndRegisterAsset(InAssetData);
+					IMetaSoundAssetManager::GetChecked().AddOrUpdateFromAssetTags(InAssetData);
 				}
 			}
 		}
@@ -307,15 +262,15 @@ namespace Metasound::Engine
 		}
 
 		auto IsAssetMetaSound = [](const UObject* Obj)
+		{
+			check(Obj);
+			if (const UClass* AssetClass = Obj->GetClass())
 			{
-				check(Obj);
-				if (const UClass* AssetClass = Obj->GetClass())
-				{
-					return IMetasoundUObjectRegistry::Get().IsRegisteredClass(*AssetClass);
-				}
+				return IMetasoundUObjectRegistry::Get().IsRegisteredClass(*AssetClass);
+			}
 
-				return false;
-			};
+			return false;
+		};
 
 		for (const TPair<UObject*, UObject*>& Pair : InPackageReloadedEvent->GetRepointedObjects())
 		{
@@ -332,13 +287,49 @@ namespace Metasound::Engine
 			{
 				if (IsAssetMetaSound(Obj))
 				{
-					IMetaSoundAssetManager::GetChecked().AddOrUpdateAsset(*Pair.Value);
+					IMetaSoundAssetManager::GetChecked().AddOrUpdateFromObject(*Pair.Value);
 					OnGraphRegister.ExecuteIfBound(*Obj, ERegistrationAssetContext::Reloading);
 				}
 			}
 		}
 	}
 
+	void FModule::PrimeAssetManager()
+	{
+		if (!FMetaSoundAssetManager::GetChecked().IsInitialAssetScanComplete())
+		{
+			AssetTagPrimeStatus = EAssetTagPrimeRequestStatus::Requested;
+			return;
+		}
+
+		if (AssetTagPrimeStatus < EAssetTagPrimeRequestStatus::Complete)
+		{
+			PrimeAssetManagerInternal();
+		}
+	}
+
+	void FModule::PrimeAssetManagerInternal()
+	{
+		TArray<FTopLevelAssetPath> ClassNames;
+		IMetasoundUObjectRegistry::Get().IterateRegisteredUClasses([&ClassNames](UClass& InClass)
+		{
+			ClassNames.Add(InClass.GetClassPathName());
+		});
+
+		FARFilter Filter;
+		Filter.ClassPaths = ClassNames;
+
+		FAssetRegistryModule& AssetRegistryModule = FModuleManager::LoadModuleChecked<FAssetRegistryModule>("AssetRegistry");
+		AssetRegistryModule.Get().EnumerateAssets(Filter, [this](const FAssetData& AssetData)
+		{
+			AddClassRegistryAsset(AssetData);
+			return true;
+		});
+
+		AssetTagPrimeStatus = EAssetTagPrimeRequestStatus::Complete;
+		FMetaSoundAssetManager::GetChecked().SetCanNotifyAssetTagScanComplete();
+	}
+
 	void FModule::OnAssetScanFinished()
 	{
 		if (IsRunningCookCommandlet())
@@ -346,11 +337,9 @@ namespace Metasound::Engine
 			return;
 		}
 
-		AssetScanStatus = EAssetScanStatus::Complete;
-
-		if (NodeClassRegistryPrimeStatus == ENodeClassRegistryPrimeStatus::Requested)
+		if (AssetTagPrimeStatus == EAssetTagPrimeRequestStatus::Requested)
 		{
-			PrimeAssetRegistryAsync();
+			PrimeAssetManagerInternal();
 		}
 
 		FAssetRegistryModule& AssetRegistryModule = FModuleManager::LoadModuleChecked<FAssetRegistryModule>("AssetRegistry");
@@ -360,7 +349,7 @@ namespace Metasound::Engine
 		AssetRegistryModule.Get().OnAssetRenamed().AddRaw(this, &FModule::RenameAssetInClassRegistry);
 
 		AssetRegistryModule.Get().OnFilesLoaded().RemoveAll(this);
-
+		
 		FCoreUObjectDelegates::OnPackageReloaded.AddRaw(this, &FModule::OnPackageReloaded);
 	}
 
@@ -435,94 +424,6 @@ namespace Metasound::Engine
 		}
 	}
 
-	void FModule::PrimeAssetRegistryAsync() 
-	{
-		// Ignore step if still loading assets from initial scan but set prime status as requested.
-		if (AssetScanStatus <= EAssetScanStatus::InProgress)
-		{
-			NodeClassRegistryPrimeStatus = ENodeClassRegistryPrimeStatus::Requested;
-			return;
-		}
-		
-		// Prime both asset manager and node class registry
-		if (NodeClassRegistryPrimeStatus != ENodeClassRegistryPrimeStatus::InProgress)
-		{
-			NodeClassRegistryPrimeStatus = ENodeClassRegistryPrimeStatus::InProgress;
-
-			TArray<FTopLevelAssetPath> ClassNames;
-			IMetasoundUObjectRegistry::Get().IterateRegisteredUClasses([&ClassNames](UClass& InClass)
-			{
-				ClassNames.Add(InClass.GetClassPathName());
-			});
-
-			FARFilter Filter;
-			Filter.ClassPaths = ClassNames;
-
-			FAssetRegistryModule& AssetRegistryModule = FModuleManager::LoadModuleChecked<FAssetRegistryModule>("AssetRegistry");
-			uint32 AssetCount = 0;
-			AssetRegistryModule.Get().EnumerateAssets(Filter, [this, &AssetCount](const FAssetData& AssetData)
-			{
-				AddClassRegistryAsset(AssetData);
-				AssetCount++;
-				return true;
-			});
-
-			// Node class registry prime is complete if there are no assets to process
-			if (AssetCount == 0 || ActiveAsyncAssetLoadRequests == 0)
-			{
-				NodeClassRegistryPrimeStatus = ENodeClassRegistryPrimeStatus::Complete;
-			}
-			// Asset manager prime also occurred as part of AddClassRegistryAsset
-			bAssetManagerPrimed = true;
-		}
-	}
-
-	void FModule::PrimeAssetManager()
-	{
-		using namespace Frontend;
-		if (AssetScanStatus <= EAssetScanStatus::InProgress)
-		{
-			UE_LOG(LogMetaSound, Display, TEXT("MetaSound Asset Manager prime requested before Asset Registry scan completed."));
-			return;
-		}
-
-		if (!IsAssetManagerPrimed())
-		{
-			TArray<FTopLevelAssetPath> ClassNames;
-			IMetasoundUObjectRegistry::Get().IterateRegisteredUClasses([&ClassNames](UClass& InClass)
-			{
-				ClassNames.Add(InClass.GetClassPathName());
-			});
-
-			FARFilter Filter;
-			Filter.ClassPaths = ClassNames;
-
-			FAssetRegistryModule& AssetRegistryModule = FModuleManager::LoadModuleChecked<FAssetRegistryModule>("AssetRegistry");
-			AssetRegistryModule.Get().EnumerateAssets(Filter, [this](const FAssetData& AssetData)
-			{
-				IMetaSoundAssetManager::GetChecked().AddOrUpdateAsset(AssetData);
-				return true;
-			});
-
-			bAssetManagerPrimed = true;
-		}
-	}
-
-	ENodeClassRegistryPrimeStatus FModule::GetNodeClassRegistryPrimeStatus() const
-	{
-		return NodeClassRegistryPrimeStatus;
-	}
-
-	EAssetScanStatus FModule::GetAssetRegistryScanStatus() const
-	{
-		return AssetScanStatus;
-	}
-
-	bool FModule::IsAssetManagerPrimed() const
-	{
-		return bAssetManagerPrimed;
-	}
-
 	FOnMetasoundGraphRegister& FModule::GetOnGraphRegisteredDelegate()
 	{
 		return OnGraphRegister;
diff --git a/Engine/Plugins/Runtime/Metasound/Source/MetasoundEngine/Private/MetasoundEngineModule.h b/Engine/Plugins/Runtime/Metasound/Source/MetasoundEngine/Private/MetasoundEngineModule.h
--- a/Engine/Plugins/Runtime/Metasound/Source/MetasoundEngine/Private/MetasoundEngineModule.h
+++ b/Engine/Plugins/Runtime/Metasound/Source/MetasoundEngine/Private/MetasoundEngineModule.h
@@ -12,12 +12,15 @@
 
 namespace Metasound::Engine
 {
-#if WITH_EDITOR
-	namespace MetasoundEngineModulePrivate
+	enum class EAssetTagPrimeRequestStatus : uint8
 	{
-		extern int32 EnableMetaSoundEditorAssetValidation;
-		extern FAutoConsoleVariableRef CVarEnableMetaSoundEditorAssetValidation;
-	} // namespace MetasoundEngineModulePrivate
+		NotRequested = 0,
+		Requested = 1,
+		Complete = 2,
+	};
+
+#if WITH_EDITOR
+	bool GetEditorAssetValidationEnabled();
 #endif // WITH_EDITOR
 
 	class FModule : public IMetasoundEngineModule
@@ -62,38 +65,30 @@ namespace Metasound::Engine
 		virtual void ShutdownModule() override;
 
 #if WITH_EDITOR
-		virtual void PrimeAssetRegistryAsync() override;
 		virtual void PrimeAssetManager() override;
-
-		virtual ENodeClassRegistryPrimeStatus GetNodeClassRegistryPrimeStatus() const override;
-		virtual EAssetScanStatus GetAssetRegistryScanStatus() const override;
-		virtual bool IsAssetManagerPrimed() const override;
-
 		virtual FOnMetasoundGraphRegister& GetOnGraphRegisteredDelegate() override;
 		virtual FOnMetasoundGraphUnregister& GetOnGraphUnregisteredDelegate() override;
 #endif // WITH_EDITOR
 
 	private:
 #if WITH_EDITOR
-		virtual void AddClassRegistryAsset(const FAssetData& InAssetData);
-		virtual void LoadAndRegisterAsset(const FAssetData& InAssetData);
-		virtual void UpdateClassRegistryAsset(const FAssetData& InAssetData);
-		virtual void RemoveAssetFromClassRegistry(const FAssetData& InAssetData);
-		virtual void RenameAssetInClassRegistry(const FAssetData& InAssetData, const FString& InOldObjectPath);
-		virtual void ShutdownAssetClassRegistry();
+		void AddClassRegistryAsset(const FAssetData& InAssetData);
+		void UpdateClassRegistryAsset(const FAssetData& InAssetData);
+		void RemoveAssetFromClassRegistry(const FAssetData& InAssetData);
+		void RenameAssetInClassRegistry(const FAssetData& InAssetData, const FString& InOldObjectPath);
+		void ShutdownAssetClassRegistry();
+
+		void OnPackageReloaded(const EPackageReloadPhase InPackageReloadPhase, FPackageReloadedEvent* InPackageReloadedEvent);
+		void OnAssetScanFinished();
 
-		virtual void OnPackageReloaded(const EPackageReloadPhase InPackageReloadPhase, FPackageReloadedEvent* InPackageReloadedEvent);
-		virtual void OnAssetScanFinished();
+		void PrimeAssetManagerInternal();
 
 		// Asset registry delegates for calling MetaSound editor module 
 		FOnMetasoundGraphRegister OnGraphRegister;
 		FOnMetasoundGraphUnregister OnGraphUnregister;
 
-		// Asset registry state
-		ENodeClassRegistryPrimeStatus NodeClassRegistryPrimeStatus = ENodeClassRegistryPrimeStatus::NotRequested;
-		EAssetScanStatus AssetScanStatus = EAssetScanStatus::NotRequested;
-		bool bAssetManagerPrimed = false;
-		int32 ActiveAsyncAssetLoadRequests = 0;
+		// Prime tag registry state
+		EAssetTagPrimeRequestStatus AssetTagPrimeStatus = EAssetTagPrimeRequestStatus::NotRequested;
 #endif // WITH_EDITOR
 	};
 } // namespace Metasound::Engine
diff --git a/Engine/Plugins/Runtime/Metasound/Source/MetasoundEngine/Private/MetasoundOperatorCacheSubsystem.cpp b/Engine/Plugins/Runtime/Metasound/Source/MetasoundEngine/Private/MetasoundOperatorCacheSubsystem.cpp
--- a/Engine/Plugins/Runtime/Metasound/Source/MetasoundEngine/Private/MetasoundOperatorCacheSubsystem.cpp
+++ b/Engine/Plugins/Runtime/Metasound/Source/MetasoundEngine/Private/MetasoundOperatorCacheSubsystem.cpp
@@ -2,6 +2,7 @@
 #include "MetasoundOperatorCacheSubsystem.h"
 
 #include "AudioMixerDevice.h"
+#include "MetasoundAssetManager.h"
 #include "MetasoundGenerator.h"
 #include "MetasoundGeneratorModule.h"
 #include "MetasoundOperatorCache.h"
@@ -89,6 +90,7 @@ void UMetaSoundCacheSubsystem::PrecacheMetaSoundInternal(UMetaSoundSource* InMet
 {
 	using namespace Audio;
 	using namespace Metasound;
+	using namespace Metasound::Frontend;
 
 	if (!OperatorCachePrivate::bOperatorPrecacheEnabled)
 	{
@@ -133,10 +135,18 @@ void UMetaSoundCacheSubsystem::PrecacheMetaSoundInternal(UMetaSoundSource* InMet
 	// that is found when a MetaSound generator is created. 
 	const UMetaSoundSource&	NoninflatableSource = InMetaSound->FindFirstNoninflatableSource(InitParams->Environment, [](const UMetaSoundSource&){});
 
+	FGuid AssetID;
+	const FMetasoundFrontendClassName& ClassName = InMetaSound->GetConstDocumentChecked().RootGraph.Metadata.GetClassName();
+	if (!IMetaSoundAssetManager::GetChecked().TryGetAssetIDFromClassName(ClassName, AssetID))
+	{
+		UE_LOG(LogMetaSound, Warning, TEXT("Failed to retrieve MetaSoundClassName when precaching operator for MetaSound '%s'"), *InMetaSound->GetPathName());
+		return;
+	}
+
 	TUniquePtr<FOperatorBuildData> Data = MakeUnique<FOperatorBuildData>(
 		  MoveTemp(InitParams.GetValue())
 		, NoninflatableSource.GetGraphRegistryKey()
-		, InMetaSound->AssetClassID
+		, AssetID
 		, InNumInstances
 		, bTouchExisting
 	);
@@ -159,6 +169,7 @@ void UMetaSoundCacheSubsystem::TouchOrPrecacheMetaSound(UMetaSoundSource* InMeta
 void UMetaSoundCacheSubsystem::RemoveCachedOperatorsForMetaSound(UMetaSoundSource* InMetaSound)
 {
 	using namespace Metasound;
+	using namespace Metasound::Frontend;
 
 	// Note: we're not checking the bOperatorPrecacheEnabled cvar here in case it was disable after some sounds had already been cached.
 	// If nothing is cached this will do very little.
@@ -177,6 +188,15 @@ void UMetaSoundCacheSubsystem::RemoveCachedOperatorsForMetaSound(UMetaSoundSourc
 
 	if (TSharedPtr<FOperatorPool> OperatorPool = Module->GetOperatorPool())
 	{
-		OperatorPool->RemoveOperatorsWithAssetClassID(InMetaSound->AssetClassID);
+		FGuid AssetID;
+		const FMetasoundFrontendClassName& ClassName = InMetaSound->GetConstDocumentChecked().RootGraph.Metadata.GetClassName();
+		if (IMetaSoundAssetManager::GetChecked().TryGetAssetIDFromClassName(ClassName, AssetID))
+		{
+			OperatorPool->RemoveOperatorsWithAssetClassID(AssetID);
+		}
+		else
+		{
+			UE_LOG(LogMetaSound, Warning, TEXT("Failed to retrieve MetaSoundClassName when removing cached operator for MetaSound '%s'"), *InMetaSound->GetPathName());
+		}
 	}
 }
diff --git a/Engine/Plugins/Runtime/Metasound/Source/MetasoundEngine/Private/MetasoundSource.cpp b/Engine/Plugins/Runtime/Metasound/Source/MetasoundEngine/Private/MetasoundSource.cpp
--- a/Engine/Plugins/Runtime/Metasound/Source/MetasoundEngine/Private/MetasoundSource.cpp
+++ b/Engine/Plugins/Runtime/Metasound/Source/MetasoundEngine/Private/MetasoundSource.cpp
@@ -451,7 +451,7 @@ void UMetaSoundSource::PreDuplicate(FObjectDuplicationParameters& DupParams)
 void UMetaSoundSource::PostDuplicate(EDuplicateMode::Type InDuplicateMode)
 {
 	Super::PostDuplicate(InDuplicateMode);
-	Metasound::Engine::FAssetHelper::PostDuplicate(this, InDuplicateMode, AssetClassID);
+	Metasound::Engine::FAssetHelper::PostDuplicate(this, InDuplicateMode);
 }
 
 void UMetaSoundSource::PostEditChangeProperty(FPropertyChangedEvent& InEvent)
@@ -745,12 +745,6 @@ FText UMetaSoundSource::GetDisplayName() const
 	FString TypeName = UMetaSoundSource::StaticClass()->GetName();
 	return FMetasoundAssetBase::GetDisplayName(MoveTemp(TypeName));
 }
-
-
-void UMetaSoundSource::SetRegistryAssetClassInfo(const Metasound::Frontend::FNodeClassInfo& InNodeInfo)
-{
-	Metasound::Engine::FAssetHelper::SetMetaSoundRegistryAssetClassInfo(*this, InNodeInfo);
-}
 #endif // WITH_EDITORONLY_DATA
 
 void UMetaSoundSource::PostLoad()
@@ -946,6 +940,12 @@ bool UMetaSoundSource::IsPlayable() const
 	return true;
 }
 
+void UMetaSoundSource::GetAssetRegistryTags(FAssetRegistryTagsContext Context) const
+{
+	Super::GetAssetRegistryTags(Context);
+	Metasound::Engine::FAssetHelper::GetAssetRegistryTags(this, Context);
+}
+
 float UMetaSoundSource::GetDuration() const
 {
 	// This is an unfortunate function required by logic in determining what sounds can be potentially
@@ -1890,6 +1890,7 @@ const UMetaSoundSource& UMetaSoundSource::FindFirstNoninflatableSource(Metasound
 const UMetaSoundSource& UMetaSoundSource::FindFirstNoninflatableSourceInternal(TArray<FGuid>& OutHierarchy, TFunctionRef<void(const UMetaSoundSource&)> OnTraversal) const
 {
 	using namespace Metasound;
+	using namespace Metasound::Frontend;
 
 	const bool bIsDynamic = DynamicTransactor.IsValid();
 	if (!bIsDynamic && ConsoleVariables::bEnableExperimentalRuntimePresetGraphInflation && bIsPresetGraphInflationSupported)
@@ -1903,8 +1904,10 @@ const UMetaSoundSource& UMetaSoundSource::FindFirstNoninflatableSourceInternal(T
 			TObjectPtr<const UMetaSoundSource> BaseMetaSoundSource = Cast<const UMetaSoundSource>(BaseGraph);
 			if (ensure(BaseMetaSoundSource)) // SourcePresets assume they are referencing a UMetaSoundSource
 			{
-				// Preset graph inflation needs to emulate the graph hierarchy for nodes which depend upon accurate graph hierarchies. 
-				OutHierarchy.Add(AssetClassID);
+				// Preset graph inflation needs to emulate the graph hierarchy for nodes which depend upon accurate graph hierarchies.
+				FGuid ClassID;
+				ensureAlways(IMetaSoundAssetManager::GetChecked().TryGetAssetIDFromClassName(RootMetasoundDocument.RootGraph.Metadata.GetClassName(), ClassID));
+				OutHierarchy.Add(ClassID);
 
 				OnTraversal(*this);
 
diff --git a/Engine/Plugins/Runtime/Metasound/Source/MetasoundEngine/Public/IMetasoundEngineModule.h b/Engine/Plugins/Runtime/Metasound/Source/MetasoundEngine/Public/IMetasoundEngineModule.h
--- a/Engine/Plugins/Runtime/Metasound/Source/MetasoundEngine/Public/IMetasoundEngineModule.h
+++ b/Engine/Plugins/Runtime/Metasound/Source/MetasoundEngine/Public/IMetasoundEngineModule.h
@@ -9,7 +9,7 @@
 namespace Metasound::Engine
 {
 #if WITH_EDITOR
-	// Status of initial asset scan when editor loads up.
+	// AssetScanStatus is no longer necessary and will be deprecated once usage is removed from MetaSound plugin API
 	enum class EAssetScanStatus : uint8
 	{
 		NotRequested = 0,
@@ -17,9 +17,7 @@ namespace Metasound::Engine
 		Complete = 3
 	};
 
-	// Node class primes status of MetaSound assets.  Priming an asset
-	// loads the asset asynchronously (if not already loaded)
-	// & registers it with the MetaSound Node Class Registry.
+	// NodeClassRegistryPrimeStatus is no longer necessary and will be deprecated once usage is removed from MetaSound plugin API
 	enum class ENodeClassRegistryPrimeStatus : uint8
 	{
 		NotRequested = 0,
@@ -48,19 +46,20 @@ namespace Metasound::Engine
 		virtual void ShutdownModule() override = 0;
 
 #if WITH_EDITOR
-		// Primes MetaSound assets, registering them with the asset manager
-		// and loads assets asynchronously (if not already loaded) 
-		// & registers them if not already registered with the MetaSound Node Class Registry.
-		virtual void PrimeAssetRegistryAsync() = 0;
+		UE_DEPRECATED(5.5, "Async prime no longer required (nor recommended as AssetTag updates now enable querying MetaSound assets for interface information without loading full asset.")
+		virtual void PrimeAssetRegistryAsync() { }
 
-		// Adds assets to the MetaSound asset manager.
-		// This is a subset of PrimeAssetRegistryAsync (does not load/register assets)
+		// Request enumeration of scanned assets and kick off add of MetaSound asset tag data to MetaSoundAssetManager
 		virtual void PrimeAssetManager() = 0;
 
-		virtual ENodeClassRegistryPrimeStatus GetNodeClassRegistryPrimeStatus() const = 0;
-		virtual EAssetScanStatus GetAssetRegistryScanStatus() const = 0;
-		// Bool rather than enum because this does not require async asset loading like priming the node class registry
-		virtual bool IsAssetManagerPrimed() const = 0;
+		UE_DEPRECATED(5.5, "Async prime no longer required (nor recommended as AssetTag updates now enable querying MetaSound assets for interface information without loading full asset.")
+		virtual ENodeClassRegistryPrimeStatus GetNodeClassRegistryPrimeStatus() { return ENodeClassRegistryPrimeStatus::NotRequested; }
+
+		UE_DEPRECATED(5.5, "Use FMetaSoundAssetManager::IsInitialAssetScanComplete instead")
+		virtual EAssetScanStatus GetAssetRegistryScanStatus() const { return EAssetScanStatus::NotRequested; }
+
+		UE_DEPRECATED(5.5, "Async prime no longer required (nor recommended as AssetTag updates now enable querying MetaSound assets for interface information without loading full asset.")
+		virtual bool IsAssetManagerPrimed() const { return false; }
 
 		// Asset registry delegates for calling MetaSound editor module register/unregister with frontend 
 		virtual FOnMetasoundGraphRegister& GetOnGraphRegisteredDelegate() = 0;
diff --git a/Engine/Plugins/Runtime/Metasound/Source/MetasoundEngine/Public/Metasound.h b/Engine/Plugins/Runtime/Metasound/Source/MetasoundEngine/Public/Metasound.h
--- a/Engine/Plugins/Runtime/Metasound/Source/MetasoundEngine/Public/Metasound.h
+++ b/Engine/Plugins/Runtime/Metasound/Source/MetasoundEngine/Public/Metasound.h
@@ -100,28 +100,25 @@ protected:
 public:
 	UMetaSoundPatch(const FObjectInitializer& ObjectInitializer);
 
-	UPROPERTY(AssetRegistrySearchable)
+#if WITH_EDITORONLY_DATA
+	UPROPERTY(meta = (DeprecatedProperty, DeprecationMessage = "Property is now serialized directly as asset tag"))
 	FGuid AssetClassID;
 
-#if WITH_EDITORONLY_DATA
-	UPROPERTY(AssetRegistrySearchable)
+	UPROPERTY(meta = (DeprecatedProperty, DeprecationMessage = "Property is now serialized directly as asset tag"))
 	FString RegistryInputTypes;
 
-	UPROPERTY(AssetRegistrySearchable)
+	UPROPERTY(meta = (DeprecatedProperty, DeprecationMessage = "Property is now serialized directly as asset tag"))
 	FString RegistryOutputTypes;
 
-	UPROPERTY(AssetRegistrySearchable)
+	UPROPERTY(meta = (DeprecatedProperty, DeprecationMessage = "Property is now serialized directly as asset tag"))
 	int32 RegistryVersionMajor = 0;
 
-	UPROPERTY(AssetRegistrySearchable)
+	UPROPERTY(meta = (DeprecatedProperty, DeprecationMessage = "Property is now serialized directly as asset tag"))
 	int32 RegistryVersionMinor = 0;
 
-	UPROPERTY(AssetRegistrySearchable)
+	UPROPERTY(meta = (DeprecatedProperty, DeprecationMessage = "Property is now serialized directly as asset tag"))
 	bool bIsPreset = false;
 
-	// Sets Asset Registry Metadata associated with this MetaSound
-	virtual void SetRegistryAssetClassInfo(const Metasound::Frontend::FNodeClassInfo& InClassInfo) override;
-
 	// Returns document name (for editor purposes, and avoids making document public for edit
 	// while allowing editor to reference directly)
 	static FName GetDocumentPropertyName()
@@ -148,6 +145,8 @@ public:
 	}
 #endif // #if WITH_EDITORONLY_DATA
 
+	virtual void GetAssetRegistryTags(FAssetRegistryTagsContext Context) const override;
+
 	virtual FTopLevelAssetPath GetAssetPathChecked() const override;
 	virtual const UClass& GetBaseMetaSoundUClass() const final override;
 	virtual const UClass& GetBuilderUClass() const final override;
diff --git a/Engine/Plugins/Runtime/Metasound/Source/MetasoundEngine/Public/MetasoundAssetSubsystem.h b/Engine/Plugins/Runtime/Metasound/Source/MetasoundEngine/Public/MetasoundAssetSubsystem.h
--- a/Engine/Plugins/Runtime/Metasound/Source/MetasoundEngine/Public/MetasoundAssetSubsystem.h
+++ b/Engine/Plugins/Runtime/Metasound/Source/MetasoundEngine/Public/MetasoundAssetSubsystem.h
@@ -7,19 +7,179 @@
 #include "MetasoundAssetManager.h"
 #include "MetasoundBuilderBase.h"
 #include "MetasoundDocumentInterface.h"
+#include "MetasoundFrontendQuery.h"
 #include "Subsystems/EngineSubsystem.h"
+#include "UObject/AssetRegistryTagsContext.h"
 #include "UObject/Object.h"
 
 #include "MetasoundAssetSubsystem.generated.h"
 
 // Forward Declarations
-class UAssetManager;
 class FMetasoundAssetBase;
+class UAssetManager;
 
 struct FDirectoryPath;
 struct FMetaSoundFrontendDocumentBuilder;
 
 
+namespace Metasound::Engine
+{
+	void DeinitializeAssetManager();
+	void InitializeAssetManager();
+
+
+	class METASOUNDENGINE_API FMetaSoundAssetManager :
+		public Frontend::IMetaSoundAssetManager,
+		public FGCObject
+	{
+	public:
+		friend class UMetaSoundAssetSubsystem;
+
+		using FAssetClassInfoMap = TMap<Frontend::FAssetKey, TArray<Frontend::FAssetClassQueryInfo>>;
+		using FAssetInfo = Frontend::IMetaSoundAssetManager::FAssetInfo;
+		using FAssetKey = Frontend::FAssetKey;
+		using FNodeRegistryKey = Frontend::FNodeRegistryKey;
+
+#if WITH_EDITORONLY_DATA
+		using FVersionDocumentResults = Frontend::IMetaSoundAssetManager::FVersionDocumentResults;
+#endif // WITH_EDITORONLY_DATA
+
+		FMetaSoundAssetManager() = default;
+		virtual ~FMetaSoundAssetManager();
+
+		static FMetaSoundAssetManager* Get()
+		{
+			using namespace Frontend;
+			return static_cast<FMetaSoundAssetManager*>(IMetaSoundAssetManager::Get());
+		}
+
+		static FMetaSoundAssetManager& GetChecked()
+		{
+			using namespace Frontend;
+			return static_cast<FMetaSoundAssetManager&>(IMetaSoundAssetManager::GetChecked());
+		}
+
+		TArray<Frontend::FAssetClassQueryInfo> FindAssetTagData(const Frontend::FAssetKey& InKey) const;
+
+		Frontend::FAssetKey GetAssetKey(const FSoftObjectPath& InObjectPath) const;
+
+#if WITH_EDITOR
+		int32 GetActiveAsyncLoadRequestCount() const;
+#endif // WITH_EDITOR
+
+		bool IsInitialAssetScanComplete() const;
+
+		void RebuildDenyListCache(const UAssetManager& InAssetManager);
+		void RegisterAssetClassesInDirectories(const TArray<FMetaSoundAssetDirectory>& Directories);
+
+#if WITH_EDITOR
+		bool ReplaceReferencesInDirectory(const TArray<FMetaSoundAssetDirectory>& InDirectories, const FNodeRegistryKey& OldClassKey, const FNodeRegistryKey& NewClassKey) const;
+#endif // WITH_EDITOR
+		void RequestAsyncLoadReferencedAssets(FMetasoundAssetBase& InAssetBase);
+
+		UE_DEPRECATED(5.6, "Implementation has been privitized")
+		void OnAssetScanComplete() { }
+
+		void SearchAndIterateDirectoryAssets(const TArray<FDirectoryPath>& InDirectories, TFunctionRef<void(const FAssetData&)> InFunction) const;
+
+#if WITH_EDITOR
+		void SetCanNotifyAssetTagScanComplete();
+#endif // WITH_EDITOR
+
+		FMetasoundAssetBase* TryLoadAsset(const FSoftObjectPath& InObjectPath) const;
+		void UnregisterAssetClassesInDirectories(const TArray<FMetaSoundAssetDirectory>& Directories);
+
+		/* IMetaSoundAssetManager Implementation */
+#if WITH_EDITORONLY_DATA
+		virtual bool AddAssetReferences(FMetasoundAssetBase& InAssetBase) override;
+#endif // WITH_EDITORONLY_DATA
+		virtual FAssetKey AddOrUpdateFromObject(const UObject& InObject) override;
+		virtual void AddOrLoadAndUpdateFromObjectAsync(const FAssetData& InAssetData, TUniqueFunction<void(FAssetKey, UObject&)>&& OnAssetRegistered) override;
+		virtual void AddOrUpdateFromAssetTags(const FAssetData& InAssetData) override;
+		virtual bool CanAutoUpdate(const FMetasoundFrontendClassName& InClassName) const override;
+		virtual bool ContainsKey(const FAssetKey& InKey) const override;
+		virtual bool ContainsKey(const Frontend::FNodeRegistryKey& InRegistryKey) const override;
+		virtual FMetasoundAssetBase* FindAsset(const FAssetKey& InKey) const override;
+		virtual TScriptInterface<IMetaSoundDocumentInterface> FindAssetAsDocumentInterface(const FAssetKey& InKey) const override;
+		virtual FTopLevelAssetPath FindAssetPath(const FAssetKey& InKey) const override;
+		virtual TArray<FTopLevelAssetPath> FindAssetPaths(const FAssetKey& InKey) const override;
+		virtual FMetasoundAssetBase* GetAsAsset(UObject& InObject) const override;
+		virtual const FMetasoundAssetBase* GetAsAsset(const UObject& InObject) const override;
+#if WITH_EDITOR
+		virtual TSet<FAssetInfo> GetReferencedAssetClasses(const FMetasoundAssetBase& InAssetBase) const override;
+		virtual bool ReassignClassName(TScriptInterface<IMetaSoundDocumentInterface> DocInterface) override;
+#endif // WITH_EDITOR
+
+		virtual bool IsAssetClass(const FMetasoundFrontendClassMetadata& ClassMetadata) const override;
+
+#if WITH_EDITOR
+		// Iterates all asset tag data (Only recommended in certain editor contexts, as slow and blocks access to reference map).
+		void IterateAssetTagData(TFunctionRef<void(Frontend::FAssetClassQueryInfo)> Iter, bool bIterateDuplicates = false) const;
+
+		// Iterates all references of a given asset key entry (Only recommended in certain editor contexts, as can be slow for
+		// deep reference trees and blocks access to reference map).
+		void IterateReferences(const FAssetKey& InKey, TFunctionRef<void(const FAssetKey&)> VisitFunction) const;
+#endif // WITH_EDITOR
+
+		virtual void ReloadMetaSoundAssets() const override;
+		virtual void RemoveAsset(const UObject& InObject) override;
+		virtual void RemoveAsset(const FAssetData& InAssetData) override;
+		virtual void RenameAsset(const FAssetData& InAssetData, const FString& InOldObjectPath) override;
+		virtual void SetLogActiveAssetsOnShutdown(bool bInLogActiveAssetsOnShutdown) override;
+		virtual FMetasoundAssetBase* TryLoadAssetFromKey(const FAssetKey& InKey) const override;
+		virtual bool TryGetAssetIDFromClassName(const FMetasoundFrontendClassName& InClassName, FGuid& OutGuid) const override;
+		virtual bool TryLoadReferencedAssets(const FMetasoundAssetBase& InAssetBase, TArray<FMetasoundAssetBase*>& OutReferencedAssets) const override;
+#if WITH_EDITORONLY_DATA
+		virtual FVersionDocumentResults VersionDocumentsInFolders(const TArray<FString>& FolderPaths) const override;
+#endif // WITH_EDITORONLY_DATA
+
+		virtual void WaitUntilAsyncLoadReferencedAssetsComplete(FMetasoundAssetBase& InAssetBase) override;
+
+		/* FGCObject */
+		virtual void AddReferencedObjects(FReferenceCollector& Collector) override;
+		virtual FString GetReferencerName() const override { return TEXT("FMetaSoundAssetManager"); }
+
+	private:
+		TArray<FMetaSoundAsyncAssetDependencies> LoadingDependencies;
+
+		static void DepthFirstVisit_AssetKey(const FAssetKey& InKey, TFunctionRef<TSet<FAssetKey>(const FAssetKey&)> VisitFunction);
+		FMetaSoundAsyncAssetDependencies* FindLoadingDependencies(const UObject* InParentAsset);
+		FMetaSoundAsyncAssetDependencies* FindLoadingDependencies(int32 InLoadID);
+		void RemoveLoadingDependencies(int32 InLoadID);
+		void OnReferencedAssetsLoaded(int32 InLoadID);
+
+		struct FPackageLoadedArgs
+		{
+			FName PackageName;
+			UPackage* Package = nullptr;
+			EAsyncLoadingResult::Type Result = EAsyncLoadingResult::Failed;
+			TUniqueFunction<void(FAssetKey, UObject&)> OnAssetRegistered;
+			bool bMarkAsVersioned = false;
+		};
+
+		static void OnPackageLoaded(const FPackageLoadedArgs& PackageLoadedArgs);
+
+		FStreamableManager StreamableManager;
+		int32 AsyncLoadIDCounter = 0;
+
+		int32 AutoUpdateDenyListChangeID = INDEX_NONE;
+		TSet<FName> AutoUpdateDenyListCache;
+
+		FAssetClassInfoMap ClassInfoMap;
+
+		// Critical section primarily for allowing safe access of class info map during async loading of MetaSound assets.
+		mutable FCriticalSection TagDataMapCriticalSection;
+
+		bool bLogActiveAssetsOnShutdown = true;
+
+#if WITH_EDITORONLY_DATA
+		bool bNotifyTagDataScanComplete = false;
+		int32 ActiveAsyncAssetLoadRequests = 0;
+#endif // WITH_EDITORONLY_DATA
+	};
+} // namespace Metasound::Engine
+
+
 USTRUCT(BlueprintType)
 struct METASOUNDENGINE_API FMetaSoundAssetDirectory
 {
@@ -49,12 +209,6 @@ struct FMetaSoundAsyncAssetDependencies
 	TSharedPtr<FStreamableHandle> StreamableHandle;
 };
 
-namespace Metasound::Engine
-{
-	void DeinitializeAssetManager();
-	void InitializeAssetManager();
-} // namespace Metasound::Engine
-
 
 UCLASS(meta = (DisplayName = "MetaSound Asset Subsystem"))
 class METASOUNDENGINE_API UMetaSoundAssetSubsystem : public UEngineSubsystem
@@ -80,10 +234,10 @@ public:
 	virtual void AddAssetReferences(FMetasoundAssetBase& InAssetBase);
 #endif
 
-	UE_DEPRECATED(5.5, "Moved to internal implementation, use IMetaSoundAssetManager::GetChecked() and analogous call")
+	UE_DEPRECATED(5.5, "Moved to internal implementation, use IMetaSoundAssetManager::GetChecked() 'AddOrLoadAndUpdateFromObjectAsync' call")
 	Metasound::Frontend::FNodeRegistryKey AddOrUpdateAsset(const FAssetData& InAssetData);
 
-	UE_DEPRECATED(5.5, "Moved to internal implementation, use IMetaSoundAssetManager::GetChecked() and analogous call")
+	UE_DEPRECATED(5.5, "Moved to internal implementation, use IMetaSoundAssetManager::GetChecked() 'AddOrUpdateFromObject' call")
 	virtual Metasound::Frontend::FNodeRegistryKey AddOrUpdateAsset(const UObject& InObject);
 
 	UE_DEPRECATED(5.5, "Moved to internal implementation, use IMetaSoundAssetManager::GetChecked() and analogous call")
diff --git a/Engine/Plugins/Runtime/Metasound/Source/MetasoundEngine/Public/MetasoundBuilderBase.cpp b/Engine/Plugins/Runtime/Metasound/Source/MetasoundEngine/Public/MetasoundBuilderBase.cpp
--- a/Engine/Plugins/Runtime/Metasound/Source/MetasoundEngine/Public/MetasoundBuilderBase.cpp
+++ b/Engine/Plugins/Runtime/Metasound/Source/MetasoundEngine/Public/MetasoundBuilderBase.cpp
@@ -312,17 +312,18 @@ TScriptInterface<IMetaSoundDocumentInterface> UMetaSoundBuilderBase::Build(const
 
 void UMetaSoundBuilderBase::BuildAndOverwriteMetaSound(TScriptInterface<IMetaSoundDocumentInterface> ExistingMetaSound, bool bForceUniqueClassName)
 {
-	if (!ExistingMetaSound)
+	UObject* MetaSoundPtr = ExistingMetaSound.GetObject();
+	if (!MetaSoundPtr)
 	{
 		UE_LOG(LogMetaSound, Error, TEXT("Failed to build and overwrite MetaSound: No existing MetaSound supplied."));
 		return;
 	}
 
-	if (ExistingMetaSound.GetObject()->IsAsset())
+	if (MetaSoundPtr->IsAsset())
 	{
 		UE_LOG(LogMetaSound, Error, TEXT("Failed to build and overwrite MetaSound: Cannot overwrite serialized asset "
-			"(use 'BuildNewMetaSound' to create a new, transient MetaSound. Overwriting serialized asset is only supported at edit time via "
-			"UMetaSoundEditorSubsystem::BuildToAsset."));
+			"(use 'BuildNewMetaSound' to create a new, transient MetaSound. Overwriting serialized asset is only "
+			"supported at edit-time via UMetaSoundEditorSubsystem::BuildToAsset."));
 		return;
 	}
 
@@ -1056,7 +1057,7 @@ void UMetaSoundBuilderBase::RegisterGraphIfOutstandingTransactions(UObject& InMe
 	{
 		UObject* RefMetaSound = Reference->GetOwningAsset();
 		check(RefMetaSound);
-		AssetManager.AddOrUpdateAsset(*RefMetaSound);
+		AssetManager.AddOrUpdateFromObject(*RefMetaSound);
 		RegisterGraphIfOutstandingTransactions(*RefMetaSound);
 	}
 
diff --git a/Engine/Plugins/Runtime/Metasound/Source/MetasoundEngine/Public/MetasoundSource.h b/Engine/Plugins/Runtime/Metasound/Source/MetasoundEngine/Public/MetasoundSource.h
--- a/Engine/Plugins/Runtime/Metasound/Source/MetasoundEngine/Public/MetasoundSource.h
+++ b/Engine/Plugins/Runtime/Metasound/Source/MetasoundEngine/Public/MetasoundSource.h
@@ -15,6 +15,7 @@
 #include "MetasoundRouter.h"
 #include "MetasoundVertex.h"
 #include "Sound/SoundWaveProcedural.h"
+#include "UObject/AssetRegistryTagsContext.h"
 #include "UObject/MetaData.h"
 
 #include <atomic>
@@ -158,30 +159,25 @@ public:
 	// Override the SampleRate for this Sound (overrides Quality). NOTE: A Zero value will have no effect and use either the Quality setting (if set), or the Device Rate
 	UPROPERTY(EditAnywhere, BlueprintReadOnly, AdvancedDisplay, Category = Metasound, meta = (UIMin = 0, UIMax = 96000, DisplayName = "Override Sample Rate (in Hz)"))
 	FPerPlatformInt SampleRateOverride = 0;
-#endif //WITH_EDITOR_DATA
 
-	UPROPERTY(AssetRegistrySearchable)
+	UPROPERTY(meta = (DeprecatedProperty, DeprecationMessage = "Property is now serialized directly as asset tag"))
 	FGuid AssetClassID;
 
-#if WITH_EDITORONLY_DATA
-	UPROPERTY(AssetRegistrySearchable)
+	UPROPERTY(meta = (DeprecatedProperty, DeprecationMessage = "Property is now serialized directly as asset tag"))
 	FString RegistryInputTypes;
 
-	UPROPERTY(AssetRegistrySearchable)
+	UPROPERTY(meta = (DeprecatedProperty, DeprecationMessage = "Property is now serialized directly as asset tag"))
 	FString RegistryOutputTypes;
 
-	UPROPERTY(AssetRegistrySearchable)
+	UPROPERTY(meta = (DeprecatedProperty, DeprecationMessage = "Property is now serialized directly as asset tag"))
 	int32 RegistryVersionMajor = 0;
 
-	UPROPERTY(AssetRegistrySearchable)
+	UPROPERTY(meta = (DeprecatedProperty, DeprecationMessage = "Property is now serialized directly as asset tag"))
 	int32 RegistryVersionMinor = 0;
 
-	UPROPERTY(AssetRegistrySearchable)
+	UPROPERTY(meta = (DeprecatedProperty, DeprecationMessage = "Property is now serialized directly as asset tag"))
 	bool bIsPreset = false;
 
-	// Sets Asset Registry Metadata associated with this MetaSoundSource
-	virtual void SetRegistryAssetClassInfo(const Metasound::Frontend::FNodeClassInfo& InNodeInfo) override;
-
 	// Returns document name (for editor purposes, and avoids making document public for edit
 	// while allowing editor to reference directly)
 	static FName GetDocumentPropertyName()
@@ -256,6 +252,8 @@ public:
 
 	virtual bool ConformObjectToDocument() override;
 
+	virtual void GetAssetRegistryTags(FAssetRegistryTagsContext Context) const override;
+
 	virtual FTopLevelAssetPath GetAssetPathChecked() const override;
 
 	UObject* GetOwningAsset() override
diff --git a/Engine/Plugins/Runtime/Metasound/Source/MetasoundEngineTest/Private/EngineTestMetaSoundAutomatedNodeTest.cpp b/Engine/Plugins/Runtime/Metasound/Source/MetasoundEngineTest/Private/EngineTestMetaSoundAutomatedNodeTest.cpp
--- a/Engine/Plugins/Runtime/Metasound/Source/MetasoundEngineTest/Private/EngineTestMetaSoundAutomatedNodeTest.cpp
+++ b/Engine/Plugins/Runtime/Metasound/Source/MetasoundEngineTest/Private/EngineTestMetaSoundAutomatedNodeTest.cpp
@@ -67,7 +67,7 @@ namespace Metasound::EngineTest{
 		return Environment;
 	}
 
-	FString GetPrettyName(const Frontend::FNodeRegistryKey& InRegistryKey)
+	FString GetPrettyName(const Frontend::FNodeClassRegistryKey& InRegistryKey)
 	{
 		const FTopLevelAssetPath ObjectPath = Frontend::IMetaSoundAssetManager::GetChecked().FindAssetPath(InRegistryKey);
 		if (ObjectPath.IsValid())
@@ -787,6 +787,7 @@ PRAGMA_ENABLE_DEPRECATION_WARNINGS
 		FMetasoundFrontendRegistryContainer* NodeRegistry = FMetasoundFrontendRegistryContainer::Get();
 		check(nullptr != NodeRegistry);
 
+		Frontend::IMetaSoundAssetManager* AssetManager = Frontend::IMetaSoundAssetManager::Get();
 		for (const FMetasoundFrontendClass& NodeClass : AllClasses)
 		{
 			// Exclude template classes because they cannot be created directly from the node registry
@@ -802,8 +803,8 @@ PRAGMA_ENABLE_DEPRECATION_WARNINGS
 
 			Frontend::FNodeRegistryKey NodeRegistryKey(NodeClass.Metadata);
 
-			// Exclude non-native nodes (Nodes defined by assets instead of C++)
-			if (!NodeRegistry->IsNodeNative(NodeRegistryKey))
+			// Exclude asset-defined node classes
+			if (AssetManager && AssetManager->ContainsKey(NodeRegistryKey))
 			{
 				continue;
 			}
diff --git a/Engine/Plugins/Runtime/Metasound/Source/MetasoundEngineTest/Private/EngineTestMetasoundSource.cpp b/Engine/Plugins/Runtime/Metasound/Source/MetasoundEngineTest/Private/EngineTestMetasoundSource.cpp
--- a/Engine/Plugins/Runtime/Metasound/Source/MetasoundEngineTest/Private/EngineTestMetasoundSource.cpp
+++ b/Engine/Plugins/Runtime/Metasound/Source/MetasoundEngineTest/Private/EngineTestMetasoundSource.cpp
@@ -11,6 +11,7 @@
 #include "Interfaces/MetasoundFrontendSourceInterface.h"
 #include "Math/UnrealMathUtility.h"
 #include "MetasoundDataReference.h"
+#include "MetasoundFrontendDocumentBuilder.h"
 #include "MetasoundLog.h"
 #include "MetasoundSettings.h"
 #include "MetasoundSource.h"
@@ -177,10 +178,11 @@ namespace EngineTestMetaSoundSourcePrivate
 		using namespace Metasound::Engine;
 		using namespace Metasound::Frontend;
 
-		FMetasoundFrontendDocument Document;
+		// TODO: Move this to generate the basic document fully from builder. As of now, this test
+		// uses a weird mix of builder API & controller API calls.
 
-		Document.RootGraph.Metadata.SetClassName(FMetasoundFrontendClassName { "Namespace", "Unit Test Node", *LexToString(FGuid::NewGuid()) });
-		Document.RootGraph.Metadata.SetType(EMetasoundFrontendClassType::Graph);
+		FMetasoundFrontendDocument Document;
+		FMetaSoundFrontendDocumentBuilder::InitGraphClassMetadata(Document.RootGraph.Metadata);
 		Document.RootGraph.InitDefaultGraphPage();
 
 		FDocumentHandle DocumentHandle = IDocumentController::CreateDocumentHandle(Document);
diff --git a/Engine/Plugins/Runtime/Metasound/Source/MetasoundFrontend/MetasoundFrontend.Build.cs b/Engine/Plugins/Runtime/Metasound/Source/MetasoundFrontend/MetasoundFrontend.Build.cs
--- a/Engine/Plugins/Runtime/Metasound/Source/MetasoundFrontend/MetasoundFrontend.Build.cs
+++ b/Engine/Plugins/Runtime/Metasound/Source/MetasoundFrontend/MetasoundFrontend.Build.cs
@@ -18,6 +18,8 @@ namespace UnrealBuildTool.Rules
 					"AudioExtensions",
 					"Core",
 					"CoreUObject",
+					"Json",
+					"JsonUtilities",
 					"Serialization",
 					"SignalProcessing",
 					"MetasoundGraphCore"
diff --git a/Engine/Plugins/Runtime/Metasound/Source/MetasoundFrontend/Private/MetasoundAssetBase.cpp b/Engine/Plugins/Runtime/Metasound/Source/MetasoundFrontend/Private/MetasoundAssetBase.cpp
--- a/Engine/Plugins/Runtime/Metasound/Source/MetasoundFrontend/Private/MetasoundAssetBase.cpp
+++ b/Engine/Plugins/Runtime/Metasound/Source/MetasoundFrontend/Private/MetasoundAssetBase.cpp
@@ -82,7 +82,8 @@ namespace Metasound
 				FDocumentNodeRegistryEntryForSerialization(const FMetasoundFrontendDocument& InDocument, const FTopLevelAssetPath& InAssetPath)
 					: Interfaces(InDocument.Interfaces)
 					, FrontendClass(InDocument.RootGraph)
-					, ClassInfo(InDocument.RootGraph, InAssetPath)
+					, ClassInfo(InDocument.RootGraph)
+					, AssetPath(InAssetPath)
 				{
 					// Copy FrontendClass to preserve original document.
 					FrontendClass.Metadata.SetType(EMetasoundFrontendClassType::External);
@@ -116,15 +117,11 @@ namespace Metasound
 					return CreateVertexInterfaceFromClassNoProxy(FrontendClass);
 				}
 
-				virtual bool IsNative() const override
-				{
-					return false;
-				}
-
 			private:
 				TSet<FMetasoundFrontendVersion> Interfaces;
 				FMetasoundFrontendClass FrontendClass;
 				FNodeClassInfo ClassInfo;
+				FTopLevelAssetPath AssetPath;
 			};
 
 			void GetUpdatePathForDocument(const FMetasoundFrontendVersion& InCurrentVersion, const FMetasoundFrontendVersion& InTargetVersion, TArray<const IInterfaceRegistryEntry*>& OutUpgradePath)
@@ -316,11 +313,6 @@ void FMetasoundAssetBase::UpdateAndRegisterForExecution(Metasound::Frontend::FMe
 	UObject* Owner = GetOwningAsset();
 	check(Owner);
 
-	// This should not be necessary as it should be added on asset load,
-	// but currently registration is required to be called prior to adding
-	// an object-defined graph class to the registry so it was placed here.
-	IMetaSoundAssetManager::GetChecked().AddOrUpdateAsset(*Owner);
-
 	// Auto update must be done after all referenced asset classes are registered
 	if (InRegistrationOptions.bAutoUpdate)
 	{
@@ -361,13 +353,7 @@ void FMetasoundAssetBase::UpdateAndRegisterForExecution(Metasound::Frontend::FMe
 #endif // WITH_EDITOR
 
 	GraphRegistryKey = FNodeClassRegistry::Get().RegisterGraph(Owner);
-	if (GraphRegistryKey.IsValid())
-	{
-#if WITH_EDITORONLY_DATA
-		UpdateAssetRegistry();
-#endif // WITH_EDITORONLY_DATA
-	}
-	else
+	if (!GraphRegistryKey.IsValid())
 	{
 		UClass* Class = Owner->GetClass();
 		check(Class);
@@ -399,7 +385,6 @@ void FMetasoundAssetBase::UpdateAndRegisterForSerialization(FName CookPlatformNa
 	}
 
 	UpdateAndRegisterReferencesForSerialization(CookPlatformName);
-	IMetaSoundAssetManager::GetChecked().AddOrUpdateAsset(*GetOwningAsset());
 
 	UObject* Owner = GetOwningAsset();
 	check(Owner);
@@ -453,11 +438,7 @@ void FMetasoundAssetBase::UpdateAndRegisterForSerialization(FName CookPlatformNa
 		GraphRegistryKey = FGraphRegistryKey { NodeKey, AssetPath };
 	}
 
-	if (GraphRegistryKey.IsValid())
-	{
-		UpdateAssetRegistry();
-	}
-	else
+	if (!GraphRegistryKey.IsValid())
 	{
 		const UClass* Class = Owner->GetClass();
 		check(Class);
@@ -526,16 +507,6 @@ bool FMetasoundAssetBase::IsInterfaceDeclared(const FMetasoundFrontendVersion& I
 	return GetConstDocumentChecked().Interfaces.Contains(InVersion);
 }
 
-Metasound::Frontend::FNodeClassInfo FMetasoundAssetBase::GetAssetClassInfo() const
-{
-	using namespace Metasound::Frontend;
-
-	const UObject* Owner = GetOwningAsset();
-	check(Owner);
-	TScriptInterface<const IMetaSoundDocumentInterface> DocInterface((UObject*)Owner);
-	return FNodeClassInfo { GetConstDocumentChecked().RootGraph, DocInterface->GetAssetPathChecked() };
-}
-
 void FMetasoundAssetBase::SetDocument(FMetasoundFrontendDocument InDocument, bool bMarkDirty)
 {
 	FMetasoundFrontendDocument* Document = GetDocumentAccessPtr().Get();
@@ -941,8 +912,6 @@ void FMetasoundAssetBase::RegisterAssetDependencies(const Metasound::Frontend::F
 	{
 		if (InRegistrationOptions.bForceReregister || !Reference->IsRegistered())
 		{
-			// TODO: Check for infinite recursion and error if so
-			AssetManager.AddOrUpdateAsset(*(Reference->GetOwningAsset()));
 			Reference->UpdateAndRegisterForExecution(InRegistrationOptions);
 		}
 	}
@@ -960,14 +929,11 @@ void FMetasoundAssetBase::UpdateAndRegisterReferencesForSerialization(FName Cook
 {
 	using namespace Metasound::Frontend;
 
-	IMetaSoundAssetManager& AssetManager = IMetaSoundAssetManager::GetChecked();
 	TArray<FMetasoundAssetBase*> References = GetReferencedAssets();
 	for (FMetasoundAssetBase* Reference : References)
 	{
 		if (!Reference->IsRegistered())
 		{
-			// TODO: Check for infinite recursion and error if so
-			AssetManager.AddOrUpdateAsset(*(Reference->GetOwningAsset()));
 			Reference->UpdateAndRegisterForSerialization(CookPlatformName);
 		}
 	}
@@ -989,37 +955,6 @@ bool FMetasoundAssetBase::AutoUpdate(bool bInLogWarningsOnDroppedConnection)
 #endif // !WITH_EDITORONLY_DATA
 }
 
-#if WITH_EDITORONLY_DATA
-void FMetasoundAssetBase::UpdateAssetRegistry() 
-{
-	using namespace Metasound;
-	using namespace Metasound::Frontend;
-
-	UObject* Owner = GetOwningAsset();
-	check(Owner);
-	const FMetasoundFrontendGraphClass& DocumentClassGraph = GetDocumentChecked().RootGraph;
-
-	TScriptInterface<IMetaSoundDocumentInterface> DocInterface(Owner);
-	FNodeClassInfo AssetClassInfo(DocumentClassGraph, DocInterface->GetAssetPathChecked());
-
-	// Refresh Asset Registry Info if successfully registered with Frontend
-	const FMetasoundFrontendClassMetadata& DocumentClassMetadata = DocumentClassGraph.Metadata;
-	AssetClassInfo.AssetClassID = FGuid(DocumentClassMetadata.GetClassName().Name.ToString());
-	FNodeClassName ClassName = DocumentClassMetadata.GetClassName().ToNodeClassName();
-	FMetasoundFrontendClass GraphClass;
-
-	AssetClassInfo.bIsPreset = DocumentClassGraph.PresetOptions.bIsPreset;
-	AssetClassInfo.Version = DocumentClassMetadata.GetVersion();
-	AssetClassInfo.InputTypes.Reset();
-	Algo::Transform(GraphClass.Interface.Inputs, AssetClassInfo.InputTypes, [] (const FMetasoundFrontendClassInput& Input) { return Input.TypeName; });
-
-	AssetClassInfo.OutputTypes.Reset();
-	Algo::Transform(GraphClass.Interface.Outputs, AssetClassInfo.OutputTypes, [](const FMetasoundFrontendClassOutput& Output) { return Output.TypeName; });
-
-	SetRegistryAssetClassInfo(MoveTemp(AssetClassInfo));
-}
-#endif
-
 bool FMetasoundAssetBase::TryUpdateInterfaceFromVersion(const FMetasoundFrontendVersion& Version)
 {
 	using namespace Metasound::Frontend;
@@ -1089,11 +1024,6 @@ FMetasoundFrontendInterface FMetasoundAssetBase::GetInterfaceToVersion(const FMe
 	return TargetInterface;
 }
 
-TSharedPtr<FMetasoundFrontendDocument> FMetasoundAssetBase::PreprocessDocument()
-{
-	return nullptr;
-}
-
 #if WITH_EDITORONLY_DATA
 bool FMetasoundAssetBase::GetVersionedOnLoad() const
 {
@@ -1111,14 +1041,6 @@ void FMetasoundAssetBase::SetVersionedOnLoad()
 }
 #endif // WITH_EDITORONLY_DATA
 
-PRAGMA_DISABLE_DEPRECATION_WARNINGS
-const FMetasoundAssetBase::FRuntimeData& FMetasoundAssetBase::GetRuntimeData() const
-{
-	static const FRuntimeData PlaceholderForDeprecatedMethod;
-	return PlaceholderForDeprecatedMethod;
-}
-PRAGMA_ENABLE_DEPRECATION_WARNINGS
-
 FMetasoundAssetProxy::FMetasoundAssetProxy(const FParameters& InParams)
 {
 	Interfaces = InParams.Interfaces;
diff --git a/Engine/Plugins/Runtime/Metasound/Source/MetasoundFrontend/Private/MetasoundAssetManager.cpp b/Engine/Plugins/Runtime/Metasound/Source/MetasoundFrontend/Private/MetasoundAssetManager.cpp
--- a/Engine/Plugins/Runtime/Metasound/Source/MetasoundFrontend/Private/MetasoundAssetManager.cpp
+++ b/Engine/Plugins/Runtime/Metasound/Source/MetasoundFrontend/Private/MetasoundAssetManager.cpp
@@ -1,31 +1,21 @@
 // Copyright Epic Games, Inc. All Rights Reserved.
 
 #include "MetasoundAssetManager.h"
+#include "MetasoundAssetReferenceTagCollection.h"
 #include "MetasoundFrontendRegistryKey.h"
 
+
 namespace Metasound::Frontend
 {
 	namespace AssetManagerPrivate
 	{
 		static TUniquePtr<IMetaSoundAssetManager> Instance;
-
-		bool IsAssetClassType(EMetasoundFrontendClassType ClassType)
-		{
-			switch (ClassType)
-			{
-				case EMetasoundFrontendClassType::External:
-				case EMetasoundFrontendClassType::Graph:
-				case EMetasoundFrontendClassType::Invalid:
-					return true;
-
-				default:
-					return false;
-			}
-		}
 	} // AssetManagerPrivate
 
 	namespace AssetTags
 	{
+		// Deprecated Tag Identifiers
+PRAGMA_DISABLE_DEPRECATION_WARNINGS
 		const FString ArrayDelim = TEXT(",");
 
 		const FName AssetClassID = "AssetClassID";
@@ -41,37 +31,192 @@ namespace Metasound::Frontend
 		const FName RegistryInputTypes = "RegistryInputTypes";
 		const FName RegistryOutputTypes = "RegistryOutputTypes";
 #endif // WITH_EDITORONLY_DATA
+PRAGMA_ENABLE_DEPRECATION_WARNINGS
 	} // namespace AssetTags
 
-	FAssetKey::FAssetKey(const FMetasoundFrontendClassName& InClassName, const FMetasoundFrontendVersionNumber& InVersion)
-		: ClassName(InClassName)
-		, Version(InVersion)
+	namespace AssetTagsPrivate
 	{
+		static const FLazyName AssetClassID = "AssetClassID";
+		static const FLazyName ClassVersion = "ClassVersion";
+
+#if WITH_EDITORONLY_DATA
+		static const FLazyName ClassReferences = "ClassReferences";
+		static const FLazyName DisplayName = "DisplayName";
+		static const FLazyName DocumentVersion = "DocumentVersion";
+		static const FLazyName Inputs = "Inputs";
+		static const FLazyName IsPreset = "bIsPreset";
+		static const FLazyName Outputs = "Outputs";
+		static const FLazyName DisplayInfo = "DisplayInfo";
+#endif // WITH_EDITORONLY_DATA
 	}
 
-	FAssetKey::FAssetKey(const FNodeRegistryKey& RegKey)
-		: ClassName(RegKey.ClassName)
-		, Version(RegKey.Version)
+	FAssetClassQueryInfo::FAssetClassQueryInfo()
+		: bIsValid(0)
 	{
-		checkf(AssetManagerPrivate::IsAssetClassType(RegKey.Type), TEXT("Invalid ClassType '%s' for Registry Key"), LexToString(RegKey.Type));
+#if WITH_EDITORONLY_DATA
+		bIsPreset = 0;
+#endif // WITH_EDITORONLY_DATA
 	}
 
-	FAssetKey::FAssetKey(const FMetasoundFrontendClassMetadata& InMetadata)
-		: ClassName(InMetadata.GetClassName())
-		, Version(InMetadata.GetVersion())
+	FAssetClassQueryInfo::FAssetClassQueryInfo(const IMetaSoundDocumentInterface& InDocInterface)
 	{
-		checkf(AssetManagerPrivate::IsAssetClassType(InMetadata.GetType()), TEXT("Invalid ClassType '%s' for Registry Key"), LexToString(InMetadata.GetType()));
+		InitFromDocument(InDocInterface);
 	}
 
-	FString FAssetKey::ToString() const
+	FAssetClassQueryInfo::FAssetClassQueryInfo(const FAssetData& InAssetData)
 	{
-		TStringBuilder<128> KeyStringBuilder;
-		KeyStringBuilder.Append(ClassName.GetFullName().ToString());
-		KeyStringBuilder.AppendChar('_');
-		KeyStringBuilder.Append(FString::FromInt(Version.Major));
-		KeyStringBuilder.AppendChar('.');
-		KeyStringBuilder.Append(FString::FromInt(Version.Minor));
-		return KeyStringBuilder.ToString();
+		using namespace Frontend;
+
+		if (InAssetData.IsAssetLoaded())
+		{
+			UObject* MetaSound = InAssetData.GetAsset();
+			TScriptInterface<IMetaSoundDocumentInterface> DocInterface(MetaSound);
+
+			check(DocInterface.GetObject());
+			InitFromDocument(*DocInterface.GetInterface());
+			return;
+		}
+
+		bool bSuccess = true;
+
+#if WITH_EDITORONLY_DATA
+		{
+			FString DocVersionStr;
+			bSuccess &= InAssetData.GetTagValue(AssetTagsPrivate::DocumentVersion, DocVersionStr);
+			bSuccess &= FMetasoundFrontendVersionNumber::Parse(DocVersionStr, DocumentVersion);
+		}
+
+		{
+			bool bTagIsPreset = false;
+			bSuccess &= InAssetData.GetTagValue(AssetTagsPrivate::IsPreset, bTagIsPreset);
+			bIsPreset = (uint8)bTagIsPreset;
+		}
+
+		{
+			auto DeserializeTagFromJson = [&bSuccess, &InAssetData]<typename TStructType>(FName TagName, TStructType & OutStruct)
+			{
+				FString TagString;
+				const bool bTagDataFound = InAssetData.GetTagValue(TagName, TagString);
+				if (bTagDataFound)
+				{
+					return FJsonObjectConverter::JsonObjectStringToUStruct(TagString, &OutStruct);
+				}
+				return false;
+			};
+
+
+			{
+				FMetaSoundAssetReferenceTagCollection ReferenceTagCollection;
+				const bool bTagFound = DeserializeTagFromJson(AssetTagsPrivate::ClassReferences, ReferenceTagCollection);
+				if (bTagFound)
+				{
+					ReferencedAssetKeys = MoveTemp(ReferenceTagCollection.AssetKeys);
+				}
+				else
+				{
+					bSuccess = false;
+				}
+			}
+
+			bSuccess &= DeserializeTagFromJson(AssetTagsPrivate::DisplayInfo, DisplayInfo);
+			bSuccess &= DeserializeTagFromJson(AssetTagsPrivate::Inputs, Inputs);
+			bSuccess &= DeserializeTagFromJson(AssetTagsPrivate::Outputs, Outputs);
+		}
+#endif // WITH_EDITORONLY_DATA
+
+		{
+			FString ClassVersionStr;
+			bSuccess &= InAssetData.GetTagValue(AssetTagsPrivate::ClassVersion, ClassVersionStr);
+			bSuccess &= FMetasoundFrontendVersionNumber::Parse(ClassVersionStr, Version);
+		}
+
+		{
+			FString AssetClassID;
+			bSuccess &= InAssetData.GetTagValue(AssetTagsPrivate::AssetClassID, AssetClassID);
+			ClassName = FMetasoundFrontendClassName(FName(), FName(*AssetClassID));
+		}
+
+		bIsValid = (int8)(bSuccess);
+		AssetPath = FTopLevelAssetPath(InAssetData.PackageName, InAssetData.AssetName);
+	}
+
+	void FAssetClassQueryInfo::ExportToContext(FAssetRegistryTagsContext& OutContext) const
+	{
+		using FAssetRegistryTag = UObject::FAssetRegistryTag;
+
+		OutContext.AddTag(FAssetRegistryTag(AssetTagsPrivate::AssetClassID, ClassName.Name.ToString(), FAssetRegistryTag::TT_Alphabetical));
+		OutContext.AddTag(FAssetRegistryTag(AssetTagsPrivate::ClassVersion, Version.ToString(), FAssetRegistryTag::TT_Alphabetical));
+
+#if WITH_EDITORONLY_DATA
+		auto SerializeTagToJson = []<typename TStructType>(const TStructType & InStruct)
+		{
+			FString JsonString;
+			FJsonObjectConverter::UStructToJsonObjectString(InStruct, JsonString);
+			return JsonString;
+		};
+
+		OutContext.AddTag(FAssetRegistryTag(AssetTagsPrivate::DocumentVersion, DocumentVersion.ToString(), FAssetRegistryTag::TT_Alphabetical));
+		OutContext.AddTag(FAssetRegistryTag(AssetTagsPrivate::IsPreset, FString(bIsPreset ? TEXT("1") : TEXT("0")), FAssetRegistryTag::TT_Numerical));
+
+		const FMetaSoundAssetReferenceTagCollection AssetReferences{ ReferencedAssetKeys };
+		OutContext.AddTag(FAssetRegistryTag(AssetTagsPrivate::ClassReferences, SerializeTagToJson(AssetReferences), FAssetRegistryTag::TT_Hidden));
+		OutContext.AddTag(FAssetRegistryTag(AssetTagsPrivate::DisplayInfo, SerializeTagToJson(DisplayInfo), FAssetRegistryTag::TT_Hidden));
+		OutContext.AddTag(FAssetRegistryTag(AssetTagsPrivate::Inputs, SerializeTagToJson(Inputs), FAssetRegistryTag::TT_Hidden));
+		OutContext.AddTag(FAssetRegistryTag(AssetTagsPrivate::Outputs, SerializeTagToJson(Outputs), FAssetRegistryTag::TT_Hidden));
+#endif // WITH_EDITORONLY_DATA
+	}
+
+	void FAssetClassQueryInfo::InitFromDocument(const IMetaSoundDocumentInterface& InDocInterface)
+	{
+		using namespace Frontend;
+
+		const FMetasoundFrontendDocument& Document = InDocInterface.GetConstDocument();
+		const FMetasoundFrontendClassMetadata& ClassMetadata = Document.RootGraph.Metadata;
+
+		Version = ClassMetadata.GetVersion();
+		AssetPath = InDocInterface.GetAssetPathChecked();
+
+#if WITH_EDITORONLY_DATA
+		bIsPreset = (uint8)Document.RootGraph.PresetOptions.bIsPreset;
+		DocumentVersion = Document.Metadata.Version.Number;
+
+		DisplayInfo.ClassDescription = ClassMetadata.GetDescription();
+		DisplayInfo.ClassDisplayName = ClassMetadata.GetDisplayName();
+		DisplayInfo.Hierarchy = ClassMetadata.GetCategoryHierarchy();
+		DisplayInfo.Keywords = ClassMetadata.GetKeywords();
+
+		const FMetasoundFrontendGraphClass& GraphClass = Document.RootGraph;
+
+		{
+			IMetaSoundAssetManager& AssetManager = IMetaSoundAssetManager::GetChecked();
+			auto IsAssetReference = [&AssetManager](const FMetasoundFrontendClass& Class)
+			{
+				return AssetManager.IsAssetClass(Class.Metadata);
+			};
+			
+			auto GetAssetKey = [](const FMetasoundFrontendClass& Class)
+			{
+				return FMetaSoundAssetKey(Class.Metadata);
+			};
+			Algo::TransformIf(Document.Dependencies, ReferencedAssetKeys, IsAssetReference, GetAssetKey);
+		}
+
+		{
+			auto CopyVertexInfo = [](const FMetasoundFrontendClassVertex& Vertex)
+			{
+				return FClassQueryVertexInfo
+				{
+					.Name = Vertex.Name,
+					.TypeName = Vertex.TypeName,
+					.AccessType = Vertex.AccessType
+				};
+			};
+			Algo::Transform(GraphClass.Interface.Inputs, Inputs.ClassVertexInfo, CopyVertexInfo);
+			Algo::Transform(GraphClass.Interface.Outputs, Outputs.ClassVertexInfo, CopyVertexInfo);
+		}
+#endif // WITH_EDITORONLY_DATA
+
+		bIsValid = 1;
 	}
 
 	IMetaSoundAssetManager* IMetaSoundAssetManager::Get()
@@ -104,3 +249,79 @@ namespace Metasound::Frontend
 		return false;
 	}
 } // namespace Metasound::Frontend
+
+FMetaSoundAssetKey::FMetaSoundAssetKey(const FMetasoundFrontendClassName& InClassName, const FMetasoundFrontendVersionNumber& InVersion)
+	: ClassName(InClassName)
+	, Version(InVersion)
+{
+}
+
+FMetaSoundAssetKey::FMetaSoundAssetKey(const Metasound::Frontend::FNodeClassRegistryKey& RegKey)
+{
+	using namespace Metasound::Frontend;
+
+	if (RegKey.Type != EMetasoundFrontendClassType::Invalid)
+	{
+		checkf(IsValidType(RegKey.Type), TEXT("Invalid ClassType '%s' for Registry Key"), LexToString(RegKey.Type));
+		ClassName = RegKey.ClassName;
+		Version = RegKey.Version;
+	}
+}
+
+FMetaSoundAssetKey::FMetaSoundAssetKey(const FMetasoundFrontendClassMetadata& InMetadata)
+{
+	using namespace Metasound::Frontend;
+
+	if (InMetadata.GetType() != EMetasoundFrontendClassType::Invalid)
+	{
+		checkf(IsValidType(InMetadata.GetType()), TEXT("Invalid ClassType '%s' for Registry Key"), LexToString(InMetadata.GetType()));
+		ClassName = InMetadata.GetClassName();
+		Version = InMetadata.GetVersion();
+	}
+}
+
+const FMetaSoundAssetKey& FMetaSoundAssetKey::GetInvalid()
+{
+	static FMetaSoundAssetKey Invalid;
+	return Invalid;
+}
+
+bool FMetaSoundAssetKey::IsValid() const
+{
+	return ClassName.IsValid() && Version.IsValid();
+}
+
+bool FMetaSoundAssetKey::IsValidType(EMetasoundFrontendClassType ClassType)
+{
+	switch (ClassType)
+	{
+		case EMetasoundFrontendClassType::External:
+		case EMetasoundFrontendClassType::Graph:
+			return true;
+
+		case EMetasoundFrontendClassType::Input:
+		case EMetasoundFrontendClassType::Literal:
+		case EMetasoundFrontendClassType::Output:
+		case EMetasoundFrontendClassType::Template:
+		case EMetasoundFrontendClassType::Variable:
+		case EMetasoundFrontendClassType::VariableDeferredAccessor:
+		case EMetasoundFrontendClassType::VariableAccessor:
+		case EMetasoundFrontendClassType::VariableMutator:
+		case EMetasoundFrontendClassType::Invalid:
+		default:
+			break;
+	}
+
+	return false;
+}
+
+FString FMetaSoundAssetKey::ToString() const
+{
+	TStringBuilder<128> KeyStringBuilder;
+	KeyStringBuilder.Append(ClassName.GetFullName().ToString());
+	KeyStringBuilder.AppendChar('_');
+	KeyStringBuilder.Append(FString::FromInt(Version.Major));
+	KeyStringBuilder.AppendChar('.');
+	KeyStringBuilder.Append(FString::FromInt(Version.Minor));
+	return KeyStringBuilder.ToString();
+}
diff --git a/Engine/Plugins/Runtime/Metasound/Source/MetasoundFrontend/Private/MetasoundAssetReferenceTagCollection.h b/Engine/Plugins/Runtime/Metasound/Source/MetasoundFrontend/Private/MetasoundAssetReferenceTagCollection.h
new file mode 100644
--- /dev/null
+++ b/Engine/Plugins/Runtime/Metasound/Source/MetasoundFrontend/Private/MetasoundAssetReferenceTagCollection.h
@@ -0,0 +1,20 @@
+// Copyright Epic Games, Inc. All Rights Reserved.
+#pragma once
+
+#include "MetasoundAssetKey.h"
+#include "UObject/Object.h"
+
+#include "MetasoundAssetReferenceTagCollection.generated.h"
+
+
+// Minimal Json-serializable array of references. Helper struct for asset tag serialization.
+USTRUCT()
+struct FMetaSoundAssetReferenceTagCollection
+{
+	GENERATED_BODY()
+
+#if WITH_EDITORONLY_DATA
+	UPROPERTY()
+	TArray<FMetaSoundAssetKey> AssetKeys;
+#endif // WITH_EDITORONLY_DATA
+};
\ No newline at end of file
diff --git a/Engine/Plugins/Runtime/Metasound/Source/MetasoundFrontend/Private/MetasoundFrontendDataTypeRegistry.cpp b/Engine/Plugins/Runtime/Metasound/Source/MetasoundFrontend/Private/MetasoundFrontendDataTypeRegistry.cpp
--- a/Engine/Plugins/Runtime/Metasound/Source/MetasoundFrontend/Private/MetasoundFrontendDataTypeRegistry.cpp
+++ b/Engine/Plugins/Runtime/Metasound/Source/MetasoundFrontend/Private/MetasoundFrontendDataTypeRegistry.cpp
@@ -137,13 +137,7 @@ namespace Metasound
 					return FVertexInterface();
 				}
 
-				virtual bool IsNative() const override
-				{
-					return true;
-				}
-
 			protected:
-
 				void UpdateNodeClassInfo(const FMetasoundFrontendClass& InClass)
 				{
 					ClassInfo = FNodeClassInfo(InClass.Metadata);
diff --git a/Engine/Plugins/Runtime/Metasound/Source/MetasoundFrontend/Private/MetasoundFrontendDocument.cpp b/Engine/Plugins/Runtime/Metasound/Source/MetasoundFrontend/Private/MetasoundFrontendDocument.cpp
--- a/Engine/Plugins/Runtime/Metasound/Source/MetasoundFrontend/Private/MetasoundFrontendDocument.cpp
+++ b/Engine/Plugins/Runtime/Metasound/Source/MetasoundFrontend/Private/MetasoundFrontendDocument.cpp
@@ -282,6 +282,48 @@ void FMetasoundFrontendDocumentModifyContext::AddNodeIDsModified(const TSet<FGui
 #endif // WITH_EDITORONLY_DATA
 
 
+const FMetasoundFrontendVersionNumber& FMetasoundFrontendVersionNumber::GetInvalid()
+{
+	static const FMetasoundFrontendVersionNumber Invalid{ 0, 0 };
+	return Invalid;
+}
+
+bool FMetasoundFrontendVersionNumber::IsValid() const
+{
+	return *this != GetInvalid();
+}
+
+bool FMetasoundFrontendVersionNumber::Parse(const FString& InString, FMetasoundFrontendVersionNumber& OutVersionNumber)
+{
+	if (!InString.StartsWith(TEXT("v")))
+	{
+		return false;
+	}
+
+	TArray<FString> Parts;
+	InString.ParseIntoArray(Parts, TEXT("."));
+	if (Parts.Num() != 2)
+	{
+		return false;
+	}
+
+	Parts[0].RightChopInline(1); // Remove 'v'
+
+	OutVersionNumber.Major = FCString::Atoi(*Parts[0]);
+	OutVersionNumber.Minor = FCString::Atoi(*Parts[1]);
+	return true;
+}
+
+Audio::FParameterInterface::FVersion FMetasoundFrontendVersionNumber::ToInterfaceVersion() const
+{
+	return Audio::FParameterInterface::FVersion{ Major, Minor };
+}
+
+FString FMetasoundFrontendVersionNumber::ToString() const
+{
+	return FString::Format(TEXT("v{0}.{1}"), { Major, Minor });
+}
+
 FMetasoundFrontendNodeInterface::FMetasoundFrontendNodeInterface(const FMetasoundFrontendClassInterface& InClassInterface)
 {
 	for (const FMetasoundFrontendClassInput& Input : InClassInterface.Inputs)
@@ -815,18 +857,21 @@ FMetasoundFrontendClassStyle FMetasoundFrontendClassStyle::GenerateClassStyle(co
 }
 #endif // WITH_EDITORONLY_DATA
 
+FMetasoundFrontendClassMetadata::FMetasoundFrontendClassMetadata()
+#if WITH_EDITORONLY_DATA
+PRAGMA_DISABLE_DEPRECATION_WARNINGS
+	: bAutoUpdateManagesInterface(false)
+PRAGMA_ENABLE_DEPRECATION_WARNINGS
+#endif // WITH_EDITORONLY_DATA
+{
+}
+
 FMetasoundFrontendClassMetadata FMetasoundFrontendClassMetadata::GenerateClassMetadata(const Metasound::FNodeClassMetadata& InNodeClassMetadata, EMetasoundFrontendClassType InType)
 {
 	FMetasoundFrontendClassMetadata NewMetadata;
 
 	NewMetadata.Type = InType;
 
-	// TODO: This flag is only used by the graph class' metadata.
-	// Should probably be moved elsewhere (AssetBase?) as to not
-	// get confused with behavior encapsulated on registry class
-	// descriptions/individual node class dependencies.
-	NewMetadata.bAutoUpdateManagesInterface = false;
-
 	NewMetadata.ClassName = InNodeClassMetadata.ClassName;
 	NewMetadata.Version = { InNodeClassMetadata.MajorVersion, InNodeClassMetadata.MinorVersion };
 
diff --git a/Engine/Plugins/Runtime/Metasound/Source/MetasoundFrontend/Private/MetasoundFrontendNodeClassRegistry.cpp b/Engine/Plugins/Runtime/Metasound/Source/MetasoundFrontend/Private/MetasoundFrontendNodeClassRegistry.cpp
--- a/Engine/Plugins/Runtime/Metasound/Source/MetasoundFrontend/Private/MetasoundFrontendNodeClassRegistry.cpp
+++ b/Engine/Plugins/Runtime/Metasound/Source/MetasoundFrontend/Private/MetasoundFrontendNodeClassRegistry.cpp
@@ -225,11 +225,17 @@ namespace Metasound::Frontend
 		class FDocumentNodeClassRegistryEntry : public INodeClassRegistryEntry
 		{
 		public:
-			FDocumentNodeClassRegistryEntry(const FMetasoundFrontendGraphClass& InGraphClass, const TSet<FMetasoundFrontendVersion>& InInterfaces, FNodeClassInfo&& InNodeClassInfo, TSharedPtr<const IGraph> InGraph)
+			FDocumentNodeClassRegistryEntry(
+				const FMetasoundFrontendGraphClass& InGraphClass,
+				const TSet<FMetasoundFrontendVersion>& InInterfaces,
+				FNodeClassInfo&& InNodeClassInfo,
+				TSharedPtr<const IGraph> InGraph,
+				FTopLevelAssetPath InAssetPath)
 			: FrontendClass(InGraphClass)
 			, Interfaces(InInterfaces)
 			, ClassInfo(MoveTemp(InNodeClassInfo))
 			, Graph(InGraph)
+			, AssetPath(InAssetPath)
 			{
 				FrontendClass.Metadata.SetType(EMetasoundFrontendClassType::External);
 			}
@@ -251,7 +257,7 @@ namespace Metasound::Frontend
 				}
 				else
 				{
-					UE_LOG(LogMetaSound, Error, TEXT("Cannot create MetaSound node from asset %s due to prior failure to build graph"), *ClassInfo.AssetPath.ToString());
+					UE_LOG(LogMetaSound, Error, TEXT("Cannot create MetaSound node from class %s due to prior failure to build graph"), *ClassInfo.ClassName.ToString());
 					return TUniquePtr<INode>();
 				}
 			}
@@ -278,17 +284,12 @@ namespace Metasound::Frontend
 				}
 			}
 
-			virtual bool IsNative() const override
-			{
-				return false;
-			}
-
 		private:
-
 			FMetasoundFrontendClass FrontendClass;
 			TSet<FMetasoundFrontendVersion> Interfaces;
 			FNodeClassInfo ClassInfo;
 			TSharedPtr<const IGraph> Graph;
+			FTopLevelAssetPath AssetPath;
 		};
 	} // namespace RegistryPrivate
 	  
@@ -419,31 +420,46 @@ PRAGMA_ENABLE_DEPRECATION_WARNINGS
 			return FNodeClassRegistryKey(InClassInfo);
 		}
 	} // namespace NodeClassRegistryKey
-	  
-	void FNodeClassRegistry::BuildAndRegisterGraphFromDocument(const FMetasoundFrontendDocument& InDocument, const FProxyDataCache& InProxyDataCache, FNodeClassInfo&& InNodeClassInfo)
+
+	void FNodeClassRegistry::BuildAndRegisterGraphFromDocument(const FMetasoundFrontendDocument& InDocument, const FProxyDataCache& InProxyDataCache, FNodeClassInfo&& InNodeClassInfo, const FTopLevelAssetPath& AssetPath)
 	{
 		using namespace RegistryPrivate;
 
 		METASOUND_TRACE_CPUPROFILER_EVENT_SCOPE(Metasound::FNodeClassRegistry::BuildAndRegisterGraphFromDocument);
-		METASOUND_TRACE_CPUPROFILER_EVENT_SCOPE_TEXT(*FString::Printf(TEXT("Metasound::FNodeClassRegistry::BuildAndRegisterGraphFromDocument asset %s"), *InNodeClassInfo.AssetPath.ToString()));
+		METASOUND_TRACE_CPUPROFILER_EVENT_SCOPE_TEXT(*FString::Printf(TEXT("Metasound::FNodeClassRegistry::BuildAndRegisterGraphFromDocument asset %s"), *AssetPath.ToString()));
 
-		// Use the asset class id for the graph id because it should be locally unique. Unlike other ids, it is regenerated on asset duplicate. 
-		TUniquePtr<FFrontendGraph> FrontendGraph = FGraphBuilder::CreateGraph(InDocument, InProxyDataCache, InNodeClassInfo.AssetPath.ToString(), /*GraphId=*/InNodeClassInfo.AssetClassID);
-		if (!FrontendGraph.IsValid())
+		FGuid AssetClassID;
+		if (IMetaSoundAssetManager* AssetManager = IMetaSoundAssetManager::Get())
+		{
+			ensureAlways(AssetManager->TryGetAssetIDFromClassName(InNodeClassInfo.ClassName, AssetClassID));
+		}
+		else
 		{
-			UE_LOG(LogMetaSound, Error, TEXT("Failed to build MetaSound graph in asset '%s'"), *InNodeClassInfo.AssetPath.ToString());
+			UE_LOG(LogMetaSound, Warning, TEXT("No AssetManager registered, registering bespoke AssetClassID '%s' for asset '%s'"), *AssetClassID.ToString(), *AssetPath.ToString());
+			AssetClassID = FGuid::NewGuid();
 		}
 
-		TSharedPtr<const FGraph> GraphToRegister = MakeShareable<const FGraph>(FrontendGraph.Release());
+		// Use the asset class id for the graph id because it should be locally unique per asset.
+		TUniquePtr<FFrontendGraph> FrontendGraph = Frontend::FGraphBuilder::CreateGraph(InDocument, InProxyDataCache, AssetPath.ToString(), /*GraphId=*/AssetClassID);
+		if (!FrontendGraph.IsValid())
+		{
+			UE_LOG(LogMetaSound, Error, TEXT("Failed to build MetaSound graph in asset '%s'"), *AssetPath.ToString());
+		}
 
-		TUniquePtr<INodeClassRegistryEntry> RegistryEntry = MakeUnique<FDocumentNodeClassRegistryEntry>(InDocument.RootGraph, InDocument.Interfaces, MoveTemp(InNodeClassInfo), GraphToRegister);
+		TSharedPtr<const FGraph> GraphToRegister(FrontendGraph.Release());
+		TUniquePtr<INodeClassRegistryEntry> RegistryEntry = MakeUnique<FDocumentNodeClassRegistryEntry>(
+			InDocument.RootGraph,
+			InDocument.Interfaces,
+			MoveTemp(InNodeClassInfo),
+			GraphToRegister,
+			AssetPath);
 
 		const FNodeClassRegistryKey RegistryKey = RegisterNodeInternal(MoveTemp(RegistryEntry));
 
-		// Key must use the asset path provided to class info and *NOT* that of the
-		// document's owning DocumentInterface object, as that may be a built transient object
-		// with a different transient asset path.
-		const FGraphRegistryKey GraphKey { RegistryKey, InNodeClassInfo.AssetPath };
+		// Key must use the asset path provided to function and *NOT* that of the
+		// document's owning DocumentInterface object, as that may be a built/optimized
+		// transient object with a different, transient asset path.
+		const FGraphRegistryKey GraphKey { RegistryKey, AssetPath };
 		RegisterGraphInternal(GraphKey, GraphToRegister);
 	}
 
@@ -649,7 +665,7 @@ PRAGMA_ENABLE_DEPRECATION_WARNINGS
 
 		// Store update to newly registered node in history so nodes
 		// can be queried by transaction ID
-		FNodeClassInfo NodeClassInfo(Document.RootGraph, AssetPath);
+		FNodeClassInfo NodeClassInfo(Document.RootGraph);
 		{
 			FNodeClassRegistryTransaction::FTimeType Timestamp = FPlatformTime::Cycles64();
 			TransactionBuffer->AddTransaction(FNodeClassRegistryTransaction(FNodeClassRegistryTransaction::ETransactionType::NodeRegistration, NodeClassInfo, Timestamp));
@@ -659,12 +675,12 @@ PRAGMA_ENABLE_DEPRECATION_WARNINGS
 		{
 			Tasks::FTask BuildAndRegisterTask = AsyncRegistrationPipe.Launch(
 				UE_SOURCE_LOCATION,
-				[RegistryKey, ClassInfo = MoveTemp(NodeClassInfo), RegistryDocInterface, ProxyDataCache = MoveTemp(ProxyDataCache)]() mutable
+				[RegistryKey, ClassInfo = MoveTemp(NodeClassInfo), AssetPath, RegistryDocInterface, ProxyDataCache = MoveTemp(ProxyDataCache)]() mutable
 				{
 					FNodeClassRegistry& Registry = FNodeClassRegistry::Get();
 					// Unregister the graph before re-registering
 					Registry.UnregisterGraphInternal(RegistryKey);
-					Registry.BuildAndRegisterGraphFromDocument(RegistryDocInterface->GetConstDocument(), ProxyDataCache, MoveTemp(ClassInfo));
+					Registry.BuildAndRegisterGraphFromDocument(RegistryDocInterface->GetConstDocument(), ProxyDataCache, MoveTemp(ClassInfo), AssetPath);
 					Registry.RemoveRegistrationTask(RegistryKey, FNodeClassRegistryTransaction::ETransactionType::NodeRegistration);
 					Registry.RemoveDocumentReference(RegistryDocInterface);
 				}
@@ -683,7 +699,7 @@ PRAGMA_ENABLE_DEPRECATION_WARNINGS
 			UnregisterGraphInternal(RegistryKey);
 
 			// Build and register graph synchronously
-			BuildAndRegisterGraphFromDocument(RegistryDocInterface->GetConstDocument(), ProxyDataCache, MoveTemp(NodeClassInfo));
+			BuildAndRegisterGraphFromDocument(RegistryDocInterface->GetConstDocument(), ProxyDataCache, MoveTemp(NodeClassInfo), AssetPath);
 		}
 
 		return RegistryKey;
@@ -916,24 +932,40 @@ PRAGMA_ENABLE_DEPRECATION_WARNINGS
 		const FNodeClassRegistryKey Key = FNodeClassRegistryKey(InEntry->GetClassInfo());
 		METASOUND_TRACE_CPUPROFILER_EVENT_SCOPE_TEXT(*Key.ToString(TEXT("FNodeClassRegistry::RegisterNodeInternal")))
 
-		TSharedRef<INodeClassRegistryEntry, ESPMode::ThreadSafe> Entry(InEntry.Release());
+
+#if !NO_LOGGING
+		TArray<TSharedRef<INodeClassRegistryEntry>> Entries;
+#endif // !NO_LOGGING
+
 		{
+			TSharedRef<INodeClassRegistryEntry, ESPMode::ThreadSafe> Entry(InEntry.Release());
 			FScopeLock Lock(&RegistryMapsCriticalSection);
+			// check to see if an identical node was already registered, and log if necessary
+			// Store registry elements in map so nodes can be queried using registry key.
+			RegisteredNodes.Add(Key, Entry);
 
-			// check to see if an identical node was already registered, and log
-			if (const TSharedRef<INodeClassRegistryEntry>* ExistingEntry = RegisteredNodes.Find(Key))
+#if !NO_LOGGING
+			RegisteredNodes.MultiFind(Key, Entries);
+#endif // !NO_LOGGING
+		}
+
+#if !NO_LOGGING
+		if (Entries.Num() > 1)
+		{
+			if (IMetaSoundAssetManager* AssetManager = IMetaSoundAssetManager::Get())
 			{
-				const FNodeClassInfo& ClassInfo = (*ExistingEntry)->GetClassInfo();
+				const TArray<FTopLevelAssetPath> AssetPaths = AssetManager->FindAssetPaths(Frontend::FAssetKey(Key));
+				TArray<FString> ExistingPaths;
+				Algo::Transform(AssetPaths, ExistingPaths, [](const FTopLevelAssetPath& AssetPath) { return AssetPath.ToString(); });
+				const FString ExistingAssetPaths = FString::Join(ExistingPaths, TEXT("\n"));
 				UE_LOG(LogMetaSound, Error,
-					TEXT("Node with registry key '%s' already registered by asset '%s' encountered while registering node with asset %s. MetaSounds which depend on these assets may utilize incorrect asset."),
-					*Key.ToString(),
-					*ClassInfo.AssetPath.ToString(),
-					*Entry->GetClassInfo().AssetPath.ToString());
+					TEXT("Multiple node classes with key '%s' registered. Assets currently registered with class name:\n%s"),
+						* Key.ToString(),
+						*ExistingAssetPaths
+				);
 			}
-
-			// Store registry elements in map so nodes can be queried using registry key.
-			RegisteredNodes.Add(Key, Entry);
 		}
+#endif // !NO_LOGGING
 
 		return Key;
 	}
@@ -1107,26 +1139,6 @@ PRAGMA_ENABLE_DEPRECATION_WARNINGS
 		}
 	}
 
-	bool FNodeClassRegistry::IsNodeNative(const FNodeClassRegistryKey& InKey) const
-	{
-		bool bIsNative = false;
-		auto SetIsNative = [&bIsNative](const INodeClassRegistryEntry& Entry) 
-		{ 
-			bIsNative = Entry.IsNative();
-		};
-
-		if (AccessNodeEntryThreadSafe(InKey, SetIsNative))
-		{
-			return bIsNative;
-		}
-
-		if (const INodeTemplateRegistryEntry* TemplateEntry = FindNodeTemplateEntry(InKey))
-		{
-			return true;
-		}
-		return false;
-	}
-
 	bool FNodeClassRegistry::FindDefaultVertexInterface(const FNodeClassRegistryKey& InKey, FVertexInterface& OutVertexInterface) const 
 	{
 		auto GetDefaultVertexInterface = [&OutVertexInterface](const INodeClassRegistryEntry& Entry)
diff --git a/Engine/Plugins/Runtime/Metasound/Source/MetasoundFrontend/Private/MetasoundFrontendNodeClassRegistryPrivate.h b/Engine/Plugins/Runtime/Metasound/Source/MetasoundFrontend/Private/MetasoundFrontendNodeClassRegistryPrivate.h
--- a/Engine/Plugins/Runtime/Metasound/Source/MetasoundFrontend/Private/MetasoundFrontendNodeClassRegistryPrivate.h
+++ b/Engine/Plugins/Runtime/Metasound/Source/MetasoundFrontend/Private/MetasoundFrontendNodeClassRegistryPrivate.h
@@ -88,7 +88,6 @@ namespace Metasound::Frontend
 		virtual bool UnregisterNode(const FNodeClassRegistryKey& InKey) override;
 		virtual bool IsNodeRegistered(const FNodeClassRegistryKey& InKey) const override;
 		virtual bool IsGraphRegistered(const FGraphRegistryKey& InKey) const override;
-		virtual bool IsNodeNative(const FNodeClassRegistryKey& InKey) const override;
 
 		virtual bool RegisterConversionNode(const FConverterNodeClassRegistryKey& InNodeKey, const FConverterNodeInfo& InNodeInfo) override;
 
@@ -139,7 +138,7 @@ namespace Metasound::Frontend
 			FTopLevelAssetPath AssetPath;
 		};
 
-		void BuildAndRegisterGraphFromDocument(const FMetasoundFrontendDocument& InDocument, const FProxyDataCache& InProxyDataCache, FNodeClassInfo&& InNodeClassInfo);
+		void BuildAndRegisterGraphFromDocument(const FMetasoundFrontendDocument& InDocument, const FProxyDataCache& InProxyDataCache, FNodeClassInfo&& InNodeClassInfo, const FTopLevelAssetPath& AssetPath);
 
 		void AddRegistrationTask(const FGraphRegistryKey& InKey, FActiveRegistrationTaskInfo&& TaskInfo);
 		void RemoveRegistrationTask(const FGraphRegistryKey& InKey, FNodeClassRegistryTransaction::ETransactionType TransactionType);
diff --git a/Engine/Plugins/Runtime/Metasound/Source/MetasoundFrontend/Private/MetasoundFrontendNodeController.cpp b/Engine/Plugins/Runtime/Metasound/Source/MetasoundFrontend/Private/MetasoundFrontendNodeController.cpp
--- a/Engine/Plugins/Runtime/Metasound/Source/MetasoundFrontend/Private/MetasoundFrontendNodeController.cpp
+++ b/Engine/Plugins/Runtime/Metasound/Source/MetasoundFrontend/Private/MetasoundFrontendNodeController.cpp
@@ -1035,16 +1035,9 @@ namespace Metasound
 				// define changes in native vs asset class definitions. Need to
 				// hash a changeID from natively defined node classes in order to
 				// merge branches.
-				const FNodeRegistryKey RegistryKey = FNodeRegistryKey(RegistryClass.Metadata);
-				const bool bIsClassNative = FMetasoundFrontendRegistryContainer::Get()->IsNodeNative(RegistryKey);
-				if (bIsClassNative)
-				{
-					if (!MetaSoundAutoUpdateNativeClassesOfEqualVersionCVar)
-					{
-						return false;
-					}
-				}
-				else
+				const IMetaSoundAssetManager* AssetManager = IMetaSoundAssetManager::Get();
+				const bool bIsAssetClass = AssetManager && AssetManager->IsAssetClass(RegistryClass.Metadata);
+				if (bIsAssetClass)
 				{
 					if (RegistryClass.Metadata.GetChangeID() == NodeClassMetadata.GetChangeID())
 					{
@@ -1055,6 +1048,13 @@ namespace Metasound
 						}
 					}
 				}
+				else
+				{
+					if (!MetaSoundAutoUpdateNativeClassesOfEqualVersionCVar)
+					{
+						return false;
+					}
+				}
 
 				DiffAgainstRegistryInterface(OutInterfaceUpdates, true /* bUseHighestMinorVersion */);
 				return OutInterfaceUpdates.ContainsChanges();
diff --git a/Engine/Plugins/Runtime/Metasound/Source/MetasoundFrontend/Private/MetasoundFrontendRegistryKey.cpp b/Engine/Plugins/Runtime/Metasound/Source/MetasoundFrontend/Private/MetasoundFrontendRegistryKey.cpp
--- a/Engine/Plugins/Runtime/Metasound/Source/MetasoundFrontend/Private/MetasoundFrontendRegistryKey.cpp
+++ b/Engine/Plugins/Runtime/Metasound/Source/MetasoundFrontend/Private/MetasoundFrontendRegistryKey.cpp
@@ -10,29 +10,134 @@ namespace Metasound::Frontend
 		auto GetVertexTypeName = [](const FMetasoundFrontendVertex& Vertex) { return Vertex.TypeName; };
 	}
 
+	FNodeClassInfo::FNodeClassInfo()
+		PRAGMA_DISABLE_DEPRECATION_WARNINGS
+		: AssetClassID()
+		, AssetPath()
+#if WITH_EDITORONLY_DATA
+		, InputTypes()
+		, OutputTypes()
+		, bIsPreset(false)
+#endif // WITH_EDITORONLY_DATA
+		PRAGMA_ENABLE_DEPRECATION_WARNINGS
+	{
+	}
+
 	FNodeClassInfo::FNodeClassInfo(const FMetasoundFrontendClassMetadata& InMetadata)
+		PRAGMA_DISABLE_DEPRECATION_WARNINGS
 		: ClassName(InMetadata.GetClassName())
 		, Type(InMetadata.GetType())
+		, AssetClassID(FGuid(InMetadata.GetClassName().Name.ToString()))
+		, AssetPath()
 		, Version(InMetadata.GetVersion())
+#if WITH_EDITORONLY_DATA
+		, InputTypes()
+		, OutputTypes()
+		, bIsPreset(false)
+#endif // WITH_EDITORONLY_DATA
+		PRAGMA_ENABLE_DEPRECATION_WARNINGS
+	{
+	}
+
+	FNodeClassInfo::FNodeClassInfo(const FMetasoundFrontendGraphClass& InClass)
+		PRAGMA_DISABLE_DEPRECATION_WARNINGS
+		: ClassName(InClass.Metadata.GetClassName())
+		, Type(EMetasoundFrontendClassType::External) // Overridden as it is considered the same as an external class in registries
+		, AssetClassID(FGuid(ClassName.Name.ToString()))
+		, AssetPath()
+		, Version(InClass.Metadata.GetVersion())
+#if WITH_EDITORONLY_DATA
+		, InputTypes()
+		, OutputTypes()
+		, bIsPreset(false)
+#endif // WITH_EDITORONLY_DATA
+		PRAGMA_ENABLE_DEPRECATION_WARNINGS
 	{
 	}
 
 	FNodeClassInfo::FNodeClassInfo(const FMetasoundFrontendGraphClass& InClass, const FTopLevelAssetPath& InAssetPath)
+		PRAGMA_DISABLE_DEPRECATION_WARNINGS
 		: ClassName(InClass.Metadata.GetClassName())
 		, Type(EMetasoundFrontendClassType::External) // Overridden as it is considered the same as an external class in registries
 		, AssetClassID(FGuid(ClassName.Name.ToString()))
 		, AssetPath(InAssetPath)
 		, Version(InClass.Metadata.GetVersion())
+#if WITH_EDITORONLY_DATA
+		, InputTypes()
+		, OutputTypes()
+		, bIsPreset(false)
+#endif // WITH_EDITORONLY_DATA
+		PRAGMA_ENABLE_DEPRECATION_WARNINGS
 	{
-		using namespace NodeClassInfoPrivate;
+	}
 
-		ensure(!AssetPath.IsNull());
+	FNodeClassInfo::FNodeClassInfo(const FNodeClassInfo& Other)
+		PRAGMA_DISABLE_DEPRECATION_WARNINGS
+		: ClassName(Other.ClassName)
+		, Type(Other.Type)
+		, AssetClassID(Other.AssetClassID)
+		, AssetPath(Other.AssetPath)
+		, Version(Other.Version)
+#if WITH_EDITORONLY_DATA
+		, InputTypes(Other.InputTypes)
+		, OutputTypes(Other.OutputTypes)
+		, bIsPreset(Other.bIsPreset)
+#endif // WITH_EDITORONLY_DATA
+		PRAGMA_ENABLE_DEPRECATION_WARNINGS
+	{
+	}
+
+	FNodeClassInfo::FNodeClassInfo(FNodeClassInfo&& Other)
+		PRAGMA_DISABLE_DEPRECATION_WARNINGS
+		: ClassName(MoveTemp(Other.ClassName))
+		, Type(Other.Type)
+		, AssetClassID(MoveTemp(Other.AssetClassID))
+		, AssetPath(MoveTemp(Other.AssetPath))
+		, Version(MoveTemp(Other.Version))
+#if WITH_EDITORONLY_DATA
+		, InputTypes(MoveTemp(Other.InputTypes))
+		, OutputTypes(MoveTemp(Other.OutputTypes))
+		, bIsPreset(Other.bIsPreset)
+#endif // WITH_EDITORONLY_DATA
+		PRAGMA_ENABLE_DEPRECATION_WARNINGS
+	{
+	}
+
+	FNodeClassInfo& FNodeClassInfo::operator=(const FNodeClassInfo& Other)
+	{
+		PRAGMA_DISABLE_DEPRECATION_WARNINGS
+		ClassName = Other.ClassName;
+		Type = Other.Type;
+		AssetClassID = Other.AssetClassID;
+		AssetPath = Other.AssetPath;
+		Version = Other.Version;
 
 #if WITH_EDITORONLY_DATA
-		Algo::Transform(InClass.Interface.Inputs, InputTypes, GetVertexTypeName);
-		Algo::Transform(InClass.Interface.Outputs, OutputTypes, GetVertexTypeName);
-		bIsPreset = InClass.PresetOptions.bIsPreset;
+		InputTypes = Other.InputTypes;
+		OutputTypes = Other.OutputTypes;
+		bIsPreset = Other.bIsPreset;
 #endif // WITH_EDITORONLY_DATA
+		PRAGMA_ENABLE_DEPRECATION_WARNINGS
+
+		return *this;
+	}
+
+	FNodeClassInfo& FNodeClassInfo::operator=(FNodeClassInfo&& Other)
+	{
+		PRAGMA_DISABLE_DEPRECATION_WARNINGS
+		ClassName = MoveTemp(Other.ClassName);
+		Type = Other.Type;
+		AssetClassID = MoveTemp(Other.AssetClassID);
+		AssetPath = MoveTemp(Other.AssetPath);
+		Version = MoveTemp(Other.Version);
+#if WITH_EDITORONLY_DATA
+		InputTypes = MoveTemp(Other.InputTypes);
+		OutputTypes = MoveTemp(Other.OutputTypes);
+		bIsPreset = Other.bIsPreset;
+#endif // WITH_EDITORONLY_DATA
+		PRAGMA_ENABLE_DEPRECATION_WARNINGS
+
+		return *this;
 	}
 
 	FNodeClassRegistryKey::FNodeClassRegistryKey(EMetasoundFrontendClassType InType, const FMetasoundFrontendClassName& InClassName, int32 InMajorVersion, int32 InMinorVersion)
diff --git a/Engine/Plugins/Runtime/Metasound/Source/MetasoundFrontend/Private/MetasoundNodeRegistrationMacro.cpp b/Engine/Plugins/Runtime/Metasound/Source/MetasoundFrontend/Private/MetasoundNodeRegistrationMacro.cpp
--- a/Engine/Plugins/Runtime/Metasound/Source/MetasoundFrontend/Private/MetasoundNodeRegistrationMacro.cpp
+++ b/Engine/Plugins/Runtime/Metasound/Source/MetasoundFrontend/Private/MetasoundNodeRegistrationMacro.cpp
@@ -42,11 +42,6 @@ namespace Metasound::Frontend
 			return ClassMetadata->DefaultInterface;
 		}
 
-		bool FNodeRegistryEntryBase::IsNative() const
-		{
-			return true;
-		}
-
 		TSharedRef<const FNodeClassMetadata> FNodeRegistryEntryBase::GetNodeClassMetadata() const
 		{
 			return ClassMetadata;
diff --git a/Engine/Plugins/Runtime/Metasound/Source/MetasoundFrontend/Public/MetasoundAssetBase.h b/Engine/Plugins/Runtime/Metasound/Source/MetasoundFrontend/Public/MetasoundAssetBase.h
--- a/Engine/Plugins/Runtime/Metasound/Source/MetasoundFrontend/Public/MetasoundAssetBase.h
+++ b/Engine/Plugins/Runtime/Metasound/Source/MetasoundFrontend/Public/MetasoundAssetBase.h
@@ -74,9 +74,8 @@ public:
 	// @param Editor graph associated with this metasound object.
 	virtual void SetGraph(UEdGraph* InGraph) = 0;
 
-	// Only required for editor builds. Adds metadata to properties available when the object is
-	// not loaded for use by the Asset Registry.
-	virtual void SetRegistryAssetClassInfo(const Metasound::Frontend::FNodeClassInfo& InClassInfo) = 0;
+	UE_DEPRECATED(5.6, "AssetClass tags now directly serialized using UObject GetAssetRegistryTags call")
+	virtual void SetRegistryAssetClassInfo(const Metasound::Frontend::FNodeClassInfo& InClassInfo) { }
 #endif // WITH_EDITORONLY_DATA
 
 	UE_DEPRECATED(5.5, "Moved to IMetaSoundDocumentInterface::ConformObjectToDocument")
@@ -110,10 +109,6 @@ public:
 	// Returns whether an interface with the given version is declared by the given asset's document.
 	bool IsInterfaceDeclared(const FMetasoundFrontendVersion& InVersion) const;
 
-	// Gets the asset class info.
-	UE_DEPRECATED(5.4, "NodeClassInfo can be constructed directly from document's root graph & asset's path and requires no specialized virtual getter.")
-	virtual Metasound::Frontend::FNodeClassInfo GetAssetClassInfo() const;
-
 	// Returns all the class keys of this asset's referenced assets
 	virtual const TSet<FString>& GetReferencedAssetClassKeys() const = 0;
 
@@ -235,38 +230,6 @@ protected:
 	// Returns an access pointer to the document.
 	virtual Metasound::Frontend::FConstDocumentAccessPtr GetDocumentConstAccessPtr() const = 0;
 
-	// Waits for a graph to be registered in the scenario where a graph is registered on an async task.
-	//
-	// When a graph is registered, the underlying IMetaSoundDocumentInterface may be accessed on an
-	// async tasks. If modifications need to be made to the IMetaSoundDocumentInterface, callers should
-	// wait for the inflight graph registration to complete by calling this method.
-	
-	// Container for runtime data of MetaSound graph.
-	struct FRuntimeData_DEPRECATED
-	{
-		// Current ID of graph.
-		FGuid ChangeID;
-
-		// Array of inputs which can be set for construction. 
-		TArray<FMetasoundFrontendClassInput> PublicInputs;
-
-		// Array of inputs which can be transmitted to.
-		TArray<FMetasoundFrontendClassInput> TransmittableInputs;
-
-		// Core graph.
-		TSharedPtr<Metasound::FGraph, ESPMode::ThreadSafe> Graph;
-	};
-
-	struct UE_DEPRECATED(5.4, "FRuntimeData is no longer used to store runtime graphs and inputs. Runtime graphs are stored in the node registry. Runtime inputs are stored on the UMetaSoundSoruce") FRuntimeData : FRuntimeData_DEPRECATED
-	{
-	};
-
-	// Returns the cached runtime data.
-	PRAGMA_DISABLE_DEPRECATION_WARNINGS
-	UE_DEPRECATED(5.4, "Access to graph and public inputs has moved. Use the node registry to access the graph and GetPublicClassInputs() to access public inputs")
-	const FRuntimeData& GetRuntimeData() const;
-	PRAGMA_ENABLE_DEPRECATION_WARNINGS
-
 	UE_DEPRECATED(5.5, "AutoUpdate implementation now private and implemented within 'Version Dependencies'")
 	bool AutoUpdate(bool bInLogWarningsOnDroppedConnection);
 
@@ -276,12 +239,8 @@ protected:
 	// Ensures all referenced graph classes are registered (or re-registers depending on options).
 	void RegisterAssetDependencies(const Metasound::Frontend::FMetaSoundAssetRegistrationOptions & InRegistrationOptions);
 
-	UE_DEPRECATED(5.4, "Template node transformation moved to private implementation. A MetaSound asset will likely never have the function Process(...). Without a Process function, you cannot have preprocessing.")
-	TSharedPtr<FMetasoundFrontendDocument> PreprocessDocument();
-
 private:
 #if WITH_EDITORONLY_DATA
-	void UpdateAssetRegistry();
 	void UpdateAndRegisterReferencesForSerialization(FName CookPlatformName);
 #endif // WITH_EDITORONLY_DATA
 
diff --git a/Engine/Plugins/Runtime/Metasound/Source/MetasoundFrontend/Public/MetasoundAssetKey.h b/Engine/Plugins/Runtime/Metasound/Source/MetasoundFrontend/Public/MetasoundAssetKey.h
new file mode 100644
--- /dev/null
+++ b/Engine/Plugins/Runtime/Metasound/Source/MetasoundFrontend/Public/MetasoundAssetKey.h
@@ -0,0 +1,69 @@
+// Copyright Epic Games, Inc. All Rights Reserved.
+#pragma once
+
+#include "MetasoundFrontendDocument.h"
+#include "MetasoundFrontendRegistryKey.h"
+#include "MetasoundVertex.h"
+#include "UObject/Package.h"
+#include "UObject/SoftObjectPath.h"
+#include "UObject/TopLevelAssetPath.h"
+#include "UObject/UObjectGlobals.h"
+#include "UObject/WeakObjectPtrTemplates.h"
+
+#include "MetasoundAssetKey.generated.h"
+
+
+USTRUCT()
+struct METASOUNDFRONTEND_API FMetaSoundAssetKey
+{
+	GENERATED_BODY()
+
+	UPROPERTY()
+	FMetasoundFrontendClassName ClassName;
+
+	UPROPERTY()
+	FMetasoundFrontendVersionNumber Version;
+
+	FMetaSoundAssetKey() = default;
+	FMetaSoundAssetKey(const FMetasoundFrontendClassName& InClassName, const FMetasoundFrontendVersionNumber& InVersion);
+	FMetaSoundAssetKey(const Metasound::Frontend::FNodeClassRegistryKey& RegKey);
+	FMetaSoundAssetKey(const FMetasoundFrontendClassMetadata& InMetadata);
+
+	FORCEINLINE friend bool operator==(const FMetaSoundAssetKey& InLHS, const FMetaSoundAssetKey& InRHS)
+	{
+		return (InLHS.ClassName == InRHS.ClassName) && (InLHS.Version == InRHS.Version);
+	}
+
+	FORCEINLINE friend bool operator<(const FMetaSoundAssetKey& InLHS, const FMetaSoundAssetKey& InRHS)
+	{
+		if (InLHS.ClassName == InRHS.ClassName)
+		{
+			return InLHS.Version < InRHS.Version;
+		}
+
+		return InLHS.ClassName < InRHS.ClassName;
+	}
+
+	friend FORCEINLINE uint32 GetTypeHash(const FMetaSoundAssetKey& InKey)
+	{
+		return HashCombineFast(GetTypeHash(InKey.ClassName), GetTypeHash(InKey.Version));
+	}
+
+	// Returns invalid asset key.
+	static const FMetaSoundAssetKey& GetInvalid();
+
+	// Returns whether or not key is valid.
+	bool IsValid() const;
+
+	// Returns whether or not ClassType is supported by asset/asset key.
+	static bool IsValidType(EMetasoundFrontendClassType ClassType);
+
+	// Returns string representation of asset key.
+	FString ToString() const;
+};
+
+namespace Metasound::Frontend
+{
+	// Struct moved to global namespace and made a USTRUCT. This alias is here for back compat.
+	using FAssetKey = FMetaSoundAssetKey;
+} // namespace Metasound::Frontend
diff --git a/Engine/Plugins/Runtime/Metasound/Source/MetasoundFrontend/Public/MetasoundAssetManager.h b/Engine/Plugins/Runtime/Metasound/Source/MetasoundFrontend/Public/MetasoundAssetManager.h
--- a/Engine/Plugins/Runtime/Metasound/Source/MetasoundFrontend/Public/MetasoundAssetManager.h
+++ b/Engine/Plugins/Runtime/Metasound/Source/MetasoundFrontend/Public/MetasoundAssetManager.h
@@ -1,42 +1,65 @@
 // Copyright Epic Games, Inc. All Rights Reserved.
 #pragma once
 
+#include "JsonObjectConverter.h"
+#include "MetasoundAssetKey.h"
+#include "MetasoundDocumentInterface.h"
 #include "MetasoundFrontendDocument.h"
+#include "MetasoundFrontendQuery.h"
 #include "MetasoundFrontendRegistryKey.h"
 #include "MetasoundVertex.h"
+#include "UObject/AssetRegistryTagsContext.h"
+#include "UObject/Object.h"
 #include "UObject/Package.h"
 #include "UObject/SoftObjectPath.h"
+#include "UObject/TopLevelAssetPath.h"
+#include "UObject/UObjectGlobals.h"
 #include "UObject/WeakObjectPtrTemplates.h"
 
+
 // Forward Declarations
 class FMetasoundAssetBase;
 class UEdGraph;
 
+class IMetaSoundDocumentInterface;
+
 struct FAssetData;
 struct FMetasoundFrontendClassName;
 struct FMetaSoundFrontendDocumentBuilder;
-class IMetaSoundDocumentInterface;
+
 
 namespace Metasound::Frontend
 {
 	namespace AssetTags
 	{
+		UE_DEPRECATED(5.6, "AssetTags no longer public. Construct FAssetClassQueryInfo from AssetData to parse tag fields.")
 		extern const FString METASOUNDFRONTEND_API ArrayDelim;
 
 #if WITH_EDITORONLY_DATA
+		UE_DEPRECATED(5.6, "AssetTags no longer public. Construct FAssetClassQueryInfo from AssetData to parse tag fields.")
 		extern const FName METASOUNDFRONTEND_API IsPreset;
 #endif // WITH_EDITORONLY_DATA
 
+
+		UE_DEPRECATED(5.6, "AssetTags no longer public. Construct FAssetClassQueryInfo from AssetData to parse tag fields.")
 		extern const FName METASOUNDFRONTEND_API AssetClassID;
+
+		UE_DEPRECATED(5.6, "AssetTags no longer public. Construct FAssetClassQueryInfo from AssetData to parse tag fields.")
 		extern const FName METASOUNDFRONTEND_API RegistryVersionMajor;
+
+		UE_DEPRECATED(5.6, "AssetTags no longer public. Construct FAssetClassQueryInfo from AssetData to parse tag fields.")
 		extern const FName METASOUNDFRONTEND_API RegistryVersionMinor;
 
 #if WITH_EDITORONLY_DATA
+		UE_DEPRECATED(5.6, "AssetTags no longer public. Construct FAssetClassQueryInfo from AssetData to parse tag fields.")
 		extern const FName METASOUNDFRONTEND_API RegistryInputTypes;
+
+		UE_DEPRECATED(5.6, "AssetTags no longer public. Construct FAssetClassQueryInfo from AssetData to parse tag fields.")
 		extern const FName METASOUNDFRONTEND_API RegistryOutputTypes;
 #endif // WITH_EDITORONLY_DATA
 	} // namespace AssetTags
 
+
 	struct METASOUNDFRONTEND_API FMetaSoundAssetRegistrationOptions
 	{
 		// If true, forces a re-register of this class (and all class dependencies
@@ -70,54 +93,51 @@ namespace Metasound::Frontend
 #endif // WITH_EDITOR
 	};
 
-	struct METASOUNDFRONTEND_API FAssetKey
+
+	/** Contains a minimal set of information needed to display or query additional
+	  * class data from the internal MetaSound Asset or loaded Node Class Registries. This data is
+	  * serializable to asset tags, and is available even if asset is unloaded / not registered with
+	  * the MetaSound Node Class Registry. */
+	struct METASOUNDFRONTEND_API FAssetClassQueryInfo : public FClassQueryInfo
 	{
-		FMetasoundFrontendClassName ClassName;
-		FMetasoundFrontendVersionNumber Version;
+		FAssetClassQueryInfo();
+		FAssetClassQueryInfo(const IMetaSoundDocumentInterface& InDocInterface);
 
-		FAssetKey() = default;
-		FAssetKey(const FMetasoundFrontendClassName& InClassName, const FMetasoundFrontendVersionNumber& InVersion);
-		FAssetKey(const FNodeRegistryKey& RegKey);
-		FAssetKey(const FMetasoundFrontendClassMetadata& InMetadata);
+		// Attempts to return tag data from the given AssetData, which may or may not be loaded.
+		// Invalid tag constructed if object isn't loaded and any fields fail to load/are not serialized
+		// (i.e. tags are out-of-date).
+		FAssetClassQueryInfo(const FAssetData& InAssetData);
 
-		FORCEINLINE friend bool operator==(const FAssetKey& InLHS, const FAssetKey& InRHS)
-		{
-			return (InLHS.ClassName == InRHS.ClassName) && (InLHS.Version == InRHS.Version);
-		}
+		// Exports tag data to the given RegistryContext
+		void ExportToContext(FAssetRegistryTagsContext& OutContext) const;
 
-		FORCEINLINE friend bool operator<(const FAssetKey& InLHS, const FAssetKey& InRHS)
-		{
-			if (InLHS.ClassName == InRHS.ClassName)
-			{
-				return InLHS.Version < InRHS.Version;
-			}
+	private:
+		void InitFromDocument(const IMetaSoundDocumentInterface& InDocInterface);
 
-			return InLHS.ClassName < InRHS.ClassName;
-		}
+	public:
+		// Path to asset containing graph if external type and references asset class.
+		FTopLevelAssetPath AssetPath;
 
-		friend FORCEINLINE uint32 GetTypeHash(const FAssetKey& InKey)
-		{
-			return HashCombineFast(GetTypeHash(InKey.ClassName), GetTypeHash(InKey.Version));
-		}
+#if WITH_EDITORONLY_DATA
+		// Document Version of asset
+		FMetasoundFrontendVersionNumber DocumentVersion;
 
-		static const FAssetKey& GetInvalid()
-		{
-			static FAssetKey Invalid;
-			return Invalid;
-		}
+		// Versions of referenced asset class keys
+		TArray<FMetaSoundAssetKey> ReferencedAssetKeys;
 
-		bool IsValid() const
-		{
-			return ClassName.IsValid() && Version.IsValid();
-		}
+		// Whether asset is a preset or not.
+		uint8 bIsPreset : 1;
+#endif // WITH_EDITORONLY_DATA
 
-		FString ToString() const;
+		// True, asset tag data is invalid due to provided asset not being loaded and out-of-date
+		uint8 bIsValid : 1;
 	};
 
 
 	class METASOUNDFRONTEND_API IMetaSoundAssetManager
 	{
 	public:
+
 		virtual ~IMetaSoundAssetManager() = default;
 
 		static IMetaSoundAssetManager* Get();
@@ -164,11 +184,25 @@ namespace Metasound::Frontend
 		// to prime system from asset attempting to register prior to asset scan being complete.
 		// Returns true if references were added, false if it they are already found.
 		virtual bool AddAssetReferences(FMetasoundAssetBase& InAssetBase) = 0;
-#endif
+#endif // WITH_EDITORONLY_DATA
+
+		UE_DEPRECATED(5.6, "Moved to AddOrUpdateFromObject")
+		virtual FAssetKey AddOrUpdateAsset(const UObject& InObject) { return { }; }
+
+		UE_DEPRECATED(5.6, "Moved to AddOrLoadAndUpdateFromObjectAsync")
+		virtual FAssetKey AddOrUpdateAsset(const FAssetData& InAssetData) { return { }; }
+
+		// Add or Update a MetaSound Asset's entry data from a loaded MetaSound asset UObject.
+		virtual FAssetKey AddOrUpdateFromObject(const UObject& InObject) = 0;
 
-		// Add or Update a MetaSound Asset's entry data
-		virtual FAssetKey AddOrUpdateAsset(const UObject& InObject) = 0;
-		virtual FAssetKey AddOrUpdateAsset(const FAssetData& InAssetData) = 0;
+		// Add or Update a MetaSound Asset's entry data from an object, loading it if it isn't already.  On initial call, requests object load
+		// asynchronously and runs provided function on successful completion. If asset is already loaded, runs provided function immediately
+		// on entry update (synchronously).
+		virtual void AddOrLoadAndUpdateFromObjectAsync(const FAssetData& InAssetData, TUniqueFunction<void(FAssetKey, UObject&)>&& OnAssetRegistered) = 0;
+
+		// Add or Update a MetaSound Asset's entry data from AssetTags. Potentially loads asset and added asynchronously if asset is or asset tags
+		// are out-of-date.
+		virtual void AddOrUpdateFromAssetTags(const FAssetData& InAssetData) = 0;
 
 		// Whether or not the class is eligible for auto-update
 		virtual bool CanAutoUpdate(const FMetasoundFrontendClassName& InClassName) const = 0;
@@ -176,6 +210,10 @@ namespace Metasound::Frontend
 		// Whether or not the asset manager has loaded the given asset
 		virtual bool ContainsKey(const Metasound::Frontend::FAssetKey& InAssetKey) const = 0;
 
+		// Whether or not the asset manager has loaded one or more assets with the given registry key.
+		// Returns false if key is not valid asset key (ex. input or output class key, variable, etc.)
+		virtual bool ContainsKey(const Metasound::Frontend::FNodeRegistryKey& InRegistryKey) const = 0;
+
 		// Returns object (if loaded) associated with the given key (null if key not registered with the AssetManager)
 		// If multiple assets are associated with the given key, the last one is returned. 
 		virtual FMetasoundAssetBase* FindAsset(const Metasound::Frontend::FAssetKey& InAssetKey) const = 0;
@@ -199,8 +237,11 @@ namespace Metasound::Frontend
 		virtual TSet<FAssetInfo> GetReferencedAssetClasses(const FMetasoundAssetBase& InAssetBase) const = 0;
 #endif // WITH_EDITOR
 
-		// Iterate all known MetaSound asset paths
-		virtual void IterateAssets(TFunctionRef<void(const FAssetKey, const TArray<FTopLevelAssetPath>&)> Iter) const = 0;
+		// returns whether or not the given class is defined as a registered asset
+		virtual bool IsAssetClass(const FMetasoundFrontendClassMetadata& ClassMetadata) const = 0;
+
+		UE_DEPRECATED(5.6, "Moved to Iterate ClassQueryInfo")
+		virtual void IterateAssets(TFunctionRef<void(const FAssetKey, const TArray<FTopLevelAssetPath>&)> Iter) const { }
 
 		UE_DEPRECATED(5.5, "Rescan no longer supported nor required by Frontend")
 		virtual void RescanAutoUpdateDenyList() { }
@@ -230,7 +271,7 @@ namespace Metasound::Frontend
 		// Requests an async load of all async referenced assets of the input asset.
 		virtual void RequestAsyncLoadReferencedAssets(FMetasoundAssetBase& InAssetBase) = 0;
 
-		// Synchronously requests unregister and reregister of all loaded MetaSound assets node class entries.
+		// Synchronously requests unregister and re-register of all loaded MetaSound assets node class entries.
 		virtual void ReloadMetaSoundAssets() const = 0;
 
 		// Removes object from MetaSound asset manager
diff --git a/Engine/Plugins/Runtime/Metasound/Source/MetasoundFrontend/Public/MetasoundDocumentInterface.h b/Engine/Plugins/Runtime/Metasound/Source/MetasoundFrontend/Public/MetasoundDocumentInterface.h
--- a/Engine/Plugins/Runtime/Metasound/Source/MetasoundFrontend/Public/MetasoundDocumentInterface.h
+++ b/Engine/Plugins/Runtime/Metasound/Source/MetasoundFrontend/Public/MetasoundDocumentInterface.h
@@ -16,10 +16,6 @@
 // Forward Declarations
 struct FMetaSoundFrontendDocumentBuilder;
 
-namespace Metasound::Frontend
-{
-	struct FAssetKey;
-}
 
 // UInterface for all MetaSound UClasses that implement a MetaSound document
 // as a means for accessing data via code, scripting, execution, or node
diff --git a/Engine/Plugins/Runtime/Metasound/Source/MetasoundFrontend/Public/MetasoundFrontendDocument.h b/Engine/Plugins/Runtime/Metasound/Source/MetasoundFrontend/Public/MetasoundFrontendDocument.h
--- a/Engine/Plugins/Runtime/Metasound/Source/MetasoundFrontend/Public/MetasoundFrontendDocument.h
+++ b/Engine/Plugins/Runtime/Metasound/Source/MetasoundFrontend/Public/MetasoundFrontendDocument.h
@@ -212,6 +212,8 @@ struct METASOUNDFRONTEND_API FMetasoundFrontendVersionNumber
 {
 	GENERATED_BODY()
 
+	static bool Parse(const FString& InString, FMetasoundFrontendVersionNumber& OutVersionNumber);
+
 	// Major version number.
 	UPROPERTY(VisibleAnywhere, BlueprintReadOnly, Category = General)
 	int32 Major = 1;
@@ -220,21 +222,10 @@ struct METASOUNDFRONTEND_API FMetasoundFrontendVersionNumber
 	UPROPERTY(VisibleAnywhere, BlueprintReadOnly, Category = General)
 	int32 Minor = 0;
 
-	static const FMetasoundFrontendVersionNumber& GetInvalid()
-	{
-		static const FMetasoundFrontendVersionNumber Invalid { 0, 0 };
-		return Invalid;
-	}
-
-	bool IsValid() const
-	{
-		return *this != GetInvalid();
-	}
+	static const FMetasoundFrontendVersionNumber& GetInvalid();
+	bool IsValid() const;
 
-	Audio::FParameterInterface::FVersion ToInterfaceVersion() const
-	{
-		return Audio::FParameterInterface::FVersion { Major, Minor };
-	}
+	Audio::FParameterInterface::FVersion ToInterfaceVersion() const;
 
 	friend bool operator==(const FMetasoundFrontendVersionNumber& InLHS, const FMetasoundFrontendVersionNumber& InRHS)
 	{
@@ -286,10 +277,7 @@ struct METASOUNDFRONTEND_API FMetasoundFrontendVersionNumber
 		return InLHS == InRHS || InLHS < InRHS;
 	}
 
-	FString ToString() const
-	{
-		return FString::Format(TEXT("v{0}.{1}"), { Major, Minor });
-	}
+	FString ToString() const;
 
 	friend FORCEINLINE uint32 GetTypeHash(const FMetasoundFrontendVersionNumber& InNumber)
 	{
@@ -1410,6 +1398,8 @@ struct METASOUNDFRONTEND_API FMetasoundFrontendClassMetadata
 {
 	GENERATED_BODY()
 
+	FMetasoundFrontendClassMetadata();
+
 	// Generates class metadata intended to be used as a registry descriptor from FNodeClassMetadata. Does not initialize a change ID as it is not considered to be transactional.
 	static FMetasoundFrontendClassMetadata GenerateClassMetadata(const Metasound::FNodeClassMetadata& InNodeClassMetadata, EMetasoundFrontendClassType InType);
 
@@ -1464,13 +1454,13 @@ private:
 	UPROPERTY(EditAnywhere, Category = Metasound)
 	bool bIsDeprecated = false;
 
+#if WITH_EDITORONLY_DATA
 	// If true, auto-update will manage (add and remove)
 	// inputs/outputs associated with internally connected
 	// nodes when the interface of the given node is auto-updated.
-	UPROPERTY()
+	UPROPERTY(meta = (Deprecated = 5.6, DeprecationMessage = "Boolean no longer observed (auto-update rules are managed by project settings now)"))
 	bool bAutoUpdateManagesInterface = false;
 
-#if WITH_EDITORONLY_DATA
 	// Whether or not the given metadata text should be serialized
 	// or is procedurally maintained via auto-update & the referenced
 	// registry class (to avoid localization text desync).  Should be
@@ -1608,13 +1598,18 @@ public:
 		// ChangeID = FGuid::NewGuid();
 	}
 
+#if WITH_EDITORONLY_DATA
 	// Deprecated field in favor of GraphClass PresetOptions
 	bool GetAndClearAutoUpdateManagesInterface_Deprecated()
 	{
+PRAGMA_DISABLE_DEPRECATION_WARNINGS
 		bool bToReturn = bAutoUpdateManagesInterface;
 		bAutoUpdateManagesInterface = false;
+PRAGMA_ENABLE_DEPRECATION_WARNINGS
+
 		return bToReturn;
 	}
+#endif // WITH_EDITORONLY_DATA
 };
 
 
diff --git a/Engine/Plugins/Runtime/Metasound/Source/MetasoundFrontend/Public/MetasoundFrontendNodeClassRegistry.h b/Engine/Plugins/Runtime/Metasound/Source/MetasoundFrontend/Public/MetasoundFrontendNodeClassRegistry.h
--- a/Engine/Plugins/Runtime/Metasound/Source/MetasoundFrontend/Public/MetasoundFrontendNodeClassRegistry.h
+++ b/Engine/Plugins/Runtime/Metasound/Source/MetasoundFrontend/Public/MetasoundFrontendNodeClassRegistry.h
@@ -105,8 +105,9 @@ namespace Metasound::Frontend
 
 		virtual FVertexInterface GetDefaultVertexInterface() const = 0;
 
-		/** Whether or not the node is natively defined */
-		virtual bool IsNative() const = 0;
+		UE_DEPRECATED(5.6, "Node class registry no longer tracks nature of implementation (i.e. if was registered "
+			"from an asset or c++. Use the MetaSoundAssetManager to determine if the class has been defined within an asset(s).")
+		virtual bool IsNative() const { return false; }
 	};
 
 	using INodeRegistryEntry UE_DEPRECATED(5.6, "Use INodeClassRegistryEntry instead") = INodeClassRegistryEntry;
@@ -309,8 +310,9 @@ namespace Metasound::Frontend
 		/** Returns true if the provided registry key (node key and asset path) corresponds to a valid registered graph. */
 		virtual bool IsGraphRegistered(const FGraphRegistryKey& InKey) const = 0;
 
-		/** Returns true if the provided registry key corresponds to a valid registered node that is natively defined. */
-		virtual bool IsNodeNative(const FNodeClassRegistryKey& InKey) const = 0;
+		UE_DEPRECATED(5.6, "Node class registry no longer tracks donor asset state/implementation at the interface "
+			"level. Use the MetaSoundAssetManager to determine if the class has been defined within an asset(s).")
+		virtual bool IsNodeNative(const FNodeClassRegistryKey& InKey) const { return false; }
 
 		// Iterates class types in registry.  If InClassType is set to a valid class type (optional), only iterates classes of the given type
 		virtual void IterateRegistry(FIterateMetasoundFrontendClassFunction InIterFunc, EMetasoundFrontendClassType InClassType = EMetasoundFrontendClassType::Invalid) const = 0;
diff --git a/Engine/Plugins/Runtime/Metasound/Source/MetasoundFrontend/Public/MetasoundFrontendQuery.h b/Engine/Plugins/Runtime/Metasound/Source/MetasoundFrontend/Public/MetasoundFrontendQuery.h
--- a/Engine/Plugins/Runtime/Metasound/Source/MetasoundFrontend/Public/MetasoundFrontendQuery.h
+++ b/Engine/Plugins/Runtime/Metasound/Source/MetasoundFrontend/Public/MetasoundFrontendQuery.h
@@ -6,10 +6,14 @@
 #include "Containers/SortedMap.h"
 #include "Interfaces/MetasoundFrontendInterfaceRegistry.h"
 #include "Math/NumericLimits.h"
+#include "MetasoundAssetKey.h"
 #include "MetasoundFrontendDocument.h"
 #include "MetasoundFrontendRegistries.h"
 #include "Misc/TVariant.h"
 
+#include "MetasoundFrontendQuery.generated.h"
+
+
 /** MetaSound Frontend Query
  *
  * MetaSound Frontend Query provides a way to systematically organize and update
@@ -45,8 +49,97 @@
  * 	Limit:		Limit the size of a FFrontendQueryPartition.
  */
 
+
+// Forward Declarations
+class FAssetRegistryTagsContext;
+class IMetaSoundDocumentInterface;
+
+
+// Minimal Json-serializable vertex data which characterizes an assets interface vertex for use in tag data.
+USTRUCT()
+struct METASOUNDFRONTEND_API FClassQueryVertexInfo
+{
+	GENERATED_BODY()
+
+#if WITH_EDITORONLY_DATA
+	UPROPERTY()
+	FName Name;
+
+	UPROPERTY()
+	FName TypeName;
+
+	UPROPERTY()
+	EMetasoundFrontendVertexAccessType AccessType = EMetasoundFrontendVertexAccessType::Unset;
+#endif // WITH_EDITORONLY_DATA
+};
+
+USTRUCT()
+struct METASOUNDFRONTEND_API FClassQueryDisplayInfo
+{
+	GENERATED_BODY()
+
+#if WITH_EDITORONLY_DATA
+	// Human readable DisplayName of Class (optional, overrides the
+	// package name in the editor if specified by MetaSound Asset Author).
+	UPROPERTY()
+	FText ClassDisplayName;
+
+	UPROPERTY()
+	FText ClassDescription;
+
+	UPROPERTY()
+	TArray<FText> Hierarchy;
+
+	UPROPERTY()
+	TArray<FText> Keywords;
+#endif // WITH_EDITORONLY_DATA
+};
+
+// Minimal Json-serializable array of vertex data, which characterizes an assets interface for use in tag data.
+USTRUCT()
+struct METASOUNDFRONTEND_API FClassQueryVertexCollection
+{
+	GENERATED_BODY()
+
+#if WITH_EDITORONLY_DATA
+	UPROPERTY()
+	TArray<FClassQueryVertexInfo> ClassVertexInfo;
+#endif // WITH_EDITORONLY_DATA
+};
+
+
 namespace Metasound
 {
+	namespace Frontend
+	{
+		// Class query info accessible in the Search Engine.  Represents class info for a given class
+		// that may or may not be loaded and/or registered (ex. in editor, class info may be supplied by
+		// asset registry tags or class info could be provided by a cloud service, enabling class browsing
+		// without a local asset required to be loaded or registered with the node class registry).
+		struct METASOUNDFRONTEND_API FClassQueryInfo
+		{
+			// ClassName of class
+			FMetasoundFrontendClassName ClassName;
+
+			// Version of class
+			FMetasoundFrontendVersionNumber Version;
+
+#if WITH_EDITORONLY_DATA
+			// Editor-only display info, included for fast look-up and 
+			FClassQueryDisplayInfo DisplayInfo;
+
+			// Collection of identifiable input vertex data cached in query for fast access & serializability
+			// (ex. in asset tags)
+			FClassQueryVertexCollection Inputs;
+
+			// Collection of identifiable output vertex data cached in query for fast access & serializability
+			// (ex. in asset tags)
+			FClassQueryVertexCollection Outputs;
+#endif // WITH_EDITORONLY_DATA
+		};
+	} // namespace Frontend
+
+
 	/** FFrontendQueryKey allows entries to be partitioned by their key. A key
 	 * can be created by default constructor, int32, FString or FName.
 	 */
diff --git a/Engine/Plugins/Runtime/Metasound/Source/MetasoundFrontend/Public/MetasoundFrontendRegistryKey.h b/Engine/Plugins/Runtime/Metasound/Source/MetasoundFrontend/Public/MetasoundFrontendRegistryKey.h
--- a/Engine/Plugins/Runtime/Metasound/Source/MetasoundFrontend/Public/MetasoundFrontendRegistryKey.h
+++ b/Engine/Plugins/Runtime/Metasound/Source/MetasoundFrontend/Public/MetasoundFrontendRegistryKey.h
@@ -8,10 +8,12 @@
 #include "UObject/TopLevelAssetPath.h"
 
 
+struct FMetaSoundAssetKey;
+
 namespace Metasound::Frontend
 {
 	// Forward Declarations
-	struct FAssetKey;
+	using FAssetKey = FMetaSoundAssetKey;
 
 	/** FNodeClassInfo contains a minimal set of information needed to find and query node classes. */
 	struct METASOUNDFRONTEND_API FNodeClassInfo
@@ -22,35 +24,44 @@ namespace Metasound::Frontend
 		// The type of this node class
 		EMetasoundFrontendClassType Type = EMetasoundFrontendClassType::Invalid;
 
-		// The ID used for the Asset Classes. If zero, class is natively defined.
+		UE_DEPRECATED(5.6, "Data now characterized in asset class data")
 		FGuid AssetClassID;
 
-		// Path to asset containing graph if external type and references asset class.
+		UE_DEPRECATED(5.6, "Data now characterized in asset class data")
 		FTopLevelAssetPath AssetPath;
 
 		// Version of the registered class
 		FMetasoundFrontendVersionNumber Version;
 
 #if WITH_EDITORONLY_DATA
-		// Types of class inputs
+		UE_DEPRECATED(5.6, "Data now characterized in asset class data ('Inputs')")
 		TSet<FName> InputTypes;
 
-		// Types of class outputs
+		UE_DEPRECATED(5.6, "Type data no longer required by class info ('Outputs')")
 		TSet<FName> OutputTypes;
 
-		// Whether or not class is preset
-		bool bIsPreset = false;
+		UE_DEPRECATED(5.6, "Preset data no longer required by class info")
+		bool bIsPreset;
 #endif // WITH_EDITORONLY_DATA
 
-		FNodeClassInfo() = default;
+		FNodeClassInfo();
 
 		// Constructor used to generate NodeClassInfo from a class' Metadata.
 		// (Does not cache AssetPath and thus may not support loading asset
 		// should the class originate from one).
 		FNodeClassInfo(const FMetasoundFrontendClassMetadata& InMetadata);
 
-		// Constructor used to generate NodeClassInfo from an asset
+		// Constructor used to generate NodeClassInfo from a graph class
+		FNodeClassInfo(const FMetasoundFrontendGraphClass& InClass);
+
+		UE_DEPRECATED(5.6, "Asset data no longer contained on NodeClassInfo, see AssetClass Info")
 		FNodeClassInfo(const FMetasoundFrontendGraphClass& InClass, const FTopLevelAssetPath& InAssetPath);
+
+		FNodeClassInfo(const FNodeClassInfo& Other);
+		FNodeClassInfo(FNodeClassInfo&& Other);
+
+		FNodeClassInfo& operator=(const FNodeClassInfo& Other);
+		FNodeClassInfo& operator=(FNodeClassInfo&& Other);
 	};
 
 	struct METASOUNDFRONTEND_API FNodeClassRegistryKey
diff --git a/Engine/Plugins/Runtime/Metasound/Source/MetasoundFrontend/Public/MetasoundNodeRegistrationMacro.h b/Engine/Plugins/Runtime/Metasound/Source/MetasoundFrontend/Public/MetasoundNodeRegistrationMacro.h
--- a/Engine/Plugins/Runtime/Metasound/Source/MetasoundFrontend/Public/MetasoundNodeRegistrationMacro.h
+++ b/Engine/Plugins/Runtime/Metasound/Source/MetasoundFrontend/Public/MetasoundNodeRegistrationMacro.h
@@ -40,8 +40,6 @@ namespace Metasound::Frontend
 			virtual const TSet<FMetasoundFrontendVersion>* GetImplementedInterfaces() const override;
 
 			virtual FVertexInterface GetDefaultVertexInterface() const override;
-			
-			virtual bool IsNative() const override;
 
 		protected:
 
