Fix ignored default input on sugraph nodes.

#rnx

diff --git a/Engine/Plugins/Runtime/Metasound/Source/MetasoundFrontend/Private/MetasoundFrontendDataTypeRegistry.cpp b/Engine/Plugins/Runtime/Metasound/Source/MetasoundFrontend/Private/MetasoundFrontendDataTypeRegistry.cpp
--- a/Engine/Plugins/Runtime/Metasound/Source/MetasoundFrontend/Private/MetasoundFrontendDataTypeRegistry.cpp
+++ b/Engine/Plugins/Runtime/Metasound/Source/MetasoundFrontend/Private/MetasoundFrontendDataTypeRegistry.cpp
@@ -1530,7 +1530,7 @@ PRAGMA_ENABLE_DEPRECATION_WARNINGS
 			return Registry;
 		}
 
-		void CreateDefaults(const FOperatorSettings& InOperatorSettings, FInputVertexInterfaceData& OutVertexData)
+		void CreateDefaultsInternal(const FOperatorSettings& InOperatorSettings, bool bInOverrideExistingData, FInputVertexInterfaceData& OutVertexData)
 		{
 			using namespace VertexDataPrivate;
 
@@ -1551,6 +1551,12 @@ PRAGMA_ENABLE_DEPRECATION_WARNINGS
 
 			for (const FInputBinding& Binding : OutVertexData)
 			{
+				if (!bInOverrideExistingData && Binding.IsBound())
+				{
+					// Do not create defaults if data is already set. 
+					continue;
+				}
+
 				// Attempt to create default data reference from the literal stored
 				// on the input vertex.
 				const FInputDataVertex& InputVertex = Binding.GetVertex();
@@ -1581,5 +1587,18 @@ PRAGMA_ENABLE_DEPRECATION_WARNINGS
 				}
 			}
 		}
+
+		void CreateDefaultsIfNotBound(const FOperatorSettings& InOperatorSettings, FInputVertexInterfaceData& OutVertexData)
+		{
+			constexpr bool bOverrideExistingData = false;
+			CreateDefaultsInternal(InOperatorSettings, bOverrideExistingData, OutVertexData);
+		}
+
+		void CreateDefaults(const FOperatorSettings& InOperatorSettings, FInputVertexInterfaceData& OutVertexData)
+		{
+			constexpr bool bOverrideExistingData = true;
+			CreateDefaultsInternal(InOperatorSettings, bOverrideExistingData, OutVertexData);
+		}
+
 	}
 }
diff --git a/Engine/Plugins/Runtime/Metasound/Source/MetasoundFrontend/Private/MetasoundFrontendNodeClassRegistry.cpp b/Engine/Plugins/Runtime/Metasound/Source/MetasoundFrontend/Private/MetasoundFrontendNodeClassRegistry.cpp
--- a/Engine/Plugins/Runtime/Metasound/Source/MetasoundFrontend/Private/MetasoundFrontendNodeClassRegistry.cpp
+++ b/Engine/Plugins/Runtime/Metasound/Source/MetasoundFrontend/Private/MetasoundFrontendNodeClassRegistry.cpp
@@ -130,11 +130,18 @@ namespace Metasound::Frontend
 
 				virtual TUniquePtr<IOperator> CreateOperator(const FBuildOperatorParams& InParams, FBuildResults& OutResults) override
 				{
+					// Create default data references for inputs which do not have data references
+					// already provided. These need to be created here, but the defaults
+					// cannot be piped into the wrapped graph. The wrapped graph cannot
+					// be modified. 
+					FInputVertexInterfaceData InputData = InParams.InputData;
+					CreateDefaultsIfNotBound(InParams.OperatorSettings, InputData);
+
 					FBuildOperatorParams ForwardParams 
 					{
 						*Graph,  // Point to correct INode instance
 						InParams.OperatorSettings,
-						InParams.InputData,
+						InputData,
 						InParams.Environment,
 						InParams.Builder,
 						InParams.GraphRenderCost
@@ -148,11 +155,15 @@ namespace Metasound::Frontend
 
 		public:
 			FGraphNode(const FNodeInitData& InNodeInitData, TSharedRef<const IGraph> InGraphToWrap)
-			: InstanceID(InNodeInitData.InstanceID)
+			: FGraphNode(FNodeData(InNodeInitData.InstanceName, InNodeInitData.InstanceID, InGraphToWrap->GetMetadata().DefaultInterface), MoveTemp(InGraphToWrap))
+			{
+			}
+
+			FGraphNode(FNodeData InNodeData, TSharedRef<const IGraph> InGraphToWrap)
+			: NodeData(MoveTemp(InNodeData))
 			, Factory(MakeShared<FGraphOperatorFactoryAdapter>(*InGraphToWrap))
 			, Graph(MoveTemp(InGraphToWrap))
 			{
-				Interface = Graph->GetVertexInterface();
 			}
 
 			virtual const FName& GetInstanceName() const override
@@ -164,7 +175,7 @@ namespace Metasound::Frontend
 
 			virtual const FGuid& GetInstanceID() const override
 			{
-				return InstanceID;
+				return NodeData.ID;
 			}
 
 			virtual const FNodeClassMetadata& GetMetadata() const override
@@ -174,12 +185,12 @@ namespace Metasound::Frontend
 
 			virtual const FVertexInterface& GetVertexInterface() const override
 			{
-				return Interface;
+				return NodeData.Interface;
 			}
 
 			virtual void SetDefaultInput(const FVertexName& InVertexName, const FLiteral& InLiteral) override
 			{
-				if (FInputDataVertex* Vertex = Interface.GetInputInterface().Find(InVertexName))
+				if (FInputDataVertex* Vertex = NodeData.Interface.GetInputInterface().Find(InVertexName))
 				{
 					Vertex->SetDefaultLiteral(InLiteral);
 				}
@@ -189,24 +200,6 @@ namespace Metasound::Frontend
 				}
 			}
 
-			UE_DEPRECATED(5.6, "Mutable vertex interfaces are no longer part of the node interface.")
-			virtual bool SetVertexInterface(const FVertexInterface& InInterface) override
-			{
-				PRAGMA_DISABLE_DEPRECATION_WARNINGS
-				// Cannot set vertex interface because of const reference to IGraph
-				// Return true if the supplied interface is the same as the existing interface.
-				return Interface == InInterface;
-				PRAGMA_ENABLE_DEPRECATION_WARNINGS
-			}
-
-			UE_DEPRECATED(5.6, "Mutable vertex interfaces are no longer part of the node interface.")
-			virtual bool IsVertexInterfaceSupported(const FVertexInterface& InInterface) const override
-			{
-				PRAGMA_DISABLE_DEPRECATION_WARNINGS
-				return Graph->IsVertexInterfaceSupported(InInterface);
-				PRAGMA_ENABLE_DEPRECATION_WARNINGS
-			}
-
 			virtual FOperatorFactorySharedRef GetDefaultOperatorFactory() const override
 			{
 				return Factory;
@@ -214,10 +207,9 @@ namespace Metasound::Frontend
 
 		private:
 
-			FGuid InstanceID;
+			FNodeData NodeData;
 			TSharedRef<FGraphOperatorFactoryAdapter> Factory;
 			TSharedRef<const IGraph> Graph;
-			FVertexInterface Interface;
 		};
 
 
@@ -262,6 +254,19 @@ namespace Metasound::Frontend
 				}
 			}
 
+			virtual TUniquePtr<INode> CreateNode(FNodeData InNodeData) const override
+			{
+				if (Graph.IsValid())
+				{
+					return MakeUnique<FGraphNode>(MoveTemp(InNodeData), Graph.ToSharedRef());
+				}
+				else
+				{
+					UE_LOG(LogMetaSound, Error, TEXT("Cannot create MetaSound node from asset %s due to prior failure to build graph"), *ClassInfo.AssetPath.ToString());
+					return TUniquePtr<INode>();
+				}
+			}
+
 			virtual const FMetasoundFrontendClass& GetFrontendClass() const override
 			{
 				return FrontendClass;
diff --git a/Engine/Plugins/Runtime/Metasound/Source/MetasoundFrontend/Public/MetasoundFrontendDataTypeRegistry.h b/Engine/Plugins/Runtime/Metasound/Source/MetasoundFrontend/Public/MetasoundFrontendDataTypeRegistry.h
--- a/Engine/Plugins/Runtime/Metasound/Source/MetasoundFrontend/Public/MetasoundFrontendDataTypeRegistry.h
+++ b/Engine/Plugins/Runtime/Metasound/Source/MetasoundFrontend/Public/MetasoundFrontendDataTypeRegistry.h
@@ -458,5 +458,14 @@ namespace Metasound
 		 * data reference will be set, otherwise a value data reference will be set.
 		 */
 		METASOUNDFRONTEND_API void CreateDefaults(const FOperatorSettings& InOperatorSettings, FInputVertexInterfaceData& OutVertexData);
+
+		/** Create data references for all inputs on the FInputVertexInterfaceData 
+		 * which do not already have a data reference bound. Defaults are created
+		 * using the default literals which live on the FInputDataVertexes.
+		 *
+		 * If the vertex has access type EVertexAccessType::Reference, a writable
+		 * data reference will be set, otherwise a value data reference will be set.
+		 */
+		METASOUNDFRONTEND_API void CreateDefaultsIfNotBound(const FOperatorSettings& InOperatorSettings, FInputVertexInterfaceData& OutVertexData);
 	}
 }
