Build fix for deprecated asset path in MetaSounds

diff --git a/Engine/Plugins/Runtime/Metasound/Source/MetasoundFrontend/Private/MetasoundFrontendNodeClassRegistry.cpp b/Engine/Plugins/Runtime/Metasound/Source/MetasoundFrontend/Private/MetasoundFrontendNodeClassRegistry.cpp
--- a/Engine/Plugins/Runtime/Metasound/Source/MetasoundFrontend/Private/MetasoundFrontendNodeClassRegistry.cpp
+++ b/Engine/Plugins/Runtime/Metasound/Source/MetasoundFrontend/Private/MetasoundFrontendNodeClassRegistry.cpp
@@ -262,7 +262,7 @@ namespace Metasound::Frontend
 				}
 				else
 				{
-					UE_LOG(LogMetaSound, Error, TEXT("Cannot create MetaSound node from asset %s due to prior failure to build graph"), *ClassInfo.AssetPath.ToString());
+					UE_LOG(LogMetaSound, Error, TEXT("Cannot create MetaSound node from asset %s due to prior failure to build graph"), *AssetPath.ToString());
 					return TUniquePtr<INode>();
 				}
 			}
