Unblock build with fix for ensure during attempt to remove asset with invalid tag data (likely CDO, further investigation pending)
#rnx

diff --git a/Engine/Plugins/Runtime/Metasound/Source/MetasoundEngine/Private/MetasoundAssetSubsystem.cpp b/Engine/Plugins/Runtime/Metasound/Source/MetasoundEngine/Private/MetasoundAssetSubsystem.cpp
--- a/Engine/Plugins/Runtime/Metasound/Source/MetasoundEngine/Private/MetasoundAssetSubsystem.cpp
+++ b/Engine/Plugins/Runtime/Metasound/Source/MetasoundEngine/Private/MetasoundAssetSubsystem.cpp
@@ -891,7 +891,7 @@ namespace Metasound::Engine
 		using namespace Frontend;
 
 		const FAssetClassQueryInfo QueryInfo(InAssetData);
-		if (ensureAlways(QueryInfo.bIsValid))
+		if (QueryInfo.bIsValid)
 		{
 			if (IDocumentBuilderRegistry* BuilderRegistry = IDocumentBuilderRegistry::Get())
 			{
