Fix errors in cook due to missing class name population in asset class query info
#rnx

diff --git a/Engine/Plugins/Runtime/Metasound/Source/MetasoundFrontend/Private/MetasoundAssetManager.cpp b/Engine/Plugins/Runtime/Metasound/Source/MetasoundFrontend/Private/MetasoundAssetManager.cpp
--- a/Engine/Plugins/Runtime/Metasound/Source/MetasoundFrontend/Private/MetasoundAssetManager.cpp
+++ b/Engine/Plugins/Runtime/Metasound/Source/MetasoundFrontend/Private/MetasoundAssetManager.cpp
@@ -173,6 +173,7 @@ PRAGMA_ENABLE_DEPRECATION_WARNINGS
 		const FMetasoundFrontendDocument& Document = InDocInterface.GetConstDocument();
 		const FMetasoundFrontendClassMetadata& ClassMetadata = Document.RootGraph.Metadata;
 
+		ClassName = Document.RootGraph.Metadata.GetClassName();
 		Version = ClassMetadata.GetVersion();
 		AssetPath = InDocInterface.GetAssetPathChecked();
 
