Fix for cook causing errors when attempting to build presets round 2
#rnx

diff --git a/Engine/Plugins/Runtime/Metasound/Source/MetasoundFrontend/Private/MetasoundAssetBase.cpp b/Engine/Plugins/Runtime/Metasound/Source/MetasoundFrontend/Private/MetasoundAssetBase.cpp
--- a/Engine/Plugins/Runtime/Metasound/Source/MetasoundFrontend/Private/MetasoundAssetBase.cpp
+++ b/Engine/Plugins/Runtime/Metasound/Source/MetasoundFrontend/Private/MetasoundAssetBase.cpp
@@ -384,11 +384,16 @@ void FMetasoundAssetBase::UpdateAndRegisterForSerialization(FName CookPlatformNa
 		return;
 	}
 
-	UpdateAndRegisterReferencesForSerialization(CookPlatformName);
-
 	UObject* Owner = GetOwningAsset();
 	check(Owner);
 
+	if (IMetaSoundAssetManager* AssetManager = IMetaSoundAssetManager::Get())
+	{
+		AssetManager->AddOrUpdateFromObject(*Owner);
+	}
+
+	UpdateAndRegisterReferencesForSerialization(CookPlatformName);
+
 	bool bDidUpdate = false;
 
 	FMetaSoundFrontendDocumentBuilder& DocBuilder = IDocumentBuilderRegistry::GetChecked().FindOrBeginBuilding(Owner);
