Undo MS execution registration path change to mirror serialization change that inadvertently failed to add asset entry to manager
#rnx

diff --git a/Engine/Plugins/Runtime/Metasound/Source/MetasoundFrontend/Private/MetasoundAssetBase.cpp b/Engine/Plugins/Runtime/Metasound/Source/MetasoundFrontend/Private/MetasoundAssetBase.cpp
--- a/Engine/Plugins/Runtime/Metasound/Source/MetasoundFrontend/Private/MetasoundAssetBase.cpp
+++ b/Engine/Plugins/Runtime/Metasound/Source/MetasoundFrontend/Private/MetasoundAssetBase.cpp
@@ -303,7 +303,7 @@ void FMetasoundAssetBase::UpdateAndRegisterForExecution(Metasound::Frontend::FMe
 	{
 		RebuildReferencedAssetClasses();
 	}
-#endif
+#endif // WITH_EDITOR
 
 	if (InRegistrationOptions.bRegisterDependencies)
 	{
@@ -313,6 +313,11 @@ void FMetasoundAssetBase::UpdateAndRegisterForExecution(Metasound::Frontend::FMe
 	UObject* Owner = GetOwningAsset();
 	check(Owner);
 
+	// This should not be necessary as it should be added on asset load,
+	// but currently registration is required to be called prior to adding
+	// an object-defined graph class to the registry so it was placed here.
+	IMetaSoundAssetManager::GetChecked().AddOrUpdateFromObject(*Owner);
+
 	// Auto update must be done after all referenced asset classes are registered
 	if (InRegistrationOptions.bAutoUpdate)
 	{
@@ -384,15 +389,11 @@ void FMetasoundAssetBase::UpdateAndRegisterForSerialization(FName CookPlatformNa
 		return;
 	}
 
+	UpdateAndRegisterReferencesForSerialization(CookPlatformName);
+
 	UObject* Owner = GetOwningAsset();
 	check(Owner);
-
-	if (IMetaSoundAssetManager* AssetManager = IMetaSoundAssetManager::Get())
-	{
-		AssetManager->AddOrUpdateFromObject(*Owner);
-	}
-
-	UpdateAndRegisterReferencesForSerialization(CookPlatformName);
+	IMetaSoundAssetManager::GetChecked().AddOrUpdateFromObject(*Owner);
 
 	bool bDidUpdate = false;
 
@@ -911,7 +912,6 @@ void FMetasoundAssetBase::RegisterAssetDependencies(const Metasound::Frontend::F
 {
 	using namespace Metasound::Frontend;
 
-	IMetaSoundAssetManager& AssetManager = IMetaSoundAssetManager::GetChecked();
 	TArray<FMetasoundAssetBase*> References = GetReferencedAssets();
 	for (FMetasoundAssetBase* Reference : References)
 	{
