Enable MetaSound 5.6 node deprecation warning macro

diff --git a/Engine/Plugins/Runtime/Metasound/Source/MetasoundFrontend/MetasoundFrontend.Build.cs b/Engine/Plugins/Runtime/Metasound/Source/MetasoundFrontend/MetasoundFrontend.Build.cs
--- a/Engine/Plugins/Runtime/Metasound/Source/MetasoundFrontend/MetasoundFrontend.Build.cs
+++ b/Engine/Plugins/Runtime/Metasound/Source/MetasoundFrontend/MetasoundFrontend.Build.cs
@@ -27,6 +27,7 @@ namespace UnrealBuildTool.Rules
 			);
 
 			PublicDefinitions.Add("WITH_METASOUND_FRONTEND=1");
+			PublicDefinitions.Add("UE_METASOUND_ENABLE_5_6_NODE_REGISTRATION_DEPRECATION_WARNINGS=1");
 		}
 	}
 }
diff --git a/Engine/Plugins/Runtime/Metasound/Source/MetasoundFrontend/Private/MetasoundFrontendNodeClassRegistry.cpp b/Engine/Plugins/Runtime/Metasound/Source/MetasoundFrontend/Private/MetasoundFrontendNodeClassRegistry.cpp
--- a/Engine/Plugins/Runtime/Metasound/Source/MetasoundFrontend/Private/MetasoundFrontendNodeClassRegistry.cpp
+++ b/Engine/Plugins/Runtime/Metasound/Source/MetasoundFrontend/Private/MetasoundFrontendNodeClassRegistry.cpp
@@ -308,7 +308,7 @@ namespace Metasound::Frontend
 			bDidLogError = true;
 
 			FNodeClassInfo ClassInfo = GetClassInfo();
-			UE_LOG(LogMetaSound, Warning, TEXT("Use of deprecated code path for node registration. First occurrence on node %s (asset path: %s). Please implement INodeClassRegistryEntry::CreateNode(FNodeData) const"), *ClassInfo.ClassName.ToString(), *ClassInfo.AssetPath.ToString());
+			UE_LOG(LogMetaSound, Warning, TEXT("Use of deprecated code path for node registration. First occurrence on node %s. Please implement INodeClassRegistryEntry::CreateNode(FNodeData) const"), *ClassInfo.ClassName.ToString());
 		}
 #endif
 PRAGMA_DISABLE_DEPRECATION_WARNINGS
