Minor MetaSound preset widget BP fixes
- remove unused return value for preset widget interface function
- get referenced preset assets recursively instead of only directly

diff --git a/Engine/Plugins/Runtime/Metasound/Source/MetasoundEditor/Private/MetasoundEditor.cpp b/Engine/Plugins/Runtime/Metasound/Source/MetasoundEditor/Private/MetasoundEditor.cpp
--- a/Engine/Plugins/Runtime/Metasound/Source/MetasoundEditor/Private/MetasoundEditor.cpp
+++ b/Engine/Plugins/Runtime/Metasound/Source/MetasoundEditor/Private/MetasoundEditor.cpp
@@ -1768,6 +1768,11 @@ namespace Metasound
 		{
 			using namespace Metasound::Frontend;
 
+			if (!Builder.IsValid() || !Builder->IsPreset())
+			{
+				return nullptr;
+			}
+
 			UWorld* World = GEditor ? GEditor->GetEditorWorldContext().World() : nullptr;
 			UAudioWidgetSubsystem* WidgetSubsystem = GEngine ? GEngine->GetEngineSubsystem<UAudioWidgetSubsystem>() : nullptr;
 			if (!World || !WidgetSubsystem)
@@ -1780,11 +1785,22 @@ namespace Metasound
 			const FMetasoundFrontendGraphClass& ThisGraphClass = DocBuilder.GetConstDocumentChecked().RootGraph;
 			const FNodeRegistryKey ThisRegistryKey(ThisGraphClass);
 
-			// Referenced asset 
-			FMetasoundAssetBase* ReferencedAsset = DocBuilder.GetReferencedPresetAsset();
-			const FMetasoundFrontendDocument& ReferencedDocument = ReferencedAsset->GetConstDocumentChecked();
-			const FMetasoundFrontendGraphClass& ReferencedGraphClass = ReferencedDocument.RootGraph;
-			const FNodeRegistryKey ReferencedRegistryKey(ReferencedGraphClass);
+			// Referenced assets
+			FMetasoundAssetBase& ThisAsset = DocBuilder.GetMetasoundAsset();
+			TArray<FMetasoundAssetBase*> ReferencedAssets;
+			IMetaSoundAssetManager::GetChecked().GetReferencedPresetHierarchy(ThisAsset, ReferencedAssets);
+
+			TArray<FNodeRegistryKey> ReferencedRegistryKeys;
+			Algo::Transform(ReferencedAssets, ReferencedRegistryKeys, [](const FMetasoundAssetBase* ReferencedAsset)
+			{
+				if (ensure(ReferencedAsset))
+				{
+					const FMetasoundFrontendDocument& ReferencedDocument = ReferencedAsset->GetConstDocumentChecked();
+					const FMetasoundFrontendGraphClass& ReferencedGraphClass = ReferencedDocument.RootGraph;
+					return FNodeRegistryKey(ReferencedGraphClass);
+				}
+				return FNodeRegistryKey();
+			});
 
 			// Filter whether the widget supports this MetaSound
 			auto FilterFunction = [&](UUserWidget* UserWidget)
@@ -1805,7 +1821,7 @@ namespace Metasound
 							const FMetasoundFrontendGraphClass& ExcludedGraphClass = ExcludedMetaSound->GetConstDocument().RootGraph;
 							const FNodeRegistryKey ExcludedRegistryKey(ExcludedGraphClass);
 
-							if (ThisRegistryKey == ExcludedRegistryKey || ReferencedRegistryKey == ExcludedRegistryKey)
+							if (ThisRegistryKey == ExcludedRegistryKey || ReferencedRegistryKeys.Contains(ExcludedRegistryKey))
 							{
 								return false;
 							}
@@ -1821,7 +1837,7 @@ namespace Metasound
 					const FMetasoundFrontendGraphClass& IncludedGraphClass = IncludedMetaSound->GetConstDocument().RootGraph;
 					const FNodeRegistryKey IncludedRegistryKey(IncludedGraphClass);
 
-					if (ThisRegistryKey == IncludedRegistryKey || ReferencedRegistryKey == IncludedRegistryKey)
+					if (ThisRegistryKey == IncludedRegistryKey || ReferencedRegistryKeys.Contains(IncludedRegistryKey))
 					{
 						return true;
 					}
diff --git a/Engine/Plugins/Runtime/Metasound/Source/MetasoundEditor/Public/MetasoundPresetWidgetInterface.h b/Engine/Plugins/Runtime/Metasound/Source/MetasoundEditor/Public/MetasoundPresetWidgetInterface.h
--- a/Engine/Plugins/Runtime/Metasound/Source/MetasoundEditor/Public/MetasoundPresetWidgetInterface.h
+++ b/Engine/Plugins/Runtime/Metasound/Source/MetasoundEditor/Public/MetasoundPresetWidgetInterface.h
@@ -26,7 +26,7 @@ public:
 	// MetaSounds in the include/exclude arrays can be MetaSound presets or non presets. 
 	// If a MetaSound is not a preset, then presets of that MetaSound will be supported/excluded by this widget.
 	UFUNCTION(BlueprintImplementableEvent)
-	TScriptInterface<IMetaSoundDocumentInterface> GetSupportedMetaSounds(bool& bSupportAllPresets, TArray<TScriptInterface<IMetaSoundDocumentInterface>>& ExcludedMetaSounds, TArray<TScriptInterface<IMetaSoundDocumentInterface>>& IncludedMetaSounds) const;
+	void GetSupportedMetaSounds(bool& bSupportAllPresets, TArray<TScriptInterface<IMetaSoundDocumentInterface>>& ExcludedMetaSounds, TArray<TScriptInterface<IMetaSoundDocumentInterface>>& IncludedMetaSounds) const;
 
 	// Called when the preset widget is constructed, giving the builder of the associated MetaSound preset
 	UFUNCTION(BlueprintImplementableEvent, meta = (DisplayName = "On MetaSoundPreset Widget Constructed"))
diff --git a/Engine/Plugins/Runtime/Metasound/Source/MetasoundEngine/Private/MetasoundAssetSubsystem.cpp b/Engine/Plugins/Runtime/Metasound/Source/MetasoundEngine/Private/MetasoundAssetSubsystem.cpp
--- a/Engine/Plugins/Runtime/Metasound/Source/MetasoundEngine/Private/MetasoundAssetSubsystem.cpp
+++ b/Engine/Plugins/Runtime/Metasound/Source/MetasoundEngine/Private/MetasoundAssetSubsystem.cpp
@@ -659,6 +659,46 @@ namespace Metasound::Engine
 		}
 		return MoveTemp(OutAssetInfos);
 	}
+	
+	bool FMetaSoundAssetManager::GetReferencedPresetHierarchy(FMetasoundAssetBase& InAsset, TArray<FMetasoundAssetBase*>& OutReferencedAssets) const
+	{
+		OutReferencedAssets.Reset();
+		if (!InAsset.GetConstDocumentChecked().RootGraph.PresetOptions.bIsPreset)
+		{
+			return false;
+		}
+
+		const bool bSuccess = GetReferencedPresetHierarchyInternal(InAsset, OutReferencedAssets);
+		// Remove first element (this)
+		if (OutReferencedAssets.Num() > 0)
+		{
+			OutReferencedAssets.RemoveAtSwap(0);
+		}
+		return bSuccess;
+	}
+
+	bool FMetaSoundAssetManager::GetReferencedPresetHierarchyInternal(FMetasoundAssetBase& InAsset, TArray<FMetasoundAssetBase*>& OutReferencedAssets) const
+	{
+		OutReferencedAssets.Add(&InAsset);
+
+		const bool bIsPreset = InAsset.GetConstDocumentChecked().RootGraph.PresetOptions.bIsPreset;
+
+		// Base case (found the first non preset referenced asset)
+		if (!bIsPreset)
+		{
+			return true;
+		}
+
+		// Presets must have a single asset reference 
+		FMetasoundAssetBase* ReferencedAsset = InAsset.GetReferencedAssets().Last();
+		if (InAsset.GetReferencedAssets().Num() != 1 || !ReferencedAsset)
+		{
+			UE_LOG(LogMetaSound, Warning, TEXT("'%s' is not a valid MetaSound preset with a single referenced asset. Ending referenced preset iteration."), *InAsset.GetOwningAssetName());
+			return false;
+		}
+
+		return GetReferencedPresetHierarchyInternal(*ReferencedAsset, OutReferencedAssets);
+	}
 
 	bool FMetaSoundAssetManager::ReassignClassName(TScriptInterface<IMetaSoundDocumentInterface> DocInterface)
 	{
diff --git a/Engine/Plugins/Runtime/Metasound/Source/MetasoundEngine/Public/MetasoundAssetSubsystem.h b/Engine/Plugins/Runtime/Metasound/Source/MetasoundEngine/Public/MetasoundAssetSubsystem.h
--- a/Engine/Plugins/Runtime/Metasound/Source/MetasoundEngine/Public/MetasoundAssetSubsystem.h
+++ b/Engine/Plugins/Runtime/Metasound/Source/MetasoundEngine/Public/MetasoundAssetSubsystem.h
@@ -107,6 +107,7 @@ namespace Metasound::Engine
 		virtual const FMetasoundAssetBase* GetAsAsset(const UObject& InObject) const override;
 #if WITH_EDITOR
 		virtual TSet<FAssetInfo> GetReferencedAssetClasses(const FMetasoundAssetBase& InAssetBase) const override;
+		virtual bool GetReferencedPresetHierarchy(FMetasoundAssetBase& InAsset, TArray<FMetasoundAssetBase*>& OutReferencedAssets) const override;
 		virtual bool ReassignClassName(TScriptInterface<IMetaSoundDocumentInterface> DocInterface) override;
 #endif // WITH_EDITOR
 
@@ -148,6 +149,8 @@ namespace Metasound::Engine
 		void RemoveLoadingDependencies(int32 InLoadID);
 		void OnReferencedAssetsLoaded(int32 InLoadID);
 
+		bool GetReferencedPresetHierarchyInternal(FMetasoundAssetBase& InAsset, TArray<FMetasoundAssetBase*>& OutReferencedAssets) const;
+
 		struct FPackageLoadedArgs
 		{
 			FName PackageName;
diff --git a/Engine/Plugins/Runtime/Metasound/Source/MetasoundFrontend/Public/MetasoundAssetManager.h b/Engine/Plugins/Runtime/Metasound/Source/MetasoundFrontend/Public/MetasoundAssetManager.h
--- a/Engine/Plugins/Runtime/Metasound/Source/MetasoundFrontend/Public/MetasoundAssetManager.h
+++ b/Engine/Plugins/Runtime/Metasound/Source/MetasoundFrontend/Public/MetasoundAssetManager.h
@@ -235,6 +235,12 @@ namespace Metasound::Frontend
 #if WITH_EDITOR
 		// Generates all asset info associated with registered assets that are referenced by the provided asset's graph.
 		virtual TSet<FAssetInfo> GetReferencedAssetClasses(const FMetasoundAssetBase& InAssetBase) const = 0;
+		
+		// Get assets this asset is a preset of, recursively for presets of presets.
+		// This means finding all references including the last non preset, but not graphs used directly in non-preset composition. 
+		// Ex. MetaSound A is a preset of MetaSound B, which is a preset of MetaSound C, which references D by composition -> When called on A, will return B and C.
+		// If not a preset, returns false.
+		virtual bool GetReferencedPresetHierarchy(FMetasoundAssetBase& InAsset, TArray<FMetasoundAssetBase*>& OutReferencedAssets) const = 0;
 #endif // WITH_EDITOR
 
 		// returns whether or not the given class is defined as a registered asset
