Fix where access type is determined on graph input pins.

#rnx

diff --git a/Engine/Plugins/Runtime/Metasound/Source/MetasoundFrontend/Private/MetasoundFrontendGraphBuilder.cpp b/Engine/Plugins/Runtime/Metasound/Source/MetasoundFrontend/Private/MetasoundFrontendGraphBuilder.cpp
--- a/Engine/Plugins/Runtime/Metasound/Source/MetasoundFrontend/Private/MetasoundFrontendGraphBuilder.cpp
+++ b/Engine/Plugins/Runtime/Metasound/Source/MetasoundFrontend/Private/MetasoundFrontendGraphBuilder.cpp
@@ -185,7 +185,7 @@ namespace Metasound::Frontend
 			}
 		}
 
-		FOutputVertexInterface CreateOutputVertexInterface(const FBuildContext& InContext, const FMetasoundFrontendNode& InNode, const FMetasoundFrontendClass& InClass)
+		FOutputVertexInterface CreateOutputVertexInterface(const FBuildContext& InContext, const FMetasoundFrontendNode& InNode, const FMetasoundFrontendClass& InClass, const FMetasoundFrontendClassInput* InOwningGraphClassInput, const FMetasoundFrontendClassOutput* InOwningGraphClassOutput)
 		{
 			// This exists here until we bring sub interface functionality to the frontend. 
 			// This check will no longer be valid at that point. We will have to implement
@@ -209,8 +209,41 @@ namespace Metasound::Frontend
 					// we have to grab data from both the class and node interfaces
 					case EMetasoundFrontendClassType::Input:
 						{
-							if (InNode.Interface.Outputs.Num() == 1)
+							if (ensureMsgf(InOwningGraphClassInput, TEXT("Input nodes require the owning graph class input in order to retrieve the access type")))
 							{
+								if (InNode.Interface.Outputs.Num() == 1)
+								{
+									const FMetasoundFrontendVertex& NodeOutputVertex = InNode.Interface.Outputs[0];
+									const FMetasoundFrontendClassOutput& ClassOutputVertex = InClass.Interface.Outputs[0];
+									// Note: stripping metadata here. If we make this function public, we may want to keep metadata if WITH_EDITORONLY_DATA
+									Vertices.Add(FOutputDataVertex{
+												NodeOutputVertex.Name, 
+												NodeOutputVertex.TypeName, 
+												FDataVertexMetadata{}, 
+												// NOTE: The access type reported on the ClassInputVertex may be incorrect. When the access type is changed in the editor,
+												// the access type of the input node's FMetasoundFrontendClass is not updated. 
+												FrontendVertexAccessTypeToCoreVertexAccessType(InOwningGraphClassInput->AccessType)
+											});
+								}
+								else
+								{
+									// This is an edge case in the off scenario that someone
+									// is trying to build a graph that used the `TInputNode<FStereoAudio>` 
+									// specialization.  Auto-update should have removed these
+									// long ago, but this log will help identify any issues.
+									// This log is added in UE 5.6 and if no one is hitting
+									// this issue, can be removed in UE 5.7.
+									UE_LOG(LogMetaSound, Error, TEXT("The MetaSound %s has an malformed input node %s. Please replace the input node and resave"), *InContext.DebugAssetName, *InNode.Name.ToString());
+								}
+							}
+						}
+						break;
+
+					case EMetasoundFrontendClassType::Output:
+						{
+							if (ensureMsgf(InOwningGraphClassOutput, TEXT("Output nodes require the owning graph class output in order to retrieve the access type")))
+							{
+								check(InNode.Interface.Outputs.Num() == 1);
 								const FMetasoundFrontendVertex& NodeOutputVertex = InNode.Interface.Outputs[0];
 								const FMetasoundFrontendClassOutput& ClassOutputVertex = InClass.Interface.Outputs[0];
 								// Note: stripping metadata here. If we make this function public, we may want to keep metadata if WITH_EDITORONLY_DATA
@@ -218,35 +251,11 @@ namespace Metasound::Frontend
 											NodeOutputVertex.Name, 
 											NodeOutputVertex.TypeName, 
 											FDataVertexMetadata{}, 
-											FrontendVertexAccessTypeToCoreVertexAccessType(ClassOutputVertex.AccessType)
+											// NOTE: The access type reported on the ClassInputVertex may be incorrect. When the access type is changed in the editor,
+											// the access type of the input node's FMetasoundFrontendClass is not updated. 
+											FrontendVertexAccessTypeToCoreVertexAccessType(InOwningGraphClassOutput->AccessType)
 										});
 							}
-							else
-							{
-								// This is an edge case in the off scenario that someone
-								// is trying to build a graph that used the `TInputNode<FStereoAudio>` 
-								// specialization.  Auto-update should have removed these
-								// long ago, but this log will help identify any issues.
-								// This log is added in UE 5.6 and if no one is hitting
-								// this issue, can be removed in UE 5.7.
-								UE_LOG(LogMetaSound, Error, TEXT("The MetaSound %s has an malformed input node %s. Please replace the input node and resave"), *InContext.DebugAssetName, *InNode.Name.ToString());
-							}
-						}
-						break;
-
-					case EMetasoundFrontendClassType::Output:
-						{
-							check(InNode.Interface.Outputs.Num() == 1);
-							const FMetasoundFrontendVertex& NodeOutputVertex = InNode.Interface.Outputs[0];
-							const FMetasoundFrontendClassOutput& ClassOutputVertex = InClass.Interface.Outputs[0];
-							// Note: stripping metadata here. If we make this function public, we may want to keep metadata if WITH_EDITORONLY_DATA
-							Vertices.Add(FOutputDataVertex{
-										NodeOutputVertex.Name, 
-										NodeOutputVertex.TypeName, 
-										FDataVertexMetadata{}, 
-										FrontendVertexAccessTypeToCoreVertexAccessType(ClassOutputVertex.AccessType)
-									});
-
 						}
 						break;
 
@@ -272,7 +281,7 @@ namespace Metasound::Frontend
 			}
 		}
 
-		FInputVertexInterface CreateInputVertexInterface(const FBuildContext& InContext, const FMetasoundFrontendNode& InNode, const FMetasoundFrontendClass& InClass, const FMetasoundFrontendClassInput* InOwningGraphClassInput=nullptr)
+		FInputVertexInterface CreateInputVertexInterface(const FBuildContext& InContext, const FMetasoundFrontendNode& InNode, const FMetasoundFrontendClass& InClass, const FMetasoundFrontendClassInput* InOwningGraphClassInput, const FMetasoundFrontendClassOutput* InOwningGraphClassOutput)
 		{
 			// This exists here until we bring sub interface functionality to the frontend. 
 			// This check will no longer be valid at that point. We will have to implement
@@ -311,35 +320,42 @@ namespace Metasound::Frontend
 					// we have to grab data from both the class and node interfaces
 					case EMetasoundFrontendClassType::Input:
 						{
-							if (ensureMsgf(InOwningGraphClassInput, TEXT("Input nodes require the owning graph class input in order to retrieve the default literal")))
+							if (ensureMsgf(InOwningGraphClassInput, TEXT("Input nodes require the owning graph class input in order to retrieve the default literal and access type")))
 							{
 								check(InNode.Interface.Inputs.Num() == 1);
 								const FMetasoundFrontendVertex& NodeInputVertex = InNode.Interface.Inputs[0];
 								const FMetasoundFrontendClassInput& ClassInputVertex = InClass.Interface.Inputs[0];
 								const FMetasoundFrontendLiteral* FrontendLiteral = FindInputLiteralForInputNode(InNode, InClass, *InOwningGraphClassInput);
 								
-								Vertices.Add(MakeVertex(NodeInputVertex.Name, NodeInputVertex.TypeName, ClassInputVertex.AccessType, FrontendLiteral));
+								// NOTE: The access type reported on the ClassInputVertex may be incorrect. When the access type is changed in the editor,
+								// the access type of the input node's FMetasoundFrontendClass is not updated. 
+								Vertices.Add(MakeVertex(NodeInputVertex.Name, NodeInputVertex.TypeName, InOwningGraphClassInput->AccessType, FrontendLiteral));
 							}
 						}
 						break;
 					case EMetasoundFrontendClassType::Output:
 						{
-							if (InNode.Interface.Inputs.Num() == 1)
+							if (ensureMsgf(InOwningGraphClassOutput, TEXT("Output nodes require the owning graph class output in order to retrieve the access type")))
 							{
-								const FMetasoundFrontendVertex& NodeInputVertex = InNode.Interface.Inputs[0];
-								const FMetasoundFrontendClassInput& ClassInputVertex = InClass.Interface.Inputs[0];
-								const FMetasoundFrontendLiteral* FrontendLiteral = FindLiteralForInputVertex(NodeInputVertex.Name, InNode, ClassInputVertex);
-								Vertices.Add(MakeVertex(NodeInputVertex.Name, NodeInputVertex.TypeName, ClassInputVertex.AccessType, FrontendLiteral));
-							}
-							else
-							{
-								// This is an edge case in the off scenario that someone
-								// is trying to build a graph that used the `TOutputNode<FStereoAudio>` 
-								// specialization.  Auto-update should have removed these
-								// long ago, but this log will help identify any issues.
-								// This log is added in UE 5.6 and if no one is hitting
-								// this issue, can be removed in UE 5.7.
-								UE_LOG(LogMetaSound, Error, TEXT("The MetaSound %s has an malformed output node %s. Please replace the output node and resave"), *InContext.DebugAssetName, *InNode.Name.ToString());
+								if (InNode.Interface.Inputs.Num() == 1)
+								{
+									const FMetasoundFrontendVertex& NodeInputVertex = InNode.Interface.Inputs[0];
+									const FMetasoundFrontendClassInput& ClassInputVertex = InClass.Interface.Inputs[0];
+									const FMetasoundFrontendLiteral* FrontendLiteral = FindLiteralForInputVertex(NodeInputVertex.Name, InNode, ClassInputVertex);
+									// NOTE: The access type reported on the ClassInputVertex may be incorrect. When the access type is changed in the editor,
+									// the access type of the output node's FMetasoundFrontendClass is not updated. 
+									Vertices.Add(MakeVertex(NodeInputVertex.Name, NodeInputVertex.TypeName, InOwningGraphClassOutput->AccessType, FrontendLiteral));
+								}
+								else
+								{
+									// This is an edge case in the off scenario that someone
+									// is trying to build a graph that used the `TOutputNode<FStereoAudio>` 
+									// specialization.  Auto-update should have removed these
+									// long ago, but this log will help identify any issues.
+									// This log is added in UE 5.6 and if no one is hitting
+									// this issue, can be removed in UE 5.7.
+									UE_LOG(LogMetaSound, Error, TEXT("The MetaSound %s has an malformed output node %s. Please replace the output node and resave"), *InContext.DebugAssetName, *InNode.Name.ToString());
+								}
 							}
 						}
 						break;
@@ -366,17 +382,17 @@ namespace Metasound::Frontend
 			}
 		}
 
-		FVertexInterface CreateVertexInterface(const FBuildContext& InContext, const FMetasoundFrontendNode& InNode, const FMetasoundFrontendClass& InClass, const FMetasoundFrontendClassInput* InOwningGraphClassInput=nullptr)
+		FVertexInterface CreateVertexInterface(const FBuildContext& InContext, const FMetasoundFrontendNode& InNode, const FMetasoundFrontendClass& InClass, const FMetasoundFrontendClassInput* InOwningGraphClassInput, const FMetasoundFrontendClassOutput* InOwningGraphClassOutput)
 		{
 			return FVertexInterface(
-				CreateInputVertexInterface(InContext, InNode, InClass, InOwningGraphClassInput), 
-				CreateOutputVertexInterface(InContext, InNode, InClass), 
+				CreateInputVertexInterface(InContext, InNode, InClass, InOwningGraphClassInput, InOwningGraphClassOutput), 
+				CreateOutputVertexInterface(InContext, InNode, InClass, InOwningGraphClassInput, InOwningGraphClassOutput), 
 				CreateEnvironmentVertexInterface(InNode.Interface));
 		}
 
-		FNodeData CreateNodeData(const FBuildContext& InContext, const FMetasoundFrontendNode& InNode, const FMetasoundFrontendClass& InClass, const FMetasoundFrontendClassInput* InOwningGraphClassInput=nullptr)
+		FNodeData CreateNodeData(const FBuildContext& InContext, const FMetasoundFrontendNode& InNode, const FMetasoundFrontendClass& InClass, const FMetasoundFrontendClassInput* InOwningGraphClassInput=nullptr, const FMetasoundFrontendClassOutput* InOwningGraphClassOutput=nullptr)
 		{
-			return FNodeData(InNode.Name, InNode.GetID(), CreateVertexInterface(InContext, InNode, InClass, InOwningGraphClassInput));
+			return FNodeData(InNode.Name, InNode.GetID(), CreateVertexInterface(InContext, InNode, InClass, InOwningGraphClassInput, InOwningGraphClassOutput));
 		}
 
 		const FMetasoundFrontendVariable* FindVariableForVariableNode(const FMetasoundFrontendNode& InVariableNode, const FMetasoundFrontendGraph& InGraph)
@@ -470,7 +486,7 @@ namespace Metasound::Frontend
 			return TUniquePtr<INode>(nullptr);
 		}
 
-		TUniquePtr<INode> CreateOutputNode(const FBuildContext& InContext, const FMetasoundFrontendNode& InNode, const FMetasoundFrontendClass& InClass, FBuildGraphContext& InGraphContext)
+		TUniquePtr<INode> CreateOutputNode(const FBuildContext& InContext, const FMetasoundFrontendNode& InNode, const FMetasoundFrontendClass& InClass, const FMetasoundFrontendClassOutput& InOwningGraphClassOutput)
 		{
 			check(InClass.Metadata.GetType() == EMetasoundFrontendClassType::Output);
 			check(InNode.ClassID == InClass.ID);
@@ -495,11 +511,11 @@ namespace Metasound::Frontend
 
 			if (InClass.Interface.Outputs[0].AccessType == EMetasoundFrontendVertexAccessType::Reference)
 			{
-				return InContext.DataTypeRegistry.CreateOutputNode(OutputVertex.TypeName, CreateNodeData(InContext, InNode, InClass));
+				return InContext.DataTypeRegistry.CreateOutputNode(OutputVertex.TypeName, CreateNodeData(InContext, InNode, InClass, nullptr, &InOwningGraphClassOutput));
 			}
 			else if (InClass.Interface.Outputs[0].AccessType == EMetasoundFrontendVertexAccessType::Value)
 			{
-				return InContext.DataTypeRegistry.CreateConstructorOutputNode(OutputVertex.TypeName, CreateNodeData(InContext, InNode, InClass));
+				return InContext.DataTypeRegistry.CreateConstructorOutputNode(OutputVertex.TypeName, CreateNodeData(InContext, InNode, InClass, nullptr, &InOwningGraphClassOutput));
 			}
 			else
 			{
@@ -611,7 +627,7 @@ namespace Metasound::Frontend
 							const FMetasoundFrontendClassOutput* ClassOutput = FindClassOutputForOutputNode(InGraphContext.GraphClass, Node);
 							if (nullptr != ClassOutput)
 							{
-								TUniquePtr<INode> OutputNode = CreateOutputNode(InGraphContext.BuildContext, Node, *NodeClass, InGraphContext);
+								TUniquePtr<INode> OutputNode = CreateOutputNode(InGraphContext.BuildContext, Node, *NodeClass, *ClassOutput);
 								if (OutputNode)
 								{
 									InGraphContext.Graph->AddOutputNode(Node.GetID(), ClassOutput->Name, MoveTemp(OutputNode));
