Trivial fix for flipped filter functions causing schema actions to misrepresent what can and cannot be connected in MetaSound Editor from asset tag refactor
#rnx

diff --git a/Engine/Plugins/Runtime/Metasound/Source/MetasoundEditor/Private/MetasoundEditorGraphSchema.cpp b/Engine/Plugins/Runtime/Metasound/Source/MetasoundEditor/Private/MetasoundEditorGraphSchema.cpp
--- a/Engine/Plugins/Runtime/Metasound/Source/MetasoundEditor/Private/MetasoundEditorGraphSchema.cpp
+++ b/Engine/Plugins/Runtime/Metasound/Source/MetasoundEditor/Private/MetasoundEditorGraphSchema.cpp
@@ -2668,7 +2668,7 @@ void UMetasoundEditorGraphSchema::GetFunctionActions(FGraphActionMenuBuilder& Ac
 
 		if (AccessFilters.HasInputFilters())
 		{
-			if (!Result->CanConnectInputOfTypeAndAccess(AccessFilters.InputTypeName, AccessFilters.InputAccessType))
+			if (!Result->CanConnectOutputOfTypeAndAccess(AccessFilters.InputTypeName, AccessFilters.InputAccessType))
 			{
 				return;
 			}
@@ -2676,7 +2676,7 @@ void UMetasoundEditorGraphSchema::GetFunctionActions(FGraphActionMenuBuilder& Ac
 
 		if (AccessFilters.HasOutputFilters())
 		{
-			if (!Result->CanConnectOutputOfTypeAndAccess(AccessFilters.OutputTypeName, AccessFilters.OutputAccessType))
+			if (!Result->CanConnectInputOfTypeAndAccess(AccessFilters.OutputTypeName, AccessFilters.OutputAccessType))
 			{
 				return;
 			}
