- Clean up MetaSound tag info class hierarchy and remove "query" keyword as its confusing to BP
- Serialize interface info to asset tags
- Add interface metadata struct and serialize tags accordingly there into in preparation for asset-defined interface registration
- Don't serialize editor-only tag data if running cook commandlet
- Add tag data look-up function to editor BP utils
#rnx

diff --git a/Engine/Plugins/Runtime/Metasound/Source/MetasoundEditor/Private/MetasoundAssetDefinitions.cpp b/Engine/Plugins/Runtime/Metasound/Source/MetasoundEditor/Private/MetasoundAssetDefinitions.cpp
--- a/Engine/Plugins/Runtime/Metasound/Source/MetasoundEditor/Private/MetasoundAssetDefinitions.cpp
+++ b/Engine/Plugins/Runtime/Metasound/Source/MetasoundEditor/Private/MetasoundAssetDefinitions.cpp
@@ -48,8 +48,8 @@ namespace Metasound::Editor
 		{
 			using namespace Frontend;
 
-			const FAssetClassQueryInfo QueryInfo(InAssetData);
-			if (!QueryInfo.bIsValid)
+			const FMetaSoundAssetClassInfo ClassInfo(InAssetData);
+			if (!ClassInfo.bIsValid)
 			{
 				UE_LOG(LogMetaSound, VeryVerbose,
 					TEXT("ClassBrush for asset '%s' may return incorrect preset icon. Asset requires reserialization."),
@@ -57,7 +57,7 @@ namespace Metasound::Editor
 			}
 
 			FString BrushName = FString::Printf(TEXT("MetasoundEditor.%s"), *InClassName.ToString());
-			if (QueryInfo.bIsPreset)
+			if (ClassInfo.DocInfo.bIsPreset)
 			{
 				BrushName += TEXT(".Preset");
 			}
diff --git a/Engine/Plugins/Runtime/Metasound/Source/MetasoundEditor/Private/MetasoundDetailCustomization.cpp b/Engine/Plugins/Runtime/Metasound/Source/MetasoundEditor/Private/MetasoundDetailCustomization.cpp
--- a/Engine/Plugins/Runtime/Metasound/Source/MetasoundEditor/Private/MetasoundDetailCustomization.cpp
+++ b/Engine/Plugins/Runtime/Metasound/Source/MetasoundEditor/Private/MetasoundDetailCustomization.cpp
@@ -805,7 +805,7 @@ namespace Metasound::Editor
 					const FName InterfaceName { *NameToAdd.Get() };
 					if (ensure(ISearchEngine::Get().FindInterfaceWithHighestVersion(InterfaceName, InterfaceToAdd)))
 					{
-						const FScopedTransaction Transaction(FText::Format(LOCTEXT("AddInterfaceTransactionFormat", "Add MetaSound Interface '{0}'"), FText::FromString(InterfaceToAdd.Version.ToString())));
+						const FScopedTransaction Transaction(FText::Format(LOCTEXT("AddInterfaceTransactionFormat", "Add MetaSound Interface '{0}'"), FText::FromString(InterfaceToAdd.Metadata.Version.ToString())));
 						MetaSound.Modify();
 						FModifyInterfaceOptions Options({ }, { InterfaceToAdd });
 						Options.bSetDefaultNodeLocations = false; // Don't automatically add nodes to ed graph
@@ -854,8 +854,8 @@ namespace Metasound::Editor
 						return;
 					}
 
-					TArray<FMetasoundFrontendInterface> ImplementedInterfaces;
-					Algo::Transform(ImplementedInterfaceNames, ImplementedInterfaces, [](const FName& Name)
+					TArray<FMetasoundFrontendInterface> InheritedInterfaces;
+					Algo::Transform(ImplementedInterfaceNames, InheritedInterfaces, [](const FName& Name)
 					{
 						FMetasoundFrontendInterface Interface;
 						ISearchEngine::Get().FindInterfaceWithHighestVersion(Name, Interface);
@@ -866,7 +866,7 @@ namespace Metasound::Editor
 					{
 						const FScopedTransaction Transaction(LOCTEXT("RemoveAllInterfacesTransaction", "Remove All MetaSound Interfaces"));
 						MetaSound.Modify();
-						FModifyInterfaceOptions Options(ImplementedInterfaces, { });
+						FModifyInterfaceOptions Options(InheritedInterfaces, { });
 						Options.bSetDefaultNodeLocations = false; // Don't automatically add nodes to ed graph
 						DocBuilder.ModifyInterfaces(MoveTemp(Options));
 					}
@@ -914,7 +914,7 @@ namespace Metasound::Editor
 
 				UObject& MetaSound = DocBuilder.CastDocumentObjectChecked<UObject>();
 				{
-					const FScopedTransaction Transaction(FText::Format(LOCTEXT("RemoveInterfaceTransactionFormat", "Remove MetaSound Interface '{0}'"), FText::FromString(InterfaceEntry.Version.ToString())));
+					const FScopedTransaction Transaction(FText::Format(LOCTEXT("RemoveInterfaceTransactionFormat", "Remove MetaSound Interface '{0}'"), FText::FromString(InterfaceEntry.Metadata.Version.ToString())));
 					MetaSound.Modify();
 					FModifyInterfaceOptions Options({ InterfaceEntry }, {});
 					Options.bSetDefaultNodeLocations = false; // Don't automatically add nodes to ed graph
@@ -935,7 +935,7 @@ namespace Metasound::Editor
 				.AutoWidth()
 				[
 					SNew(STextBlock)
-					.Text(FText::FromName(InterfaceEntry.Version.Name))
+					.Text(FText::FromName(InterfaceEntry.Metadata.Version.Name))
 				]
 			+ SHorizontalBox::Slot()
 				.Padding(2.0f)
@@ -977,7 +977,8 @@ namespace Metasound::Editor
 				const FInterfaceRegistryKey Key = GetInterfaceRegistryKey(Version);
 				if (const IInterfaceRegistryEntry* Entry = IInterfaceRegistry::Get().FindInterfaceRegistryEntry(Key))
 				{
-					if (const FMetasoundFrontendInterfaceUClassOptions* Options = Entry->GetInterface().FindClassOptions(ClassName))
+					auto FindClassOptionsPredicate = [&ClassName](const FMetasoundFrontendInterfaceUClassOptions& Options) { return Options.ClassPath == ClassName; };
+					if (const FMetasoundFrontendInterfaceUClassOptions* Options = Entry->GetInterface().Metadata.UClassOptions.FindByPredicate(FindClassOptionsPredicate))
 					{
 						return Options->bIsModifiable;
 					}
@@ -989,17 +990,17 @@ namespace Metasound::Editor
 				return false;
 			};
 
-			const TSet<FMetasoundFrontendVersion>& ImplementedInterfaces = Builder->GetBuilder().GetConstDocumentChecked().Interfaces;
-			Algo::TransformIf(ImplementedInterfaces, ImplementedInterfaceNames, CanAddOrRemoveInterface, GetVersionName);
+			const TSet<FMetasoundFrontendVersion>& InheritedInterfaces = Builder->GetBuilder().GetConstDocumentChecked().Interfaces;
+			Algo::TransformIf(InheritedInterfaces, ImplementedInterfaceNames, CanAddOrRemoveInterface, GetVersionName);
 
 			TArray<FMetasoundFrontendInterface> Interfaces = Frontend::ISearchEngine::Get().FindAllInterfaces();
 			for (const FMetasoundFrontendInterface& Interface : Interfaces)
 			{
-				if (!ImplementedInterfaceNames.Contains(Interface.Version.Name))
+				if (!ImplementedInterfaceNames.Contains(Interface.Metadata.Version.Name))
 				{
-					if (CanAddOrRemoveInterface(Interface.Version))
+					if (CanAddOrRemoveInterface(Interface.Metadata.Version))
 					{
-						FString Name = Interface.Version.Name.ToString();
+						FString Name = Interface.Metadata.Version.Name.ToString();
 						AddableInterfaceNames.Add(MakeShared<FString>(MoveTemp(Name)));
 					}
 				}
diff --git a/Engine/Plugins/Runtime/Metasound/Source/MetasoundEditor/Private/MetasoundEditorGraph.cpp b/Engine/Plugins/Runtime/Metasound/Source/MetasoundEditor/Private/MetasoundEditorGraph.cpp
--- a/Engine/Plugins/Runtime/Metasound/Source/MetasoundEditor/Private/MetasoundEditorGraph.cpp
+++ b/Engine/Plugins/Runtime/Metasound/Source/MetasoundEditor/Private/MetasoundEditorGraph.cpp
@@ -851,7 +851,7 @@ bool UMetasoundEditorGraphInput::IsInterfaceMember(FMetasoundFrontendInterface*
 		// Is interface declared on this MetaSound 
 		const UObject& MetaSoundObject = GetOwningGraph()->GetMetasoundChecked();
 		const FMetasoundAssetBase* MetaSoundAsset = Metasound::IMetasoundUObjectRegistry::Get().GetObjectAsAssetBase(&MetaSoundObject);
-		if (MetaSoundAsset && MetaSoundAsset->IsInterfaceDeclared(Interface.Version))
+		if (MetaSoundAsset && MetaSoundAsset->IsInterfaceDeclared(Interface.Metadata.Version))
 		{
 			// Check if Input is a member of the found interface
 			if (const FMetasoundFrontendNode* InputNode = GetFrontendNode())
@@ -1228,7 +1228,7 @@ bool UMetasoundEditorGraphOutput::IsInterfaceMember(FMetasoundFrontendInterface*
 		// Is interface declared on this MetaSound 
 		const UObject& MetaSoundObject = GetOwningGraph()->GetMetasoundChecked();
 		const FMetasoundAssetBase* MetaSoundAsset = Metasound::IMetasoundUObjectRegistry::Get().GetObjectAsAssetBase(&MetaSoundObject);
-		if (MetaSoundAsset && MetaSoundAsset->IsInterfaceDeclared(Interface.Version))
+		if (MetaSoundAsset && MetaSoundAsset->IsInterfaceDeclared(Interface.Metadata.Version))
 		{
 			// Check if Output is a member of the found interface
 			if (const FMetasoundFrontendNode* OutputNode = GetFrontendNode())
diff --git a/Engine/Plugins/Runtime/Metasound/Source/MetasoundEditor/Private/MetasoundEditorGraphSchema.cpp b/Engine/Plugins/Runtime/Metasound/Source/MetasoundEditor/Private/MetasoundEditorGraphSchema.cpp
--- a/Engine/Plugins/Runtime/Metasound/Source/MetasoundEditor/Private/MetasoundEditorGraphSchema.cpp
+++ b/Engine/Plugins/Runtime/Metasound/Source/MetasoundEditor/Private/MetasoundEditorGraphSchema.cpp
@@ -413,10 +413,10 @@ PRAGMA_ENABLE_DEPRECATION_WARNINGS
 		class FAssetSchemaQueryResult : public ISchemaQueryResult
 		{
 		public:
-			FAssetSchemaQueryResult(Frontend::FAssetClassQueryInfo InTagData)
-				: QueryInfo(MoveTemp(InTagData))
+			FAssetSchemaQueryResult(Frontend::FMetaSoundAssetClassInfo InTagData)
+				: ClassInfo(MoveTemp(InTagData))
 			{
-				bIsLoaded = FSoftObjectPath(QueryInfo.AssetPath).ResolveObject() != nullptr;
+				bIsLoaded = FSoftObjectPath(ClassInfo.AssetPath).ResolveObject() != nullptr;
 			}
 
 			virtual ~FAssetSchemaQueryResult() = default;
@@ -425,7 +425,7 @@ PRAGMA_ENABLE_DEPRECATION_WARNINGS
 			{
 				using namespace Metasound;
 
-				if (UObject* MetaSound = FSoftObjectPath(QueryInfo.AssetPath).TryLoad())
+				if (UObject* MetaSound = FSoftObjectPath(ClassInfo.AssetPath).TryLoad())
 				{
 					if (const FMetasoundAssetBase* MetaSoundAsset = IMetasoundUObjectRegistry::Get().GetObjectAsAssetBase(MetaSound))
 					{
@@ -451,7 +451,7 @@ PRAGMA_ENABLE_DEPRECATION_WARNINGS
 			{
 				using namespace Frontend;
 
-				const bool bCanConnect = Algo::AnyOf(QueryInfo.Outputs.ClassVertexInfo, [&InputTypeName, &InputAccessType](const FClassQueryVertexInfo& Output)
+				const bool bCanConnect = Algo::AnyOf(ClassInfo.InterfaceInfo.Outputs, [&InputTypeName, &InputAccessType](const FMetaSoundClassVertexInfo& Output)
 				{
 					return InputTypeName == Output.TypeName && FMetasoundFrontendClassVertex::CanConnectVertexAccessTypes(Output.AccessType, InputAccessType);
 				});
@@ -461,7 +461,7 @@ PRAGMA_ENABLE_DEPRECATION_WARNINGS
 			virtual bool CanConnectOutputOfTypeAndAccess(FName OutputTypeName, EMetasoundFrontendVertexAccessType OutputAccessType) const override
 			{
 				using namespace Engine;
-				const bool bCanConnect = Algo::AnyOf(QueryInfo.Inputs.ClassVertexInfo, [&OutputTypeName, &OutputAccessType](const FClassQueryVertexInfo& Input)
+				const bool bCanConnect = Algo::AnyOf(ClassInfo.InterfaceInfo.Inputs, [&OutputTypeName, &OutputAccessType](const FMetaSoundClassVertexInfo& Input)
 				{
 					return OutputTypeName == Input.TypeName && FMetasoundFrontendClassVertex::CanConnectVertexAccessTypes(OutputAccessType, Input.AccessType);
 				});
@@ -470,7 +470,7 @@ PRAGMA_ENABLE_DEPRECATION_WARNINGS
 
 			virtual Frontend::FAssetKey GetAssetKey() const override
 			{
-				return Frontend::FAssetKey(QueryInfo.ClassName, QueryInfo.Version);
+				return Frontend::FAssetKey(ClassInfo.ClassName, ClassInfo.Version);
 			}
 
 			virtual const TArray<FText>& GetCategoryHierarchy() const override
@@ -483,14 +483,14 @@ PRAGMA_ENABLE_DEPRECATION_WARNINGS
 					}
 				}
 
-				return QueryInfo.DisplayInfo.Hierarchy;
+				return ClassInfo.InterfaceInfo.SearchInfo.Hierarchy;
 			}
 
 			virtual FText GetDisplayName() const override
 			{
 				using namespace Frontend;
 
-				FText DisplayName = FText::FromName(QueryInfo.AssetPath.GetAssetName());
+				FText DisplayName = FText::FromName(ClassInfo.AssetPath.GetAssetName());
 
 				if (bIsLoaded)
 				{
@@ -503,9 +503,9 @@ PRAGMA_ENABLE_DEPRECATION_WARNINGS
 						}
 					}
 				}
-				else if (!QueryInfo.DisplayInfo.ClassDisplayName.IsEmptyOrWhitespace())
+				else if (!ClassInfo.InterfaceInfo.SearchInfo.ClassDisplayName.IsEmptyOrWhitespace())
 				{
-					DisplayName = QueryInfo.DisplayInfo.ClassDisplayName;
+					DisplayName = ClassInfo.InterfaceInfo.SearchInfo.ClassDisplayName;
 				}
 
 				if (ShowUnloadedAssetInBrowserCvar)
@@ -526,7 +526,7 @@ PRAGMA_ENABLE_DEPRECATION_WARNINGS
 					}
 				}
 
-				return QueryInfo.DisplayInfo.Keywords;
+				return ClassInfo.InterfaceInfo.SearchInfo.Keywords;
 			}
 
 			virtual EMetasoundFrontendClassType GetRegistryClassType() const override
@@ -550,7 +550,7 @@ PRAGMA_ENABLE_DEPRECATION_WARNINGS
 					}
 				}
 
-				return QueryInfo.DisplayInfo.ClassDescription;
+				return ClassInfo.InterfaceInfo.SearchInfo.ClassDescription;
 			}
 
 			virtual bool IsNative() const override
@@ -559,7 +559,7 @@ PRAGMA_ENABLE_DEPRECATION_WARNINGS
 			}
 
 		private:
-			Frontend::FAssetClassQueryInfo QueryInfo;
+			Frontend::FMetaSoundAssetClassInfo ClassInfo;
 			mutable bool bIsLoaded = false;
 		};
 
@@ -2883,15 +2883,15 @@ void UMetasoundEditorGraphSchema::QueryNodeClasses(TFunctionRef<void(TUniquePtr<
 		// Append all unloaded asset classes
 		if (bScanAssetTags)
 		{
-			AssetManager.IterateAssetTagData([&FoundKeys, &Results](FAssetClassQueryInfo QueryInfo)
+			AssetManager.IterateAssetTagData([&FoundKeys, &Results](FMetaSoundAssetClassInfo ClassInfo)
 			{
 				using namespace Metasound::Frontend;
 
-				const FAssetKey AssetKey(QueryInfo.ClassName, QueryInfo.Version);
+				const FAssetKey AssetKey(ClassInfo.ClassName, ClassInfo.Version);
 				const FNodeRegistryKey RegKey(AssetKey);
 				if (!FoundKeys.Contains(RegKey))
 				{
-					Results.Add(MakeUnique<SchemaPrivate::FAssetSchemaQueryResult>(MoveTemp(QueryInfo)));
+					Results.Add(MakeUnique<SchemaPrivate::FAssetSchemaQueryResult>(MoveTemp(ClassInfo)));
 				}
 			});
 		}
diff --git a/Engine/Plugins/Runtime/Metasound/Source/MetasoundEngine/Private/Interfaces/MetasoundDeprecatedInterfaces.cpp b/Engine/Plugins/Runtime/Metasound/Source/MetasoundEngine/Private/Interfaces/MetasoundDeprecatedInterfaces.cpp
--- a/Engine/Plugins/Runtime/Metasound/Source/MetasoundEngine/Private/Interfaces/MetasoundDeprecatedInterfaces.cpp
+++ b/Engine/Plugins/Runtime/Metasound/Source/MetasoundEngine/Private/Interfaces/MetasoundDeprecatedInterfaces.cpp
@@ -85,8 +85,8 @@ namespace Metasound::Engine
 		FMetasoundFrontendInterface GetInterface()
 		{
 			FMetasoundFrontendInterface Interface;
-			Interface.Version = GetVersion();
-			Interface.UClassOptions = DeprecatedInterfacesPrivate::GetDeprecatedClassOptions();
+			Interface.Metadata.Version = GetVersion();
+			Interface.Metadata.UClassOptions = DeprecatedInterfacesPrivate::GetDeprecatedClassOptions();
 			return Interface;
 		}
 	}
@@ -108,8 +108,8 @@ namespace Metasound::Engine
 		FMetasoundFrontendInterface GetInterface()
 		{
 			FMetasoundFrontendInterface Interface;
-			Interface.Version = GetVersion();
-			Interface.UClassOptions = DeprecatedInterfacesPrivate::GetDeprecatedClassOptions();
+			Interface.Metadata.Version = GetVersion();
+			Interface.Metadata.UClassOptions = DeprecatedInterfacesPrivate::GetDeprecatedClassOptions();
 
 			// Inputs
 			FMetasoundFrontendClassVertex OnPlayTrigger;
@@ -176,8 +176,8 @@ namespace Metasound::Engine
 		FMetasoundFrontendInterface GetInterface()
 		{
 			FMetasoundFrontendInterface Interface;
-			Interface.Version = GetVersion();
-			Interface.UClassOptions = DeprecatedInterfacesPrivate::GetDeprecatedClassOptions();
+			Interface.Metadata.Version = GetVersion();
+			Interface.Metadata.UClassOptions = DeprecatedInterfacesPrivate::GetDeprecatedClassOptions();
 
 			// Inputs
 			FMetasoundFrontendClassVertex OnPlayTrigger;
@@ -259,8 +259,8 @@ namespace Metasound::Engine
 		FMetasoundFrontendInterface GetInterface()
 		{
 			FMetasoundFrontendInterface Interface;
-			Interface.Version = GetVersion();
-			Interface.UClassOptions = DeprecatedInterfacesPrivate::GetDeprecatedClassOptions();
+			Interface.Metadata.Version = GetVersion();
+			Interface.Metadata.UClassOptions = DeprecatedInterfacesPrivate::GetDeprecatedClassOptions();
 
 			// Inputs
 			FMetasoundFrontendClassVertex OnPlayTrigger;
@@ -311,8 +311,8 @@ namespace Metasound::Engine
 					return false;
 				}
 
-				InDocument->RemoveInterfaceVersion(MetasoundOutputFormatMonoV1_0::GetInterface().Version);
-				InDocument->AddInterfaceVersion(MetasoundOutputFormatMonoV1_1::GetInterface().Version);
+				InDocument->RemoveInterfaceVersion(MetasoundOutputFormatMonoV1_0::GetInterface().Metadata.Version);
+				InDocument->AddInterfaceVersion(MetasoundOutputFormatMonoV1_1::GetInterface().Metadata.Version);
 
 				FNodeHandle MonoFormatOutput = Graph->GetOutputNodeWithName(MetasoundOutputFormatMonoV1_0::GetAudioOutputName());
 #if WITH_EDITOR
@@ -425,8 +425,8 @@ namespace Metasound::Engine
 		FMetasoundFrontendInterface GetInterface()
 		{
 			FMetasoundFrontendInterface Interface;
-			Interface.Version = GetVersion();
-			Interface.UClassOptions = DeprecatedInterfacesPrivate::GetDeprecatedClassOptions();
+			Interface.Metadata.Version = GetVersion();
+			Interface.Metadata.UClassOptions = DeprecatedInterfacesPrivate::GetDeprecatedClassOptions();
 
 			// Inputs
 			FMetasoundFrontendClassVertex OnPlayTrigger;
@@ -482,8 +482,8 @@ namespace Metasound::Engine
 					return false;
 				}
 
-				InDocument->RemoveInterfaceVersion(MetasoundOutputFormatStereoV1_0::GetInterface().Version);
-				InDocument->AddInterfaceVersion(MetasoundOutputFormatStereoV1_1::GetInterface().Version);
+				InDocument->RemoveInterfaceVersion(MetasoundOutputFormatStereoV1_0::GetInterface().Metadata.Version);
+				InDocument->AddInterfaceVersion(MetasoundOutputFormatStereoV1_1::GetInterface().Metadata.Version);
 
 				FNodeHandle StereoFormatOutput = Graph->GetOutputNodeWithName(MetasoundOutputFormatStereoV1_0::GetAudioOutputName());
 				FOutputHandle LeftOutputToReconnect = IOutputController::GetInvalidHandle();
@@ -565,8 +565,8 @@ namespace Metasound::Engine
 		FMetasoundFrontendInterface GetInterface()
 		{
 			FMetasoundFrontendInterface Interface;
-			Interface.Version = GetVersion();
-			Interface.UClassOptions = DeprecatedInterfacesPrivate::GetDeprecatedClassOptions();
+			Interface.Metadata.Version = GetVersion();
+			Interface.Metadata.UClassOptions = DeprecatedInterfacesPrivate::GetDeprecatedClassOptions();
 			return Interface;
 		}
 
@@ -603,8 +603,8 @@ namespace Metasound::Engine
 		FMetasoundFrontendInterface GetInterface()
 		{
 			FMetasoundFrontendInterface Interface;
-			Interface.Version = GetVersion();
-			Interface.UClassOptions = DeprecatedInterfacesPrivate::GetDeprecatedClassOptions();
+			Interface.Metadata.Version = GetVersion();
+			Interface.Metadata.UClassOptions = DeprecatedInterfacesPrivate::GetDeprecatedClassOptions();
 			return Interface;
 		}
 
@@ -644,7 +644,7 @@ namespace Metasound::Engine
 		// Set default interface with unset version to use base UMetaSoundPatch class implementation (legacy requirement for 5.0 alpha).
 		{
 			FMetasoundFrontendInterface DeprecatedInterface;
-			DeprecatedInterface.UClassOptions = DeprecatedInterfacesPrivate::GetDeprecatedClassOptions();
+			DeprecatedInterface.Metadata.UClassOptions = DeprecatedInterfacesPrivate::GetDeprecatedClassOptions();
 			Reg.RegisterInterface(MakeUnique<FInterfaceRegistryEntry>(MoveTemp(DeprecatedInterface), RouterName, bDeprecated));
 		}
 
diff --git a/Engine/Plugins/Runtime/Metasound/Source/MetasoundEngine/Private/MetasoundAssetSubsystem.cpp b/Engine/Plugins/Runtime/Metasound/Source/MetasoundEngine/Private/MetasoundAssetSubsystem.cpp
--- a/Engine/Plugins/Runtime/Metasound/Source/MetasoundEngine/Private/MetasoundAssetSubsystem.cpp
+++ b/Engine/Plugins/Runtime/Metasound/Source/MetasoundEngine/Private/MetasoundAssetSubsystem.cpp
@@ -40,22 +40,22 @@ namespace Metasound::Engine
 
 		std::atomic<bool> InitialAssetScanComplete = false;
 
-		bool RemoveAssetTagData(FCriticalSection& MapCritSec, FAssetClassInfoMap& Map, const Frontend::FAssetKey& AssetKey, const FTopLevelAssetPath& InAssetPath)
+		bool RemoveClassInfo(FCriticalSection& MapCritSec, FAssetClassInfoMap& Map, const Frontend::FAssetKey& AssetKey, const FTopLevelAssetPath& InAssetPath)
 		{
 			using namespace Frontend;
 
 			FScopeLock Lock(&MapCritSec);
-			if (TArray<FAssetClassQueryInfo>* AssetInfoArray = Map.Find(AssetKey))
+			if (TArray<FMetaSoundAssetClassInfo>* AssetInfoArray = Map.Find(AssetKey))
 			{
-				auto ComparePaths = [&InAssetPath](const FAssetClassQueryInfo& QueryInfo)
+				auto ComparePaths = [&InAssetPath](const FMetaSoundAssetClassInfo& ClassInfo)
 				{
 					// Compare full paths if valid
-					if (QueryInfo.AssetPath.IsValid() && InAssetPath.IsValid())
+					if (ClassInfo.AssetPath.IsValid() && InAssetPath.IsValid())
 					{
-						return QueryInfo.AssetPath == InAssetPath;
+						return ClassInfo.AssetPath == InAssetPath;
 					}
 					// Package names are stripped on destruction, so only asset name is reliable
-					return QueryInfo.AssetPath.GetAssetName() == InAssetPath.GetAssetName();
+					return ClassInfo.AssetPath.GetAssetName() == InAssetPath.GetAssetName();
 				};
 
 				if (AssetInfoArray->RemoveAllSwap(ComparePaths, EAllowShrinking::No) > 0)
@@ -71,27 +71,27 @@ namespace Metasound::Engine
 			return false;
 		}
 
-		Frontend::FAssetKey AddAssetTagData(FCriticalSection& MapCritSec, FAssetClassInfoMap& Map, Frontend::FAssetClassQueryInfo&& QueryInfo)
+		Frontend::FAssetKey AddClassInfo(FCriticalSection& MapCritSec, FAssetClassInfoMap& Map, Frontend::FMetaSoundAssetClassInfo&& ClassInfo)
 		{
 			using namespace Frontend;
 
-			const Frontend::FAssetKey AssetKey(QueryInfo.ClassName, QueryInfo.Version);
+			const Frontend::FAssetKey AssetKey(ClassInfo.ClassName, ClassInfo.Version);
 			if (!AssetKey.IsValid())
 			{
 				return AssetKey;
 			}
 
 			FScopeLock Lock(&MapCritSec);
-			TArray<FAssetClassQueryInfo>& TagDatas = Map.FindOrAdd(AssetKey);
-			auto IsTagData = [&QueryInfo](const FAssetClassQueryInfo& Iter) { return Iter.AssetPath == QueryInfo.AssetPath; };
+			TArray<FMetaSoundAssetClassInfo>& TagDatas = Map.FindOrAdd(AssetKey);
+			auto IsTagData = [&ClassInfo](const FMetaSoundAssetClassInfo& Iter) { return Iter.AssetPath == ClassInfo.AssetPath; };
 			TagDatas.RemoveAllSwap(IsTagData);
-			TagDatas.Add(MoveTemp(QueryInfo));
+			TagDatas.Add(MoveTemp(ClassInfo));
 
 #if !NO_LOGGING
 			if (TagDatas.Num() > 1)
 			{
 				TArray<FString> PathStrings;
-				Algo::Transform(TagDatas, PathStrings, [](const FAssetClassQueryInfo& QueryInfo) { return QueryInfo.AssetPath.ToString(); });
+				Algo::Transform(TagDatas, PathStrings, [](const FMetaSoundAssetClassInfo& ClassInfo) { return ClassInfo.AssetPath.ToString(); });
 				UE_LOG(LogMetaSound, Warning,
 					TEXT("MetaSoundAssetManager has registered multiple assets with key '%s':\n%s\n"),
 					*AssetKey.ToString(),
@@ -110,7 +110,7 @@ namespace Metasound::Engine
 
 		if (bLogActiveAssetsOnShutdown)
 		{
-			TMap<FAssetKey, TArray<FAssetClassQueryInfo>> InfoOnShutdown;
+			TMap<FAssetKey, TArray<FMetaSoundAssetClassInfo>> InfoOnShutdown;
 			{
 				FScopeLock Lock(&TagDataMapCriticalSection);
 				InfoOnShutdown = MoveTemp(ClassInfoMap);
@@ -120,11 +120,11 @@ namespace Metasound::Engine
 			if (!InfoOnShutdown.IsEmpty())
 			{
 				UE_LOG(LogMetaSound, Display, TEXT("AssetManager is shutting down with the following %i assets active:"), InfoOnShutdown.Num());
-				for (const TPair<FAssetKey, TArray<FAssetClassQueryInfo>>& Pair : InfoOnShutdown)
+				for (const TPair<FAssetKey, TArray<FMetaSoundAssetClassInfo>>& Pair : InfoOnShutdown)
 				{
-					for (const FAssetClassQueryInfo& QueryInfo : Pair.Value)
+					for (const FMetaSoundAssetClassInfo& ClassInfo : Pair.Value)
 					{
-						UE_LOG(LogMetaSound, Display, TEXT("- %s"), *QueryInfo.AssetPath.ToString());
+						UE_LOG(LogMetaSound, Display, TEXT("- %s"), *ClassInfo.AssetPath.ToString());
 					}
 				}
 			}
@@ -221,7 +221,7 @@ namespace Metasound::Engine
 		TScriptInterface<const IMetaSoundDocumentInterface> ScriptInterface(&InObject);
 		check(ScriptInterface.GetObject() != nullptr);
 		const IMetaSoundDocumentInterface& DocInterface = *ScriptInterface;
-		return AssetSubsystemPrivate::AddAssetTagData(TagDataMapCriticalSection, ClassInfoMap, FAssetClassQueryInfo(DocInterface));
+		return AssetSubsystemPrivate::AddClassInfo(TagDataMapCriticalSection, ClassInfoMap, FMetaSoundAssetClassInfo(DocInterface));
 	}
 
 	void  FMetaSoundAssetManager::AddOrLoadAndUpdateFromObjectAsync(const FAssetData& InAssetData, TUniqueFunction<void(FAssetKey, UObject&)>&& OnAssetRegistered)
@@ -264,7 +264,7 @@ namespace Metasound::Engine
 		);
 	}
 
-	void FMetaSoundAssetManager::AddOrUpdateFromAssetTags(const FAssetData& InAssetData)
+	void FMetaSoundAssetManager::AddOrUpdateFromAssetData(const FAssetData& InAssetData)
 	{
 		using namespace Frontend;
 
@@ -280,17 +280,17 @@ namespace Metasound::Engine
 		const FSoftObjectPath Path = InAssetData.ToSoftObjectPath();
 		if (UObject* Object = Path.ResolveObject())
 		{
-			UE_LOG(LogMetaSound, VeryVerbose, TEXT("FMetaSoundAssetManager::AddOrUpdateFromAssetTags, but object '%s' already loaded. Updating entry using in-memory object definition."), *Path.ToString());
+			UE_LOG(LogMetaSound, VeryVerbose, TEXT("AddOrUpdateFromAssetData called, but object '%s' already loaded. Updating entry using in-memory object definition."), *Path.ToString());
 			AddOrUpdateFromObject(*Object);
 		}
 		else
 		{
 			// Older document versions did not serialize asset tags. In these cases, a load of the object is required
 			// in order to retrieve the appropriate asset tag data translated to asset info cached in the MetaSound Asset Manager.
-			FAssetClassQueryInfo QueryInfo(InAssetData);
-			if (QueryInfo.bIsValid)
+			FMetaSoundAssetClassInfo ClassInfo(InAssetData);
+			if (ClassInfo.bIsValid)
 			{
-				AssetSubsystemPrivate::AddAssetTagData(TagDataMapCriticalSection, ClassInfoMap, MoveTemp(QueryInfo));
+				AssetSubsystemPrivate::AddClassInfo(TagDataMapCriticalSection, ClassInfoMap, MoveTemp(ClassInfo));
 			}
 			else
 			{
@@ -392,7 +392,42 @@ namespace Metasound::Engine
 		return nullptr;
 	}
 
-	TArray<Frontend::FAssetClassQueryInfo> FMetaSoundAssetManager::FindAssetTagData(const Frontend::FAssetKey& InKey) const
+	Frontend::FMetaSoundAssetClassInfo FMetaSoundAssetManager::FindAssetClassInfo(const FTopLevelAssetPath& InPath) const
+	{
+		using namespace Frontend;
+
+		const FSoftObjectPath ObjectPath(InPath);
+		if (const UObject* Object = ObjectPath.ResolveObject())
+		{
+			TScriptInterface<const IMetaSoundDocumentInterface> DocInterface(Object);
+			if (DocInterface.GetObject())
+			{
+				return FMetaSoundAssetClassInfo(*DocInterface.GetInterface());
+			}
+		}
+
+		FAssetData AssetData;
+		const FAssetRegistryModule& AssetRegistryModule = FModuleManager::GetModuleChecked<FAssetRegistryModule>("AssetRegistry");
+		UE::AssetRegistry::EExists Exists = AssetRegistryModule.TryGetAssetByObjectPath(ObjectPath, AssetData);
+		if (Exists == UE::AssetRegistry::EExists::Exists)
+		{
+			FMetaSoundAssetKey AssetKey;
+			if (FMetaSoundAssetClassInfo::TryGetAssetKey(AssetData, AssetKey))
+			{
+				const TArray<Frontend::FMetaSoundAssetClassInfo> ClassInfo = FindAssetClassInfoInternal(AssetKey);
+				if (ClassInfo.Num() == 1)
+				{
+					return ClassInfo.Last();
+				}
+			}
+
+			return FMetaSoundAssetClassInfo(AssetData);
+		}
+
+		return { };
+	}
+
+	TArray<Frontend::FMetaSoundAssetClassInfo> FMetaSoundAssetManager::FindAssetClassInfoInternal(const Frontend::FAssetKey& InKey) const
 	{
 		FScopeLock Lock(&TagDataMapCriticalSection);
 		return ClassInfoMap.FindRef(InKey);
@@ -417,7 +452,7 @@ namespace Metasound::Engine
 		using namespace Frontend;
 
 		FScopeLock Lock(&TagDataMapCriticalSection);
-		if (const TArray<FAssetClassQueryInfo>* TagDatas = ClassInfoMap.Find(InKey))
+		if (const TArray<FMetaSoundAssetClassInfo>* TagDatas = ClassInfoMap.Find(InKey))
 		{
 			if (!TagDatas->IsEmpty())
 			{
@@ -433,10 +468,10 @@ namespace Metasound::Engine
 		using namespace Frontend;
 
 		FScopeLock Lock(&TagDataMapCriticalSection);
-		if (const TArray<FAssetClassQueryInfo>* TagDatas = ClassInfoMap.Find(InKey))
+		if (const TArray<FMetaSoundAssetClassInfo>* TagDatas = ClassInfoMap.Find(InKey))
 		{
 			TArray<FTopLevelAssetPath> Paths;
-			Algo::Transform(*TagDatas, Paths, [](const FAssetClassQueryInfo& QueryInfo) { return QueryInfo.AssetPath; });
+			Algo::Transform(*TagDatas, Paths, [](const FMetaSoundAssetClassInfo& ClassInfo) { return ClassInfo.AssetPath; });
 			return Paths;
 		}
 
@@ -464,26 +499,15 @@ namespace Metasound::Engine
 	{
 		using namespace Frontend;
 
-		FAssetKey AssetKey;
-
-		if (const UObject* Object = InObjectPath.ResolveObject())
-		{
-			TScriptInterface<const IMetaSoundDocumentInterface> DocInterface(Object);
-			if (DocInterface.GetObject())
-			{
-				return FAssetKey(DocInterface->GetConstDocument().RootGraph);
-			}
-		}
-
 		FAssetData AssetData;
 		const FAssetRegistryModule& AssetRegistryModule = FModuleManager::GetModuleChecked<FAssetRegistryModule>("AssetRegistry");
 		UE::AssetRegistry::EExists Exists = AssetRegistryModule.TryGetAssetByObjectPath(InObjectPath, AssetData);
 		if (Exists == UE::AssetRegistry::EExists::Exists)
 		{
-			const FAssetClassQueryInfo QueryInfo(AssetData);
-			if (QueryInfo.bIsValid)
+			FMetaSoundAssetKey AssetKey;
+			if (FMetaSoundAssetClassInfo::TryGetAssetKey(AssetData, AssetKey))
 			{
-				return FAssetKey(QueryInfo.ClassName, QueryInfo.Version);
+				return AssetKey;
 			}
 		}
 
@@ -726,13 +750,13 @@ namespace Metasound::Engine
 
 		{
 			const FAssetKey OldAssetKey(ClassMetadata.GetClassName(), ClassMetadata.GetVersion());
-			AssetSubsystemPrivate::RemoveAssetTagData(TagDataMapCriticalSection, ClassInfoMap, OldAssetKey, Path);
+			AssetSubsystemPrivate::RemoveClassInfo(TagDataMapCriticalSection, ClassInfoMap, OldAssetKey, Path);
 		}
 
 		Builder.GenerateNewClassName();
 
 		{
-			AssetSubsystemPrivate::AddAssetTagData(TagDataMapCriticalSection, ClassInfoMap, FAssetClassQueryInfo(*DocInterface));
+			AssetSubsystemPrivate::AddClassInfo(TagDataMapCriticalSection, ClassInfoMap, FMetaSoundAssetClassInfo(*DocInterface));
 		}
 
 		AssetBase->UpdateAndRegisterForExecution();
@@ -756,11 +780,11 @@ namespace Metasound::Engine
 	}
 
 #if WITH_EDITOR
-	void FMetaSoundAssetManager::IterateAssetTagData(TFunctionRef<void(Frontend::FAssetClassQueryInfo)> Iter, bool bIterateDuplicates) const
+	void FMetaSoundAssetManager::IterateAssetTagData(TFunctionRef<void(Frontend::FMetaSoundAssetClassInfo)> Iter, bool bIterateDuplicates) const
 	{
 		using namespace Frontend;
 
-		TArray<TArray<FAssetClassQueryInfo>> TagDataMatrix;
+		TArray<TArray<FMetaSoundAssetClassInfo>> TagDataMatrix;
 		{
 			FScopeLock Lock(&TagDataMapCriticalSection);
 			ClassInfoMap.GenerateValueArray(TagDataMatrix);
@@ -768,17 +792,17 @@ namespace Metasound::Engine
 
 		if (bIterateDuplicates)
 		{
-			for (TArray<FAssetClassQueryInfo>& TagDataArray : TagDataMatrix)
+			for (TArray<FMetaSoundAssetClassInfo>& TagDataArray : TagDataMatrix)
 			{
-				for (FAssetClassQueryInfo& QueryInfo : TagDataArray)
+				for (FMetaSoundAssetClassInfo& ClassInfo : TagDataArray)
 				{
-					Iter(MoveTemp(QueryInfo));
+					Iter(MoveTemp(ClassInfo));
 				}
 			}
 		}
 		else
 		{
-			for (TArray<FAssetClassQueryInfo>& TagDataArray : TagDataMatrix)
+			for (TArray<FMetaSoundAssetClassInfo>& TagDataArray : TagDataMatrix)
 			{
 				if (!TagDataArray.IsEmpty())
 				{
@@ -794,17 +818,17 @@ namespace Metasound::Engine
 		{
 			using namespace Frontend;
 
-			TArray<FAssetClassQueryInfo> TagDatas;
+			TArray<FMetaSoundAssetClassInfo> TagDatas;
 			{
 				FScopeLock Lock(&TagDataMapCriticalSection);
-				if (const TArray<FAssetClassQueryInfo>* TagDataArray = ClassInfoMap.Find(ReferencedKey))
+				if (const TArray<FMetaSoundAssetClassInfo>* TagDataArray = ClassInfoMap.Find(ReferencedKey))
 				{
 					TagDatas = *TagDataArray;
 				}
 			}
 
 			TSet<FAssetKey> Children;
-			for (FAssetClassQueryInfo& QueryInfo : TagDatas)
+			for (FMetaSoundAssetClassInfo& ClassInfo : TagDatas)
 			{
 				// If loaded, this is the "freshest" reference list, so use that.
 				if (const FMetasoundAssetBase* LoadedAsset = FindAsset(ReferencedKey))
@@ -834,11 +858,11 @@ namespace Metasound::Engine
 				// Otherwise, use provided reference list from the asset tag data
 				else
 				{
-					for (const FAssetKey& RefKey : QueryInfo.ReferencedAssetKeys)
+					for (const FAssetKey& RefKey : ClassInfo.DocInfo.ReferencedAssetKeys)
 					{
 						VisitFunction(RefKey);
 					}
-					Children.Append(MoveTemp(QueryInfo.ReferencedAssetKeys));
+					Children.Append(MoveTemp(ClassInfo.DocInfo.ReferencedAssetKeys));
 				}
 			}
 
@@ -873,10 +897,10 @@ namespace Metasound::Engine
 			FAssetData AssetData;
 			if (AssetManager.GetAssetDataForPath(UpdateSettings.MetaSound, AssetData))
 			{
-				const FAssetClassQueryInfo QueryInfo(AssetData);
-				if (QueryInfo.bIsValid)
+				const FMetaSoundAssetClassInfo ClassInfo(AssetData);
+				if (ClassInfo.bIsValid)
 				{
-					AutoUpdateDenyListCache.Add(QueryInfo.ClassName.GetFullName());
+					AutoUpdateDenyListCache.Add(ClassInfo.ClassName.GetFullName());
 				}
 			}
 		}
@@ -923,24 +947,24 @@ namespace Metasound::Engine
 		}
 
 		const FAssetKey AssetKey(Metadata.GetClassName(), Metadata.GetVersion());
-		AssetSubsystemPrivate::RemoveAssetTagData(TagDataMapCriticalSection, ClassInfoMap, AssetKey, AssetPath);
+		AssetSubsystemPrivate::RemoveClassInfo(TagDataMapCriticalSection, ClassInfoMap, AssetKey, AssetPath);
 	}
 
 	void FMetaSoundAssetManager::RemoveAsset(const FAssetData& InAssetData)
 	{
 		using namespace Frontend;
 
-		const FAssetClassQueryInfo QueryInfo(InAssetData);
-		if (QueryInfo.bIsValid)
+		const FMetaSoundAssetClassInfo ClassInfo(InAssetData);
+		if (ClassInfo.bIsValid)
 		{
 			if (IDocumentBuilderRegistry* BuilderRegistry = IDocumentBuilderRegistry::Get())
 			{
 				constexpr bool bForceUnregister = true;
-				BuilderRegistry->FinishBuilding(QueryInfo.ClassName, QueryInfo.AssetPath, bForceUnregister);
+				BuilderRegistry->FinishBuilding(ClassInfo.ClassName, ClassInfo.AssetPath, bForceUnregister);
 			}
 
-			const FAssetKey AssetKey(QueryInfo.ClassName, QueryInfo.Version);
-			AssetSubsystemPrivate::RemoveAssetTagData(TagDataMapCriticalSection, ClassInfoMap, AssetKey, QueryInfo.AssetPath);
+			const FAssetKey AssetKey(ClassInfo.ClassName, ClassInfo.Version);
+			AssetSubsystemPrivate::RemoveClassInfo(TagDataMapCriticalSection, ClassInfoMap, AssetKey, ClassInfo.AssetPath);
 		}
 	}
 
@@ -960,13 +984,13 @@ namespace Metasound::Engine
 		FMetasoundAssetBase* MetaSoundAsset = GetAsAsset(*InAssetData.GetAsset());
 		check(MetaSoundAsset);
 
-		FAssetClassQueryInfo QueryInfo(InAssetData);
-		if (ensureAlways(QueryInfo.bIsValid))
+		FMetaSoundAssetClassInfo ClassInfo(InAssetData);
+		if (ensureAlways(ClassInfo.bIsValid))
 		{
-			const FAssetKey AssetKey(QueryInfo.ClassName, QueryInfo.Version);
+			const FAssetKey AssetKey(ClassInfo.ClassName, ClassInfo.Version);
 			const FTopLevelAssetPath OldPath(InOldObjectPath);
-			AssetSubsystemPrivate::RemoveAssetTagData(TagDataMapCriticalSection, ClassInfoMap, AssetKey, OldPath);
-			AssetSubsystemPrivate::AddAssetTagData(TagDataMapCriticalSection, ClassInfoMap, MoveTemp(QueryInfo));
+			AssetSubsystemPrivate::RemoveClassInfo(TagDataMapCriticalSection, ClassInfoMap, AssetKey, OldPath);
+			AssetSubsystemPrivate::AddClassInfo(TagDataMapCriticalSection, ClassInfoMap, MoveTemp(ClassInfo));
 		}
 	}
 
@@ -1070,12 +1094,12 @@ namespace Metasound::Engine
 
 		IMetasoundUObjectRegistry& ObjectRegistry = IMetasoundUObjectRegistry::Get();
 		TSet<FMetasoundAssetBase*> ToReregister;
-		for (const TPair<FAssetKey, TArray<FAssetClassQueryInfo>>& Pair : ClassInfoMap)
+		for (const TPair<FAssetKey, TArray<FMetaSoundAssetClassInfo>>& Pair : ClassInfoMap)
 		{
 			if (!Pair.Value.IsEmpty())
 			{
-				const FAssetClassQueryInfo& QueryInfo = Pair.Value.Last();
-				if (UObject* Object = FSoftObjectPath(QueryInfo.AssetPath, { }).ResolveObject())
+				const FMetaSoundAssetClassInfo& ClassInfo = Pair.Value.Last();
+				if (UObject* Object = FSoftObjectPath(ClassInfo.AssetPath, { }).ResolveObject())
 				{
 					if (FMetasoundAssetBase* Asset = ObjectRegistry.GetObjectAsAssetBase(Object); Asset && Asset->IsRegistered())
 					{
@@ -1201,17 +1225,17 @@ namespace Metasound::Engine
 			}
 			else
 			{
-				const FAssetClassQueryInfo ClassQueryInfo(AssetData);
-				if (ensureAlways(ClassQueryInfo.bIsValid))
+				const FMetaSoundAssetClassInfo ClassInfo(AssetData);
+				if (ensureAlways(ClassInfo.bIsValid))
 				{
-					const FAssetKey AssetKey(ClassQueryInfo.ClassName, ClassQueryInfo.Version);
+					const FAssetKey AssetKey(ClassInfo.ClassName, ClassInfo.Version);
 					const FNodeRegistryKey RegistryKey = FNodeRegistryKey(AssetKey);
 					const bool bIsRegistered = FMetasoundFrontendRegistryContainer::Get()->IsNodeRegistered(RegistryKey);
 					if (bIsRegistered)
 					{
 						FMetasoundFrontendRegistryContainer::Get()->UnregisterNode(RegistryKey);
 						const FTopLevelAssetPath AssetPath(AssetData.PackageName, AssetData.AssetName);
-						AssetSubsystemPrivate::RemoveAssetTagData(TagDataMapCriticalSection, ClassInfoMap, AssetKey, AssetPath);
+						AssetSubsystemPrivate::RemoveClassInfo(TagDataMapCriticalSection, ClassInfoMap, AssetKey, AssetPath);
 					}
 				}
 			}
@@ -1402,7 +1426,7 @@ Metasound::Frontend::FNodeRegistryKey UMetaSoundAssetSubsystem::AddOrUpdateAsset
 Metasound::Frontend::FNodeRegistryKey UMetaSoundAssetSubsystem::AddOrUpdateAsset(const FAssetData& InAssetData)
 {
 	using namespace Metasound::Frontend;
-	IMetaSoundAssetManager::GetChecked().AddOrUpdateFromAssetTags(InAssetData);
+	IMetaSoundAssetManager::GetChecked().AddOrUpdateFromAssetData(InAssetData);
 	return { };
 }
 
@@ -1418,6 +1442,45 @@ bool UMetaSoundAssetSubsystem::ContainsKey(const Metasound::Frontend::FNodeRegis
 	return IMetaSoundAssetManager::GetChecked().ContainsKey(FAssetKey(InRegistryKey));
 }
 
+#if WITH_EDITOR
+bool UMetaSoundAssetSubsystem::FindAssetClassInfo(const FTopLevelAssetPath& InPath, FMetaSoundDocumentInfo& OutDocInfo, FMetaSoundClassInterfaceInfo& OutInterfaceInfo, bool bForceLoad) const
+{
+	using namespace Metasound;
+
+	const FSoftObjectPath ObjectPath(InPath);
+	const bool bWasLoaded = ObjectPath.ResolveObject() != nullptr;
+
+	auto FindInternal = [](const FTopLevelAssetPath& Path, FMetaSoundDocumentInfo& DocInfo, FMetaSoundClassInterfaceInfo& InterfaceInfo)
+	{
+		Frontend::FMetaSoundAssetClassInfo FoundInfo = Engine::FMetaSoundAssetManager::GetChecked().FindAssetClassInfo(Path);
+		if (FoundInfo.bIsValid)
+		{
+			DocInfo = MoveTemp(FoundInfo.DocInfo);
+			InterfaceInfo = MoveTemp(FoundInfo.InterfaceInfo);
+			return true;
+		}
+
+		return false;
+	};
+
+	if (FindInternal(InPath, OutDocInfo, OutInterfaceInfo))
+	{
+		return true;
+	}
+
+	// Attempt to load synchronously if set to force load
+	if (!bWasLoaded && bForceLoad)
+	{
+		if (ObjectPath.TryLoad())
+		{
+			return FindInternal(InPath, OutDocInfo, OutInterfaceInfo);
+		}
+	}
+
+	return false;
+}
+#endif // WITH_EDITOR
+
 FMetasoundAssetBase* UMetaSoundAssetSubsystem::GetAsAsset(UObject& InObject) const
 {
 	using namespace Metasound::Frontend;
diff --git a/Engine/Plugins/Runtime/Metasound/Source/MetasoundEngine/Private/MetasoundEngineAsset.h b/Engine/Plugins/Runtime/Metasound/Source/MetasoundEngine/Private/MetasoundEngineAsset.h
--- a/Engine/Plugins/Runtime/Metasound/Source/MetasoundEngine/Private/MetasoundEngineAsset.h
+++ b/Engine/Plugins/Runtime/Metasound/Source/MetasoundEngine/Private/MetasoundEngineAsset.h
@@ -209,8 +209,8 @@ namespace Metasound::Engine
 				return;
 			}
 
-			const FAssetClassQueryInfo QueryInfo(*DocInterface);
-			QueryInfo.ExportToContext(Context);
+			const FMetaSoundAssetClassInfo ClassInfo(*DocInterface);
+			ClassInfo.ExportToContext(Context);
 		}
 
 		template <typename TMetaSoundObject>
diff --git a/Engine/Plugins/Runtime/Metasound/Source/MetasoundEngine/Private/MetasoundEngineModule.cpp b/Engine/Plugins/Runtime/Metasound/Source/MetasoundEngine/Private/MetasoundEngineModule.cpp
--- a/Engine/Plugins/Runtime/Metasound/Source/MetasoundEngine/Private/MetasoundEngineModule.cpp
+++ b/Engine/Plugins/Runtime/Metasound/Source/MetasoundEngine/Private/MetasoundEngineModule.cpp
@@ -208,7 +208,7 @@ namespace Metasound::Engine
 		}
 		else
 		{
-			IMetaSoundAssetManager::GetChecked().AddOrUpdateFromAssetTags(InAssetData);
+			IMetaSoundAssetManager::GetChecked().AddOrUpdateFromAssetData(InAssetData);
 		}
 	}
 	
@@ -240,7 +240,7 @@ namespace Metasound::Engine
 				}
 				else
 				{
-					IMetaSoundAssetManager::GetChecked().AddOrUpdateFromAssetTags(InAssetData);
+					IMetaSoundAssetManager::GetChecked().AddOrUpdateFromAssetData(InAssetData);
 				}
 			}
 		}
diff --git a/Engine/Plugins/Runtime/Metasound/Source/MetasoundEngine/Public/MetasoundAssetSubsystem.h b/Engine/Plugins/Runtime/Metasound/Source/MetasoundEngine/Public/MetasoundAssetSubsystem.h
--- a/Engine/Plugins/Runtime/Metasound/Source/MetasoundEngine/Public/MetasoundAssetSubsystem.h
+++ b/Engine/Plugins/Runtime/Metasound/Source/MetasoundEngine/Public/MetasoundAssetSubsystem.h
@@ -35,7 +35,7 @@ namespace Metasound::Engine
 	public:
 		friend class UMetaSoundAssetSubsystem;
 
-		using FAssetClassInfoMap = TMap<Frontend::FAssetKey, TArray<Frontend::FAssetClassQueryInfo>>;
+		using FAssetClassInfoMap = TMap<Frontend::FAssetKey, TArray<Frontend::FMetaSoundAssetClassInfo>>;
 		using FAssetInfo = Frontend::IMetaSoundAssetManager::FAssetInfo;
 		using FAssetKey = Frontend::FAssetKey;
 		using FNodeRegistryKey = Frontend::FNodeRegistryKey;
@@ -59,7 +59,15 @@ namespace Metasound::Engine
 			return static_cast<FMetaSoundAssetManager&>(IMetaSoundAssetManager::GetChecked());
 		}
 
-		TArray<Frontend::FAssetClassQueryInfo> FindAssetTagData(const Frontend::FAssetKey& InKey) const;
+		/**
+		 * Attempts to find the most recent class query info for the asset at the given path without loading asset.
+		 * 1. If asset is loaded, builds tag data from asset (to ensure version provided is "freshest" if it has been edited but not re-serialized.
+		 * 2. If asset is not loaded:
+		 *	a. Returns entry in Manager if sole entry found.
+		 *	b. If not registered or multiple entries found (may be more than one due to errors in duplication of assets, improper plugin migration), builds directly from the asset data last serialized with the given path/asset.
+		 *	c. If the asset data is out-of-date or not found, returns invalid query info.
+		 **/
+		Frontend::FMetaSoundAssetClassInfo FindAssetClassInfo(const FTopLevelAssetPath& InPath) const;
 
 		Frontend::FAssetKey GetAssetKey(const FSoftObjectPath& InObjectPath) const;
 
@@ -95,7 +103,7 @@ namespace Metasound::Engine
 #endif // WITH_EDITORONLY_DATA
 		virtual FAssetKey AddOrUpdateFromObject(const UObject& InObject) override;
 		virtual void AddOrLoadAndUpdateFromObjectAsync(const FAssetData& InAssetData, TUniqueFunction<void(FAssetKey, UObject&)>&& OnAssetRegistered) override;
-		virtual void AddOrUpdateFromAssetTags(const FAssetData& InAssetData) override;
+		virtual void AddOrUpdateFromAssetData(const FAssetData& InAssetData) override;
 		virtual bool CanAutoUpdate(const FMetasoundFrontendClassName& InClassName) const override;
 		virtual bool ContainsKey(const FAssetKey& InKey) const override;
 		virtual bool ContainsKey(const Frontend::FNodeRegistryKey& InRegistryKey) const override;
@@ -115,7 +123,7 @@ namespace Metasound::Engine
 
 #if WITH_EDITOR
 		// Iterates all asset tag data (Only recommended in certain editor contexts, as slow and blocks access to reference map).
-		void IterateAssetTagData(TFunctionRef<void(Frontend::FAssetClassQueryInfo)> Iter, bool bIterateDuplicates = false) const;
+		void IterateAssetTagData(TFunctionRef<void(Frontend::FMetaSoundAssetClassInfo)> Iter, bool bIterateDuplicates = false) const;
 
 		// Iterates all references of a given asset key entry (Only recommended in certain editor contexts, as can be slow for
 		// deep reference trees and blocks access to reference map).
@@ -144,6 +152,13 @@ namespace Metasound::Engine
 		TArray<FMetaSoundAsyncAssetDependencies> LoadingDependencies;
 
 		static void DepthFirstVisit_AssetKey(const FAssetKey& InKey, TFunctionRef<TSet<FAssetKey>(const FAssetKey&)> VisitFunction);
+
+		// Returns all asset query info with the given asset key (may be more than one due to errors in duplication
+		// of assets, improper plugin migration, etc. but should usually just be a single entry) registered
+		// with the manager. Info may not be returned if the asset has yet to be registered with the MetaSoundAssetManager
+		// or may be out-of-date if edited after last registration attempt via AddOrUpdate... calls.
+		TArray<Frontend::FMetaSoundAssetClassInfo> FindAssetClassInfoInternal(const FMetaSoundAssetKey& InKey) const;
+
 		FMetaSoundAsyncAssetDependencies* FindLoadingDependencies(const UObject* InParentAsset);
 		FMetaSoundAsyncAssetDependencies* FindLoadingDependencies(int32 InLoadID);
 		void RemoveLoadingDependencies(int32 InLoadID);
@@ -284,6 +299,17 @@ public:
 	virtual void WaitUntilAsyncLoadReferencedAssetsComplete(FMetasoundAssetBase& InAssetBase) { }
 
 #if WITH_EDITOR
+	// Returns asset class info for the given MetaSound asset. Will attempt to get
+	// info without loading the asset if its tag data is up-to-date, or if set to force load, will load
+	// otherwise (synchronously and can be slow).  Returns true if asset is found, was MetaSound, and all data
+	// was retrieved successfully, false if not.
+	UFUNCTION(BlueprintCallable, Category = "MetaSounds|Utilities", meta = (AdvancedDisplay = "3"))
+	UPARAM(DisplayName = "Info Found") bool FindAssetClassInfo(
+		UPARAM(DisplayName = "Path") const FTopLevelAssetPath& InPath,
+		UPARAM(DisplayName = "Document Info") FMetaSoundDocumentInfo& OutDocInfo,
+		UPARAM(DisplayName = "Interface Info") FMetaSoundClassInterfaceInfo& OutInterfaceInfo,
+		UPARAM(DisplayName = "Force Load") bool bForceLoad = true) const;
+
 	UFUNCTION(BlueprintCallable, Category = "MetaSounds|Utilities")
 	UPARAM(DisplayName = "Reassigned") bool ReassignClassName(TScriptInterface<IMetaSoundDocumentInterface> DocInterface);
 
diff --git a/Engine/Plugins/Runtime/Metasound/Source/MetasoundEngineTest/Private/EngineTestMetasoundSource.cpp b/Engine/Plugins/Runtime/Metasound/Source/MetasoundEngineTest/Private/EngineTestMetasoundSource.cpp
--- a/Engine/Plugins/Runtime/Metasound/Source/MetasoundEngineTest/Private/EngineTestMetasoundSource.cpp
+++ b/Engine/Plugins/Runtime/Metasound/Source/MetasoundEngineTest/Private/EngineTestMetasoundSource.cpp
@@ -10,6 +10,7 @@
 #include "Interfaces/MetasoundOutputFormatInterfaces.h"
 #include "Interfaces/MetasoundFrontendSourceInterface.h"
 #include "Math/UnrealMathUtility.h"
+#include "MetasoundAssetSubsystem.h"
 #include "MetasoundDataReference.h"
 #include "MetasoundFrontendDocumentBuilder.h"
 #include "MetasoundLog.h"
@@ -915,4 +916,53 @@ bool FAudioMetasoundSourceBuilderAddRemoveNodes::RunTest(const FString& Paramete
 
 	return true;
 }
+
+// This test attempts to load class info from asset tag data set on a MetaSound source and check its validity via the MetaSound Asset Subsystem.
+IMPLEMENT_SIMPLE_AUTOMATION_TEST(FAudioMetaSoundFindAssetClassInfo, "Audio.Metasound.Asset.FindAssetClassInfo", EAutomationTestFlags::EditorContext | EAutomationTestFlags::EngineFilter)
+bool FAudioMetaSoundFindAssetClassInfo::RunTest(const FString& Parameters)
+{
+#if WITH_EDITOR
+	const FTopLevelAssetPath Path(FName("/Game/Assets/MetaSounds/AssetSubsystem/QueryClassInfoAutomatedTest_Preset"), FName("QueryClassInfoAutomatedTest_Preset"));
+
+	FMetaSoundDocumentInfo DocInfo;
+	FMetaSoundClassInterfaceInfo ClassInterfaceInfo;
+
+	UMetaSoundAssetSubsystem* Subsystem = GEngine->GetEngineSubsystem<UMetaSoundAssetSubsystem>();
+	AddErrorIfFalse(Subsystem != nullptr, TEXT("MetaSound Asset Subsystem could not be found. Skipping additional class query testing."));
+
+	if (Subsystem)
+	{
+		const bool bAssetFound = Subsystem->FindAssetClassInfo(Path, DocInfo, ClassInterfaceInfo);
+		AddErrorIfFalse(bAssetFound, FString::Format(
+			TEXT("MetaSound asset class query failed to find or parse asset tag data for asset '{0}'."),
+			{ Path.ToString() }));
+
+		if (bAssetFound)
+		{
+			TScriptInterface<IMetaSoundDocumentInterface> DocInterface = FSoftObjectPath(Path).TryLoad();
+			if (DocInterface.GetObject())
+			{
+				AddErrorIfFalse(DocInfo.bIsPreset, TEXT("MetaSound asset query tags did not return that asset is preset"));
+
+				const FMetasoundFrontendDocument& Doc = DocInterface->GetConstDocument();
+				const FMetasoundFrontendGraphClass& GraphClass = Doc.RootGraph;
+
+				const FText StringDisplayName = GraphClass.Metadata.GetDisplayName();
+				AddErrorIfFalse(ClassInterfaceInfo.SearchInfo.ClassDisplayName.EqualTo(StringDisplayName), FString::Format(TEXT("MetaSound display name should have been '{0}'"), { StringDisplayName.ToString() }));
+
+				const FText StringDesc = GraphClass.Metadata.GetDescription();
+				AddErrorIfFalse(ClassInterfaceInfo.SearchInfo.ClassDescription.EqualTo(StringDesc), *FString::Format(TEXT("MetaSound description should have been '{0}'"), { StringDesc.ToString()}));
+
+				AddErrorIfFalse(ClassInterfaceInfo.DefinedInterfaces.Num() == 0, TEXT("Cool! We support declared interfaces now! Update the test!"));
+
+				AddErrorIfFalse(ClassInterfaceInfo.InheritedInterfaces.Num() == Doc.Interfaces.Num(), TEXT("MetaSound asset query tags did not return that asset implements number of interfaces as expected"));
+				AddErrorIfFalse(ClassInterfaceInfo.Inputs.Num() == GraphClass.Interface.Inputs.Num(), TEXT("MetaSound asset query tags did not return that asset implements number of inputs expected"));
+				AddErrorIfFalse(ClassInterfaceInfo.Outputs.Num() == GraphClass.Interface.Outputs.Num(), TEXT("MetaSound asset query tags did not return that asset implements number of outputs expected"));
+			}
+		}
+	}
+#endif // WITH_EDITOR
+
+	return true;
+}
 #endif // WITH_DEV_AUTOMATION_TESTS
diff --git a/Engine/Plugins/Runtime/Metasound/Source/MetasoundFrontend/Private/Interfaces/MetasoundFrontendInterface.cpp b/Engine/Plugins/Runtime/Metasound/Source/MetasoundFrontend/Private/Interfaces/MetasoundFrontendInterface.cpp
--- a/Engine/Plugins/Runtime/Metasound/Source/MetasoundFrontend/Private/Interfaces/MetasoundFrontendInterface.cpp
+++ b/Engine/Plugins/Runtime/Metasound/Source/MetasoundFrontend/Private/Interfaces/MetasoundFrontendInterface.cpp
@@ -163,7 +163,7 @@ namespace Metasound
 
 		bool IsEqualInterface(const FMetasoundFrontendInterface& InInputInterface, const FMetasoundFrontendInterface& InTargetInterface)
 		{
-			const bool bIsMetadataEqual = InInputInterface.Version == InTargetInterface.Version;
+			const bool bIsMetadataEqual = InInputInterface.Metadata.Version == InTargetInterface.Metadata.Version;
 			if (bIsMetadataEqual)
 			{
 				// If metadata is equal, then inputs/outputs/environment variables should also be equal. 
diff --git a/Engine/Plugins/Runtime/Metasound/Source/MetasoundFrontend/Private/Interfaces/MetasoundFrontendInterfaceRegistry.cpp b/Engine/Plugins/Runtime/Metasound/Source/MetasoundFrontend/Private/Interfaces/MetasoundFrontendInterfaceRegistry.cpp
--- a/Engine/Plugins/Runtime/Metasound/Source/MetasoundFrontend/Private/Interfaces/MetasoundFrontendInterfaceRegistry.cpp
+++ b/Engine/Plugins/Runtime/Metasound/Source/MetasoundFrontend/Private/Interfaces/MetasoundFrontendInterfaceRegistry.cpp
@@ -34,7 +34,7 @@ namespace Metasound::Frontend
 				{
 					UE_LOG(LogMetaSound, Warning, TEXT("Registration of interface overwriting previously registered interface [RegistryKey: %s]"), *Key);
 						
-					FInterfaceRegistryTransaction Transaction{FInterfaceRegistryTransaction::ETransactionType::InterfaceUnregistration, Key, Entry->GetInterface().Version, TransactionTime};
+					FInterfaceRegistryTransaction Transaction{FInterfaceRegistryTransaction::ETransactionType::InterfaceUnregistration, Key, Entry->GetInterface().Metadata.Version, TransactionTime};
 					TransactionBuffer->AddTransaction(MoveTemp(Transaction));
 				}
 
@@ -67,7 +67,7 @@ namespace Metasound::Frontend
 				}
 #endif // UE_METASOUND_ENABLE_INTERFACE_VALIDATION
 
-				FInterfaceRegistryTransaction Transaction{FInterfaceRegistryTransaction::ETransactionType::InterfaceRegistration, Key, InEntry->GetInterface().Version, TransactionTime};
+				FInterfaceRegistryTransaction Transaction{FInterfaceRegistryTransaction::ETransactionType::InterfaceRegistration, Key, InEntry->GetInterface().Metadata.Version, TransactionTime};
 				TransactionBuffer->AddTransaction(MoveTemp(Transaction));
 				Entries.Add(Key, MoveTemp(InEntry));
 				return true;
@@ -115,7 +115,7 @@ namespace Metasound::Frontend
 
 	FInterfaceRegistryKey GetInterfaceRegistryKey(const FMetasoundFrontendInterface& InInterface)
 	{
-		return GetInterfaceRegistryKey(InInterface.Version);
+		return GetInterfaceRegistryKey(InInterface.Metadata.Version);
 	}
 
 	FInterfaceRegistryTransaction::FInterfaceRegistryTransaction(ETransactionType InType, const FInterfaceRegistryKey& InKey, const FMetasoundFrontendVersion& InInterfaceVersion, FInterfaceRegistryTransaction::FTimeType InTimestamp)
diff --git a/Engine/Plugins/Runtime/Metasound/Source/MetasoundFrontend/Private/MetasoundAssetBase.cpp b/Engine/Plugins/Runtime/Metasound/Source/MetasoundFrontend/Private/MetasoundAssetBase.cpp
--- a/Engine/Plugins/Runtime/Metasound/Source/MetasoundFrontend/Private/MetasoundAssetBase.cpp
+++ b/Engine/Plugins/Runtime/Metasound/Source/MetasoundFrontend/Private/MetasoundAssetBase.cpp
@@ -163,7 +163,7 @@ namespace Metasound
 					{
 						if (Entry->UpdateRootGraphInterface(InDocument))
 						{
-							LastVersionUpdated = &Entry->GetInterface().Version.Number;
+							LastVersionUpdated = &Entry->GetInterface().Metadata.Version.Number;
 						}
 					}
 				}
@@ -966,12 +966,12 @@ bool FMetasoundAssetBase::TryUpdateInterfaceFromVersion(const FMetasoundFrontend
 	using namespace AssetBasePrivate;
 
 	FMetasoundFrontendInterface TargetInterface = GetInterfaceToVersion(Version);
-	if (TargetInterface.Version.IsValid())
+	if (TargetInterface.Metadata.Version.IsValid())
 	{
 		TArray<const IInterfaceRegistryEntry*> UpgradePath;
-		GetUpdatePathForDocument(Version, TargetInterface.Version, UpgradePath);
+		GetUpdatePathForDocument(Version, TargetInterface.Metadata.Version, UpgradePath);
 		const bool bUpdated = UpdateDocumentInterface(UpgradePath, Version, GetDocumentHandle());
-		ensureMsgf(bUpdated, TEXT("Target interface '%s' was out-of-date but interface failed to be updated"), *TargetInterface.Version.ToString());
+		ensureMsgf(bUpdated, TEXT("Target interface '%s' was out-of-date but interface failed to be updated"), *TargetInterface.Metadata.Version.ToString());
 		return bUpdated;
 	}
 
@@ -1021,7 +1021,7 @@ FMetasoundFrontendInterface FMetasoundAssetBase::GetInterfaceToVersion(const FMe
 		return { };
 	}
 
-	if (TargetInterface.Version == InterfaceVersion)
+	if (TargetInterface.Metadata.Version == InterfaceVersion)
 	{
 		return { };
 	}
diff --git a/Engine/Plugins/Runtime/Metasound/Source/MetasoundFrontend/Private/MetasoundAssetManager.cpp b/Engine/Plugins/Runtime/Metasound/Source/MetasoundFrontend/Private/MetasoundAssetManager.cpp
--- a/Engine/Plugins/Runtime/Metasound/Source/MetasoundFrontend/Private/MetasoundAssetManager.cpp
+++ b/Engine/Plugins/Runtime/Metasound/Source/MetasoundFrontend/Private/MetasoundAssetManager.cpp
@@ -1,7 +1,11 @@
 // Copyright Epic Games, Inc. All Rights Reserved.
 
 #include "MetasoundAssetManager.h"
-#include "MetasoundAssetReferenceTagCollection.h"
+
+#include "Algo/Copy.h"
+#include "AssetRegistry/AssetData.h"
+#include "MetasoundAssetTagCollections.h"
+#include "MetasoundFrontendDocumentVersioning.h"
 #include "MetasoundFrontendRegistryKey.h"
 
 
@@ -37,187 +41,242 @@ PRAGMA_ENABLE_DEPRECATION_WARNINGS
 	namespace AssetTagsPrivate
 	{
 		static const FLazyName AssetClassID = "AssetClassID";
-		static const FLazyName ClassVersion = "ClassVersion";
 
 #if WITH_EDITORONLY_DATA
-		static const FLazyName ClassReferences = "ClassReferences";
-		static const FLazyName DisplayName = "DisplayName";
+		static const FLazyName AssetCollections = "AssetCollections";
 		static const FLazyName DocumentVersion = "DocumentVersion";
-		static const FLazyName Inputs = "Inputs";
 		static const FLazyName IsPreset = "bIsPreset";
-		static const FLazyName Outputs = "Outputs";
-		static const FLazyName DisplayInfo = "DisplayInfo";
 #endif // WITH_EDITORONLY_DATA
+
+		template <typename TStructType>
+		bool DeserializeTagFromJson(const FAssetData& InAssetData, FName TagName, TStructType& OutStruct)
+		{
+			FString TagString;
+			const bool bTagDataFound = InAssetData.GetTagValue(TagName, TagString);
+			if (bTagDataFound)
+			{
+				return FJsonObjectConverter::JsonObjectStringToUStruct(TagString, &OutStruct);
+			}
+			return false;
+		};
 	}
+} // namespace Metasound::Frontend
 
-	FAssetClassQueryInfo::FAssetClassQueryInfo()
-		: bIsValid(0)
-	{
+
+FMetaSoundDocumentInfo::FMetaSoundDocumentInfo()
+	: bIsPreset(0)
+{
+}
+
+FMetaSoundDocumentInfo::FMetaSoundDocumentInfo(const IMetaSoundDocumentInterface& InDocInterface)
+{
 #if WITH_EDITORONLY_DATA
-		bIsPreset = 0;
-#endif // WITH_EDITORONLY_DATA
-	}
+	using namespace Metasound::Frontend;
 
-	FAssetClassQueryInfo::FAssetClassQueryInfo(const IMetaSoundDocumentInterface& InDocInterface)
-	{
-		InitFromDocument(InDocInterface);
-	}
+	const FMetasoundFrontendDocument& Document = InDocInterface.GetConstDocument();
+	const FMetasoundFrontendClassMetadata& ClassMetadata = Document.RootGraph.Metadata;
 
-	FAssetClassQueryInfo::FAssetClassQueryInfo(const FAssetData& InAssetData)
+	if (!IsRunningCookCommandlet())
 	{
-		using namespace Frontend;
+		bIsPreset = (uint8)Document.RootGraph.PresetOptions.bIsPreset;
+		DocumentVersion = Document.Metadata.Version.Number;
 
-		if (InAssetData.IsAssetLoaded())
 		{
-			UObject* MetaSound = InAssetData.GetAsset();
-			TScriptInterface<IMetaSoundDocumentInterface> DocInterface(MetaSound);
+			IMetaSoundAssetManager& AssetManager = IMetaSoundAssetManager::GetChecked();
+			auto IsAssetReference = [&AssetManager](const FMetasoundFrontendClass& Class)
+				{
+					return AssetManager.IsAssetClass(Class.Metadata);
+				};
 
-			check(DocInterface.GetObject());
-			InitFromDocument(*DocInterface.GetInterface());
-			return;
+			auto GetAssetKey = [](const FMetasoundFrontendClass& Class)
+				{
+					return FMetaSoundAssetKey(Class.Metadata);
+				};
+			Algo::TransformIf(Document.Dependencies, ReferencedAssetKeys, IsAssetReference, GetAssetKey);
 		}
+	}
 
-		bool bSuccess = true;
+#endif // WITH_EDITORONLY_DATA
+}
+
+FMetaSoundDocumentInfo::FMetaSoundDocumentInfo(const FAssetData& InAssetData, bool& bOutIsValid)
+{
+	bOutIsValid = true;
 
 #if WITH_EDITORONLY_DATA
+	using namespace Metasound::Frontend;
+
+	if (!IsRunningCookCommandlet())
+	{
 		{
 			FString DocVersionStr;
-			bSuccess &= InAssetData.GetTagValue(AssetTagsPrivate::DocumentVersion, DocVersionStr);
-			bSuccess &= FMetasoundFrontendVersionNumber::Parse(DocVersionStr, DocumentVersion);
+			bOutIsValid &= InAssetData.GetTagValue(AssetTagsPrivate::DocumentVersion, DocVersionStr);
+			bOutIsValid &= FMetasoundFrontendVersionNumber::Parse(DocVersionStr, DocumentVersion);
 		}
 
 		{
 			bool bTagIsPreset = false;
-			bSuccess &= InAssetData.GetTagValue(AssetTagsPrivate::IsPreset, bTagIsPreset);
+			bOutIsValid &= InAssetData.GetTagValue(AssetTagsPrivate::IsPreset, bTagIsPreset);
 			bIsPreset = (uint8)bTagIsPreset;
 		}
 
 		{
-			auto DeserializeTagFromJson = [&bSuccess, &InAssetData]<typename TStructType>(FName TagName, TStructType & OutStruct)
-			{
-				FString TagString;
-				const bool bTagDataFound = InAssetData.GetTagValue(TagName, TagString);
-				if (bTagDataFound)
-				{
-					return FJsonObjectConverter::JsonObjectStringToUStruct(TagString, &OutStruct);
-				}
-				return false;
-			};
-
-
 			{
-				FMetaSoundAssetReferenceTagCollection ReferenceTagCollection;
-				const bool bTagFound = DeserializeTagFromJson(AssetTagsPrivate::ClassReferences, ReferenceTagCollection);
+				FMetaSoundAssetTagCollections TagCollections;
+				const bool bTagFound = AssetTagsPrivate::DeserializeTagFromJson(InAssetData, AssetTagsPrivate::AssetCollections, TagCollections);
 				if (bTagFound)
 				{
-					ReferencedAssetKeys = MoveTemp(ReferenceTagCollection.AssetKeys);
+					ReferencedAssetKeys = MoveTemp(TagCollections.AssetKeys);
 				}
 				else
 				{
-					bSuccess = false;
+					bOutIsValid = false;
 				}
 			}
-
-			bSuccess &= DeserializeTagFromJson(AssetTagsPrivate::DisplayInfo, DisplayInfo);
-			bSuccess &= DeserializeTagFromJson(AssetTagsPrivate::Inputs, Inputs);
-			bSuccess &= DeserializeTagFromJson(AssetTagsPrivate::Outputs, Outputs);
 		}
+	}
+#endif // WITH_EDITORONLY_DATA
+}
+
+void FMetaSoundDocumentInfo::ExportToContext(FAssetRegistryTagsContext& OutContext) const
+{
+	using namespace Metasound::Frontend;
+
+	using FAssetRegistryTag = UObject::FAssetRegistryTag;
+
+#if WITH_EDITORONLY_DATA
+	OutContext.AddTag(FAssetRegistryTag(AssetTagsPrivate::DocumentVersion, DocumentVersion.ToString(), FAssetRegistryTag::TT_Alphabetical));
+	OutContext.AddTag(FAssetRegistryTag(AssetTagsPrivate::IsPreset, FString(bIsPreset ? TEXT("1") : TEXT("0")), FAssetRegistryTag::TT_Numerical));
+
+	auto SerializeTagToJson = []<typename TStructType>(const TStructType& InStruct)
+	{
+		FString JsonString;
+		FJsonObjectConverter::UStructToJsonObjectString(InStruct, JsonString);
+		return JsonString;
+	};
+
+	const FMetaSoundAssetTagCollections TagCollections { .AssetKeys = ReferencedAssetKeys };
+	OutContext.AddTag(FAssetRegistryTag(AssetTagsPrivate::AssetCollections, SerializeTagToJson(TagCollections), FAssetRegistryTag::TT_Hidden));
 #endif // WITH_EDITORONLY_DATA
+}
 
+namespace Metasound::Frontend
+{
+	FMetaSoundAssetClassInfo::FMetaSoundAssetClassInfo(const IMetaSoundDocumentInterface& InDocInterface)
+	{
+		InitFromDocument(InDocInterface);
+	}
+
+	FMetaSoundAssetClassInfo::FMetaSoundAssetClassInfo(const FAssetData& InAssetData)
+		: FMetaSoundClassInfo(InAssetData)
+	{
+		using namespace Frontend;
+
+		if (InAssetData.IsAssetLoaded())
 		{
-			FString ClassVersionStr;
-			bSuccess &= InAssetData.GetTagValue(AssetTagsPrivate::ClassVersion, ClassVersionStr);
-			bSuccess &= FMetasoundFrontendVersionNumber::Parse(ClassVersionStr, Version);
+			UObject* MetaSound = InAssetData.GetAsset();
+			TScriptInterface<IMetaSoundDocumentInterface> DocInterface(MetaSound);
+
+			check(DocInterface.GetObject());
+			InitFromDocument(*DocInterface.GetInterface());
+			return;
 		}
 
+		FString AssetClassID;
+		bool bSuccess = TryGetAssetClassTag(InAssetData, AssetClassID);
 		{
-			FString AssetClassID;
-			bSuccess &= InAssetData.GetTagValue(AssetTagsPrivate::AssetClassID, AssetClassID);
 			ClassName = FMetasoundFrontendClassName(FName(), FName(*AssetClassID));
 		}
 
-		bIsValid = (int8)(bSuccess);
+#if WITH_EDITORONLY_DATA
+		DocInfo = FMetaSoundDocumentInfo(InAssetData, bSuccess);
+#endif // WITH_EDITORONLY_DATA
+
 		AssetPath = FTopLevelAssetPath(InAssetData.PackageName, InAssetData.AssetName);
+
+		bIsValid = bIsValid & (int8)(bSuccess);
 	}
 
-	void FAssetClassQueryInfo::ExportToContext(FAssetRegistryTagsContext& OutContext) const
+	void FMetaSoundAssetClassInfo::ExportToContext(FAssetRegistryTagsContext& OutContext) const
 	{
 		using FAssetRegistryTag = UObject::FAssetRegistryTag;
 
+		FMetaSoundClassInfo::ExportToContext(OutContext);
+
+		// AssetClassID housed in ClassName, but ID and associated tag is
+		// stored on this inheriting class, so this is serialized locally.
 		OutContext.AddTag(FAssetRegistryTag(AssetTagsPrivate::AssetClassID, ClassName.Name.ToString(), FAssetRegistryTag::TT_Alphabetical));
-		OutContext.AddTag(FAssetRegistryTag(AssetTagsPrivate::ClassVersion, Version.ToString(), FAssetRegistryTag::TT_Alphabetical));
 
 #if WITH_EDITORONLY_DATA
-		auto SerializeTagToJson = []<typename TStructType>(const TStructType & InStruct)
-		{
-			FString JsonString;
-			FJsonObjectConverter::UStructToJsonObjectString(InStruct, JsonString);
-			return JsonString;
-		};
-
-		OutContext.AddTag(FAssetRegistryTag(AssetTagsPrivate::DocumentVersion, DocumentVersion.ToString(), FAssetRegistryTag::TT_Alphabetical));
-		OutContext.AddTag(FAssetRegistryTag(AssetTagsPrivate::IsPreset, FString(bIsPreset ? TEXT("1") : TEXT("0")), FAssetRegistryTag::TT_Numerical));
-
-		const FMetaSoundAssetReferenceTagCollection AssetReferences{ ReferencedAssetKeys };
-		OutContext.AddTag(FAssetRegistryTag(AssetTagsPrivate::ClassReferences, SerializeTagToJson(AssetReferences), FAssetRegistryTag::TT_Hidden));
-		OutContext.AddTag(FAssetRegistryTag(AssetTagsPrivate::DisplayInfo, SerializeTagToJson(DisplayInfo), FAssetRegistryTag::TT_Hidden));
-		OutContext.AddTag(FAssetRegistryTag(AssetTagsPrivate::Inputs, SerializeTagToJson(Inputs), FAssetRegistryTag::TT_Hidden));
-		OutContext.AddTag(FAssetRegistryTag(AssetTagsPrivate::Outputs, SerializeTagToJson(Outputs), FAssetRegistryTag::TT_Hidden));
+		DocInfo.ExportToContext(OutContext);
 #endif // WITH_EDITORONLY_DATA
 	}
 
-	void FAssetClassQueryInfo::InitFromDocument(const IMetaSoundDocumentInterface& InDocInterface)
+	void FMetaSoundAssetClassInfo::InitFromDocument(const IMetaSoundDocumentInterface& InDocInterface)
 	{
-		using namespace Frontend;
-
-		const FMetasoundFrontendDocument& Document = InDocInterface.GetConstDocument();
-		const FMetasoundFrontendClassMetadata& ClassMetadata = Document.RootGraph.Metadata;
+		FMetaSoundClassInfo::InitFromDocument(InDocInterface);
 
-		ClassName = Document.RootGraph.Metadata.GetClassName();
-		Version = ClassMetadata.GetVersion();
 		AssetPath = InDocInterface.GetAssetPathChecked();
 
 #if WITH_EDITORONLY_DATA
-		bIsPreset = (uint8)Document.RootGraph.PresetOptions.bIsPreset;
-		DocumentVersion = Document.Metadata.Version.Number;
-
-		DisplayInfo.ClassDescription = ClassMetadata.GetDescription();
-		DisplayInfo.ClassDisplayName = ClassMetadata.GetDisplayName();
-		DisplayInfo.Hierarchy = ClassMetadata.GetCategoryHierarchy();
-		DisplayInfo.Keywords = ClassMetadata.GetKeywords();
-
-		const FMetasoundFrontendGraphClass& GraphClass = Document.RootGraph;
+		DocInfo = FMetaSoundDocumentInfo(InDocInterface);
+#endif // WITH_EDITORONLY_DATA
+	}
 
+	bool FMetaSoundAssetClassInfo::TryGetAssetKey(const FAssetData& InAssetData, FMetaSoundAssetKey& OutKey)
+	{
+		if (InAssetData.IsAssetLoaded())
 		{
-			IMetaSoundAssetManager& AssetManager = IMetaSoundAssetManager::GetChecked();
-			auto IsAssetReference = [&AssetManager](const FMetasoundFrontendClass& Class)
+			if (const UObject* Object = InAssetData.GetAsset())
 			{
-				return AssetManager.IsAssetClass(Class.Metadata);
-			};
-			
-			auto GetAssetKey = [](const FMetasoundFrontendClass& Class)
-			{
-				return FMetaSoundAssetKey(Class.Metadata);
-			};
-			Algo::TransformIf(Document.Dependencies, ReferencedAssetKeys, IsAssetReference, GetAssetKey);
+				TScriptInterface<const IMetaSoundDocumentInterface> DocInterface(Object);
+				if (DocInterface.GetObject())
+				{
+					OutKey = FMetaSoundAssetKey(DocInterface->GetConstDocument().RootGraph);
+					return true;
+				}
+			}
+		}
+
+		bool bSuccess = true;
+		{
+			FString AssetClassID;
+			bSuccess &= InAssetData.GetTagValue(AssetTagsPrivate::AssetClassID, AssetClassID);
+			OutKey.ClassName = FMetasoundFrontendClassName(FName(), FName(*AssetClassID));
 		}
 
+		bSuccess = TryGetClassVersion(InAssetData, OutKey.Version);
+		return bSuccess;
+	}
+
+	bool FMetaSoundAssetClassInfo::TryGetAssetClassName(const FAssetData& InAssetData, FMetasoundFrontendClassName& OutClassName)
+	{
+		if (InAssetData.IsAssetLoaded())
 		{
-			auto CopyVertexInfo = [](const FMetasoundFrontendClassVertex& Vertex)
+			if (const UObject* Object = InAssetData.GetAsset())
 			{
-				return FClassQueryVertexInfo
+				TScriptInterface<const IMetaSoundDocumentInterface> DocInterface(Object);
+				if (DocInterface.GetObject())
 				{
-					.Name = Vertex.Name,
-					.TypeName = Vertex.TypeName,
-					.AccessType = Vertex.AccessType
-				};
-			};
-			Algo::Transform(GraphClass.Interface.Inputs, Inputs.ClassVertexInfo, CopyVertexInfo);
-			Algo::Transform(GraphClass.Interface.Outputs, Outputs.ClassVertexInfo, CopyVertexInfo);
+					OutClassName = DocInterface->GetConstDocument().RootGraph.Metadata.GetClassName();
+					return true;
+				}
+			}
 		}
-#endif // WITH_EDITORONLY_DATA
 
-		bIsValid = 1;
+		bool bSuccess = true;
+		{
+			FString AssetClassID;
+			bSuccess &= TryGetAssetClassTag(InAssetData, AssetClassID);
+			OutClassName = FMetasoundFrontendClassName(FName(), FName(*AssetClassID));
+		}
+
+		return bSuccess;
+	}
+
+	bool FMetaSoundAssetClassInfo::TryGetAssetClassTag(const FAssetData& InAssetData, FString& OutClassIDString)
+	{
+		return InAssetData.GetTagValue(AssetTagsPrivate::AssetClassID, OutClassIDString);
 	}
 
 	IMetaSoundAssetManager* IMetaSoundAssetManager::Get()
diff --git a/Engine/Plugins/Runtime/Metasound/Source/MetasoundFrontend/Private/MetasoundAssetReferenceTagCollection.h b/Engine/Plugins/Runtime/Metasound/Source/MetasoundFrontend/Private/MetasoundAssetReferenceTagCollection.h
deleted file mode 100644
diff --git a/Engine/Plugins/Runtime/Metasound/Source/MetasoundFrontend/Private/MetasoundAssetTagCollections.h b/Engine/Plugins/Runtime/Metasound/Source/MetasoundFrontend/Private/MetasoundAssetTagCollections.h
new file mode 100644
--- /dev/null
+++ b/Engine/Plugins/Runtime/Metasound/Source/MetasoundFrontend/Private/MetasoundAssetTagCollections.h
@@ -0,0 +1,38 @@
+// Copyright Epic Games, Inc. All Rights Reserved.
+#pragma once
+
+#include "MetasoundAssetKey.h"
+#include "UObject/Object.h"
+
+#include "MetasoundAssetTagCollections.generated.h"
+
+
+// Minimal Json-serializable arrays of asset tag data. Helper struct for asset
+// tag serialization of collection properties.
+USTRUCT()
+struct FMetaSoundAssetTagCollections
+{
+	GENERATED_BODY()
+
+#if WITH_EDITORONLY_DATA
+	UPROPERTY()
+	TArray<FMetaSoundAssetKey> AssetKeys;
+#endif // WITH_EDITORONLY_DATA
+};
+
+
+// Minimal Json-serializable arrays of asset class tag data. Helper struct for asset
+// tag serialization of collection properties.
+USTRUCT()
+struct FMetaSoundAssetTagClassCollections
+{
+	GENERATED_BODY()
+
+#if WITH_EDITORONLY_DATA
+	UPROPERTY()
+	TArray<FMetasoundFrontendInterfaceMetadata> DefinedInterfaces;
+
+	UPROPERTY()
+	TArray<FMetasoundFrontendVersion> InheritedInterfaces;
+#endif // WITH_EDITORONLY_DATA
+};
diff --git a/Engine/Plugins/Runtime/Metasound/Source/MetasoundFrontend/Private/MetasoundFrontendDocument.cpp b/Engine/Plugins/Runtime/Metasound/Source/MetasoundFrontend/Private/MetasoundFrontendDocument.cpp
--- a/Engine/Plugins/Runtime/Metasound/Source/MetasoundFrontend/Private/MetasoundFrontendDocument.cpp
+++ b/Engine/Plugins/Runtime/Metasound/Source/MetasoundFrontend/Private/MetasoundFrontendDocument.cpp
@@ -427,7 +427,7 @@ FMetasoundFrontendInterface::FMetasoundFrontendInterface(Audio::FParameterInterf
 	using namespace Metasound::DocumentPrivate;
 	using namespace Metasound::Frontend;
 
-	Version = { InInterface->GetName(), FMetasoundFrontendVersionNumber { InInterface->GetVersion().Major, InInterface->GetVersion().Minor } };
+	Metadata.Version = { InInterface->GetName(), FMetasoundFrontendVersionNumber { InInterface->GetVersion().Major, InInterface->GetVersion().Minor } };
 
 	// Transfer all input data from AudioExtension interface struct to FrontendInterface
 	Algo::Transform(InInterface->GetInputs(), Inputs, [this](const Audio::FParameterInterface::FInput& Input)
@@ -469,17 +469,19 @@ FMetasoundFrontendInterface::FMetasoundFrontendInterface(Audio::FParameterInterf
 	});
 
 	// Transfer all class options from AudioExtension interface struct to FrontendInterface
-	Algo::Transform(InInterface->GetUClassOptions(), UClassOptions, [](const Audio::FParameterInterface::FClassOptions& Options)
+	Algo::Transform(InInterface->GetUClassOptions(), Metadata.UClassOptions, [](const Audio::FParameterInterface::FClassOptions& Options)
 	{
 		return FMetasoundFrontendInterfaceUClassOptions(Options);
 	});
 }
 
+#if WITH_EDITORONLY_DATA
 const FMetasoundFrontendInterfaceUClassOptions* FMetasoundFrontendInterface::FindClassOptions(const FTopLevelAssetPath& InClassPath) const
 {
 	auto FindClassOptionsPredicate = [&InClassPath](const FMetasoundFrontendInterfaceUClassOptions& Options) { return Options.ClassPath == InClassPath; };
 	return UClassOptions.FindByPredicate(FindClassOptionsPredicate);
 }
+#endif // WITH_EDITORONLY_DATA
 
 const FMetasoundFrontendClassName FMetasoundFrontendClassName::InvalidClassName;
 
diff --git a/Engine/Plugins/Runtime/Metasound/Source/MetasoundFrontend/Private/MetasoundFrontendDocumentBuilder.cpp b/Engine/Plugins/Runtime/Metasound/Source/MetasoundFrontend/Private/MetasoundFrontendDocumentBuilder.cpp
--- a/Engine/Plugins/Runtime/Metasound/Source/MetasoundFrontend/Private/MetasoundFrontendDocumentBuilder.cpp
+++ b/Engine/Plugins/Runtime/Metasound/Source/MetasoundFrontend/Private/MetasoundFrontendDocumentBuilder.cpp
@@ -552,27 +552,27 @@ namespace Metasound::Frontend
 
 				for (const FMetasoundFrontendInterface& Interface : Options.InterfacesToRemove)
 				{
-					if (Document.Interfaces.Contains(Interface.Version))
+					if (Document.Interfaces.Contains(Interface.Metadata.Version))
 					{
 						OutDelegates.InterfaceDelegates.OnRemovingInterface.Broadcast(Interface);
 						bDidEdit = true;
 #if WITH_EDITORONLY_DATA
-						Document.Metadata.ModifyContext.AddInterfaceModified(Interface.Version.Name);
+						Document.Metadata.ModifyContext.AddInterfaceModified(Interface.Metadata.Version.Name);
 #endif // WITH_EDITORONLY_DATA
-						Document.Interfaces.Remove(Interface.Version);
+						Document.Interfaces.Remove(Interface.Metadata.Version);
 					}
 				}
 
 				for (const FMetasoundFrontendInterface& Interface : Options.InterfacesToAdd)
 				{
 					bool bAlreadyInSet = false;
-					Document.Interfaces.Add(Interface.Version, &bAlreadyInSet);
+					Document.Interfaces.Add(Interface.Metadata.Version, &bAlreadyInSet);
 					if (!bAlreadyInSet)
 					{
 						OutDelegates.InterfaceDelegates.OnInterfaceAdded.Broadcast(Interface);
 						bDidEdit = true;
 #if WITH_EDITORONLY_DATA
-						Document.Metadata.ModifyContext.AddInterfaceModified(Interface.Version.Name);
+						Document.Metadata.ModifyContext.AddInterfaceModified(Interface.Metadata.Version.Name);
 #endif // WITH_EDITORONLY_DATA
 					}
 				}
@@ -1127,7 +1127,7 @@ const FMetasoundFrontendNode* FMetaSoundFrontendDocumentBuilder::AddGraphInput(F
 	FMetasoundFrontendInterface Interface;
 	if (DocumentBuilderPrivate::IsInterfaceInput(ClassInput.Name, ClassInput.TypeName, &Interface))
 	{
-		if (Document.Interfaces.Contains(Interface.Version))
+		if (Document.Interfaces.Contains(Interface.Metadata.Version))
 		{
 			FText RequiredText;
 			bIsRequired = Interface.IsMemberInputRequired(ClassInput.Name, RequiredText);
@@ -1656,17 +1656,18 @@ bool FMetaSoundFrontendDocumentBuilder::AddInterface(FName InterfaceName)
 	FMetasoundFrontendInterface Interface;
 	if (ISearchEngine::Get().FindInterfaceWithHighestVersion(InterfaceName, Interface))
 	{
-		if (GetDocumentChecked().Interfaces.Contains(Interface.Version))
+		if (GetDocumentChecked().Interfaces.Contains(Interface.Metadata.Version))
 		{
 			UE_LOG(LogMetaSound, VeryVerbose, TEXT("MetaSound interface '%s' already found on document. MetaSoundBuilder skipping add request."), *InterfaceName.ToString());
 			return true;
 		}
 
-		const FTopLevelAssetPath BuilderClassPath = GetBuilderClassPath();
-		const FInterfaceRegistryKey Key = GetInterfaceRegistryKey(Interface.Version);
+		const FInterfaceRegistryKey Key = GetInterfaceRegistryKey(Interface.Metadata.Version);
 		if (const IInterfaceRegistryEntry* Entry = IInterfaceRegistry::Get().FindInterfaceRegistryEntry(Key))
 		{
-			const FMetasoundFrontendInterfaceUClassOptions* ClassOptions = Entry->GetInterface().FindClassOptions(BuilderClassPath);
+			const FTopLevelAssetPath BuilderClassPath = GetBuilderClassPath();
+			auto FindClassOptionsPredicate = [&BuilderClassPath](const FMetasoundFrontendInterfaceUClassOptions& Options) { return Options.ClassPath == BuilderClassPath; };
+			const FMetasoundFrontendInterfaceUClassOptions* ClassOptions = Entry->GetInterface().Metadata.UClassOptions.FindByPredicate(FindClassOptionsPredicate);
 			if (ClassOptions && !ClassOptions->bIsModifiable)
 			{
 				UE_LOG(LogMetaSound, Error, TEXT("DocumentBuilder failed to add MetaSound Interface '%s' to document: is not set to be modifiable for given UClass '%s'"), *InterfaceName.ToString(), *BuilderClassPath.ToString());
@@ -2435,7 +2436,7 @@ bool FMetaSoundFrontendDocumentBuilder::FindInterfaceInputNodes(FName InterfaceN
 	const TSet<FMetasoundFrontendVersion>& Interfaces = GetConstDocumentChecked().Interfaces;
 	if (ISearchEngine::Get().FindInterfaceWithHighestVersion(InterfaceName, Interface))
 	{
-		if (Interfaces.Contains(Interface.Version))
+		if (Interfaces.Contains(Interface.Metadata.Version))
 		{
 			const IDocumentGraphNodeCache& NodeCache = DocumentCache->GetNodeCache(PageID);
 			const IDocumentGraphInterfaceCache& InterfaceCache = DocumentCache->GetInterfaceCache();
@@ -2479,7 +2480,7 @@ bool FMetaSoundFrontendDocumentBuilder::FindInterfaceOutputNodes(FName Interface
 	const TSet<FMetasoundFrontendVersion>& Interfaces = GetConstDocumentChecked().Interfaces;
 	if (ISearchEngine::Get().FindInterfaceWithHighestVersion(InterfaceName, Interface))
 	{
-		if (Interfaces.Contains(Interface.Version))
+		if (Interfaces.Contains(Interface.Metadata.Version))
 		{
 			const IDocumentGraphNodeCache& NodeCache = DocumentCache->GetNodeCache(PageID);
 			const IDocumentGraphInterfaceCache& InterfaceCache = DocumentCache->GetInterfaceCache();
@@ -3424,7 +3425,7 @@ bool FMetaSoundFrontendDocumentBuilder::IsInterfaceDeclared(FName InInterfaceNam
 	FMetasoundFrontendInterface Interface;
 	if (ISearchEngine::Get().FindInterfaceWithHighestVersion(InInterfaceName, Interface))
 	{
-		return IsInterfaceDeclared(Interface.Version);
+		return IsInterfaceDeclared(Interface.Metadata.Version);
 	}
 
 	return false;
@@ -4299,17 +4300,18 @@ bool FMetaSoundFrontendDocumentBuilder::RemoveInterface(FName InterfaceName)
 	FMetasoundFrontendInterface Interface;
 	if (ISearchEngine::Get().FindInterfaceWithHighestVersion(InterfaceName, Interface))
 	{
-		if (!GetDocumentChecked().Interfaces.Contains(Interface.Version))
+		if (!GetDocumentChecked().Interfaces.Contains(Interface.Metadata.Version))
 		{
 			UE_LOG(LogMetaSound, VeryVerbose, TEXT("MetaSound interface '%s' not found on document. MetaSoundBuilder skipping remove request."), *InterfaceName.ToString());
 			return true;
 		}
 
-		const FTopLevelAssetPath BuilderClassPath = GetBuilderClassPath();
-		const FInterfaceRegistryKey Key = GetInterfaceRegistryKey(Interface.Version);
+		const FInterfaceRegistryKey Key = GetInterfaceRegistryKey(Interface.Metadata.Version);
 		if (const IInterfaceRegistryEntry* Entry = IInterfaceRegistry::Get().FindInterfaceRegistryEntry(Key))
 		{
-			const FMetasoundFrontendInterfaceUClassOptions* ClassOptions = Entry->GetInterface().FindClassOptions(BuilderClassPath);
+			const FTopLevelAssetPath BuilderClassPath = GetBuilderClassPath();
+			auto FindClassOptionsPredicate = [&BuilderClassPath](const FMetasoundFrontendInterfaceUClassOptions& Options) { return Options.ClassPath == BuilderClassPath; };
+			const FMetasoundFrontendInterfaceUClassOptions* ClassOptions = Entry->GetInterface().Metadata.UClassOptions.FindByPredicate(FindClassOptionsPredicate);
 			if (ClassOptions && !ClassOptions->bIsModifiable)
 			{
 				UE_LOG(LogMetaSound, Error, TEXT("DocumentBuilder failed to remove MetaSound Interface '%s' to document: is not set to be modifiable for given UClass '%s'"), *InterfaceName.ToString(), *BuilderClassPath.ToString());
diff --git a/Engine/Plugins/Runtime/Metasound/Source/MetasoundFrontend/Private/MetasoundFrontendDocumentVersioning.cpp b/Engine/Plugins/Runtime/Metasound/Source/MetasoundFrontend/Private/MetasoundFrontendDocumentVersioning.cpp
--- a/Engine/Plugins/Runtime/Metasound/Source/MetasoundFrontend/Private/MetasoundFrontendDocumentVersioning.cpp
+++ b/Engine/Plugins/Runtime/Metasound/Source/MetasoundFrontend/Private/MetasoundFrontendDocumentVersioning.cpp
@@ -340,9 +340,9 @@ PRAGMA_ENABLE_DEPRECATION_WARNINGS
 					if (const FMetasoundFrontendInterface* Interface = FindMostSimilarInterfaceSupportingEnvironment(RootGraph, Dependencies, Subgraphs, CandidateInterfaces))
 					{
 						METASOUND_VERSIONING_LOG(Display, TEXT("Assigned interface [InterfaceVersion:%s] to document [RootGraphClassName:%s]"),
-							*Interface->Version.ToString(), *RootGraph.Metadata.GetClassName().ToString());
+							*Interface->Metadata.Version.ToString(), *RootGraph.Metadata.GetClassName().ToString());
 
-						InDocument->AddInterfaceVersion(Interface->Version);
+						InDocument->AddInterfaceVersion(Interface->Metadata.Version);
 					}
 					else
 					{
diff --git a/Engine/Plugins/Runtime/Metasound/Source/MetasoundFrontend/Private/MetasoundFrontendQuery.cpp b/Engine/Plugins/Runtime/Metasound/Source/MetasoundFrontend/Private/MetasoundFrontendQuery.cpp
--- a/Engine/Plugins/Runtime/Metasound/Source/MetasoundFrontend/Private/MetasoundFrontendQuery.cpp
+++ b/Engine/Plugins/Runtime/Metasound/Source/MetasoundFrontend/Private/MetasoundFrontendQuery.cpp
@@ -3,14 +3,233 @@
 #include "MetasoundFrontendQuery.h"
 
 #include "Algo/AnyOf.h"
+#include "Algo/Copy.h"
+#include "AssetRegistry/AssetData.h"
+#include "JsonObjectConverter.h"
+#include "MetasoundAssetManager.h"
+#include "MetasoundAssetTagCollections.h"
 #include "MetasoundFrontendDocument.h"
 #include "MetasoundTrace.h"
 #include "Misc/Guid.h"
 #include "Templates/TypeHash.h"
 #include "Traits/IsContiguousContainer.h"
+#include "UObject/AssetRegistryTagsContext.h"
+
+
+namespace Metasound::Frontend
+{
+	namespace ClassTagsPrivate
+	{
+		static const FLazyName ClassVersion = "ClassVersion";
+
+#if WITH_EDITORONLY_DATA
+		static const FLazyName ClassCollections = "ClassCollections";
+		static const FLazyName SearchInfo = "SearchInfo";
+		static const FLazyName Inputs = "Inputs";
+		static const FLazyName Outputs = "Outputs";
+#endif // WITH_EDITORONLY_DATA
+
+		template <typename TStructType>
+		bool DeserializeTagFromJson(const FAssetData& InAssetData, FName TagName, TStructType& OutStruct)
+		{
+			FString TagString;
+			const bool bTagDataFound = InAssetData.GetTagValue(TagName, TagString);
+			if (bTagDataFound)
+			{
+				return FJsonObjectConverter::JsonObjectStringToUStruct(TagString, &OutStruct);
+			}
+			return false;
+		};
+	}
+}
+
+
+FMetaSoundClassSearchInfo::FMetaSoundClassSearchInfo(const FMetasoundFrontendClassMetadata& InClassMetadata)
+#if WITH_EDITOR
+	: ClassDisplayName(InClassMetadata.GetDisplayName())
+	, ClassDescription(InClassMetadata.GetDescription())
+	, Hierarchy(InClassMetadata.GetCategoryHierarchy())
+	, Keywords(InClassMetadata.GetKeywords())
+#endif // WITH_EDITOR
+{
+}
+
+FMetaSoundClassInterfaceInfo::FMetaSoundClassInterfaceInfo(const IMetaSoundDocumentInterface& InDocInterface)
+{
+#if WITH_EDITORONLY_DATA
+	using namespace Metasound::Frontend;
+
+	if (!IsRunningCookCommandlet())
+	{
+		const FMetasoundFrontendDocument& Document = InDocInterface.GetConstDocument();
+		SearchInfo = FMetaSoundClassSearchInfo(Document.RootGraph.Metadata);
+
+		Algo::Copy(Document.Interfaces, InheritedInterfaces);
+
+		{
+			auto CopyVertexInfo = [](const FMetasoundFrontendClassVertex& Vertex)
+			{
+				return FMetaSoundClassVertexInfo
+				{
+					.Name = Vertex.Name,
+					.TypeName = Vertex.TypeName,
+					.AccessType = Vertex.AccessType
+				};
+			};
+
+			const FMetasoundFrontendClassInterface& ClassInterface = Document.RootGraph.Interface;
+			Algo::Transform(ClassInterface.Inputs, Inputs, CopyVertexInfo);
+			Algo::Transform(ClassInterface.Outputs, Outputs, CopyVertexInfo);
+		}
+	}
+#endif // WITH_EDITORONLY_DATA
+}
+
+FMetaSoundClassInterfaceInfo::FMetaSoundClassInterfaceInfo(const FAssetData& InAssetData, bool& bOutIsValid)
+{
+	bOutIsValid = true;
+
+#if WITH_EDITORONLY_DATA
+	using namespace Metasound::Frontend;
+
+	if (!IsRunningCookCommandlet())
+	{
+		{
+			FMetaSoundAssetTagClassCollections TagCollections;
+			const bool bTagFound = ClassTagsPrivate::DeserializeTagFromJson(InAssetData, ClassTagsPrivate::ClassCollections, TagCollections);
+			if (bTagFound)
+			{
+				DefinedInterfaces = MoveTemp(TagCollections.DefinedInterfaces);
+				InheritedInterfaces = MoveTemp(TagCollections.InheritedInterfaces);
+			}
+			else
+			{
+				bOutIsValid = false;
+			}
+		}
+
+		bOutIsValid &= ClassTagsPrivate::DeserializeTagFromJson(InAssetData, ClassTagsPrivate::SearchInfo, SearchInfo);
+
+		FMetaSoundClassVertexCollectionInfo VertexCollection;
+		bOutIsValid &= ClassTagsPrivate::DeserializeTagFromJson(InAssetData, ClassTagsPrivate::Inputs, VertexCollection);
+		Inputs = MoveTemp(VertexCollection.ClassVertexInfo);
+		VertexCollection.ClassVertexInfo.Reset();
+		bOutIsValid &= ClassTagsPrivate::DeserializeTagFromJson(InAssetData, ClassTagsPrivate::Outputs, VertexCollection);
+		Outputs = MoveTemp(VertexCollection.ClassVertexInfo);
+	}
+#endif // WITH_EDITORONLY_DATA
+}
+
+void FMetaSoundClassInterfaceInfo::ExportToContext(FAssetRegistryTagsContext& OutContext) const
+{
+#if WITH_EDITORONLY_DATA
+	using namespace Metasound::Frontend;
+	using FAssetRegistryTag = UObject::FAssetRegistryTag;
+
+	if (!OutContext.IsCooking())
+	{
+		auto SerializeTagToJson = []<typename TStructType>(const TStructType& InStruct)
+		{
+			FString JsonString;
+			FJsonObjectConverter::UStructToJsonObjectString(InStruct, JsonString);
+			return JsonString;
+		};
+
+		const FMetaSoundAssetTagClassCollections TagCollections { .DefinedInterfaces = DefinedInterfaces, .InheritedInterfaces = InheritedInterfaces };
+		OutContext.AddTag(FAssetRegistryTag(ClassTagsPrivate::ClassCollections, SerializeTagToJson(TagCollections), FAssetRegistryTag::TT_Hidden));
+		OutContext.AddTag(FAssetRegistryTag(ClassTagsPrivate::SearchInfo, SerializeTagToJson(SearchInfo), FAssetRegistryTag::TT_Hidden));
+
+		FMetaSoundClassVertexCollectionInfo VertexCollection { Inputs };
+		OutContext.AddTag(FAssetRegistryTag(ClassTagsPrivate::Inputs, SerializeTagToJson(VertexCollection), FAssetRegistryTag::TT_Hidden));
+
+		VertexCollection.ClassVertexInfo = Outputs;
+		OutContext.AddTag(FAssetRegistryTag(ClassTagsPrivate::Outputs, SerializeTagToJson(VertexCollection), FAssetRegistryTag::TT_Hidden));
+	}
+#endif // WITH_EDITORONLY_DATA
+}
+
 
 namespace Metasound
 {
+	namespace Frontend
+	{
+		FMetaSoundClassInfo::FMetaSoundClassInfo()
+			: bIsValid(0)
+		{
+		}
+
+		FMetaSoundClassInfo::FMetaSoundClassInfo(const FAssetData& InAssetData)
+			: bIsValid(0)
+		{
+			bool bSuccess = true;
+
+#if WITH_EDITORONLY_DATA
+			InterfaceInfo = FMetaSoundClassInterfaceInfo(InAssetData, bSuccess);
+#endif // WITH_EDITORONLY_DATA
+
+			bSuccess &= TryGetClassVersion(InAssetData, Version);
+			bIsValid = (uint8)(bSuccess);
+		}
+
+		void FMetaSoundClassInfo::ExportToContext(FAssetRegistryTagsContext& OutContext) const
+		{
+			using FAssetRegistryTag = UObject::FAssetRegistryTag;
+
+			OutContext.AddTag(FAssetRegistryTag(ClassTagsPrivate::ClassVersion, Version.ToString(), FAssetRegistryTag::TT_Alphabetical));
+
+#if WITH_EDITORONLY_DATA
+			if (!OutContext.IsCooking())
+			{
+				InterfaceInfo.ExportToContext(OutContext);
+			}
+#endif // WITH_EDITORONLY_DATA
+		}
+
+		void FMetaSoundClassInfo::InitFromDocument(const IMetaSoundDocumentInterface& InDocInterface)
+		{
+			const FMetasoundFrontendDocument& Document = InDocInterface.GetConstDocument();
+			const FMetasoundFrontendClassMetadata& ClassMetadata = Document.RootGraph.Metadata;
+
+			ClassName = Document.RootGraph.Metadata.GetClassName();
+			Version = ClassMetadata.GetVersion();
+
+#if WITH_EDITORONLY_DATA
+			if (!IsRunningCookCommandlet())
+			{
+				InterfaceInfo = FMetaSoundClassInterfaceInfo(InDocInterface);
+			}
+#endif // WITH_EDITORONLY_DATA
+
+			bIsValid = 1;
+		}
+
+		bool FMetaSoundClassInfo::TryGetClassVersion(const FAssetData& InAssetData, FMetasoundFrontendVersionNumber& OutVersion)
+		{
+			if (InAssetData.IsAssetLoaded())
+			{
+				if (const UObject* Object = InAssetData.GetAsset())
+				{
+					TScriptInterface<const IMetaSoundDocumentInterface> DocInterface(Object);
+					if (DocInterface.GetObject())
+					{
+						OutVersion = DocInterface->GetConstDocument().Metadata.Version.Number;
+						return true;
+					}
+				}
+			}
+
+			bool bSuccess = true;
+			{
+				FString ClassVersionStr;
+				bSuccess &= InAssetData.GetTagValue(ClassTagsPrivate::ClassVersion, ClassVersionStr);
+				bSuccess &= FMetasoundFrontendVersionNumber::Parse(ClassVersionStr, OutVersion);
+			}
+
+			return bSuccess;
+		}
+
+	} // namespace Frontend
+
 	namespace FrontendQueryPrivate
 	{
 		void CompactSelection(FFrontendQuerySelection& InSelection)
diff --git a/Engine/Plugins/Runtime/Metasound/Source/MetasoundFrontend/Private/MetasoundFrontendSearchEngineCore.cpp b/Engine/Plugins/Runtime/Metasound/Source/MetasoundFrontend/Private/MetasoundFrontendSearchEngineCore.cpp
--- a/Engine/Plugins/Runtime/Metasound/Source/MetasoundFrontend/Private/MetasoundFrontendSearchEngineCore.cpp
+++ b/Engine/Plugins/Runtime/Metasound/Source/MetasoundFrontend/Private/MetasoundFrontendSearchEngineCore.cpp
@@ -129,7 +129,7 @@ namespace Metasound::Frontend
 			const FMetasoundFrontendInterface* Interface = InValue.TryGet<FMetasoundFrontendInterface>();
 			if (ensure(Interface))
 			{
-				InValue.Set<FMetasoundFrontendVersion>(Interface->Version);
+				InValue.Set<FMetasoundFrontendVersion>(Interface->Metadata.Version);
 			}
 		}
 
@@ -145,7 +145,7 @@ namespace Metasound::Frontend
 			{
 				const FMetasoundFrontendInterface& Interface = InEntry.Value.Get<FMetasoundFrontendInterface>();
 
-				return FMapInterfaceToInterfaceNameAndMajorVersion::GetKey(Interface.Version.Name.ToString(), Interface.Version.Number.Major);
+				return FMapInterfaceToInterfaceNameAndMajorVersion::GetKey(Interface.Metadata.Version.Name.ToString(), Interface.Metadata.Version.Number.Major);
 			}
 			return FFrontendQueryKey{Key};
 		}
@@ -159,7 +159,7 @@ namespace Metasound::Frontend
 			{
 				if (ensure(Entry.Value.IsType<FMetasoundFrontendInterface>()))
 				{
-					const FMetasoundFrontendVersionNumber& VersionNumber = Entry.Value.Get<FMetasoundFrontendInterface>().Version.Number;
+					const FMetasoundFrontendVersionNumber& VersionNumber = Entry.Value.Get<FMetasoundFrontendInterface>().Metadata.Version.Number;
 					if (VersionNumber > HighestVersion)
 					{
 						HighestVersionEntry = &Entry;
@@ -193,7 +193,7 @@ namespace Metasound::Frontend
 				{
 					return CreateKeyFromUClassPathName(InOptions.ClassPath);
 				};
-				Algo::TransformIf(Interface->UClassOptions, Keys, IsDefaultInterface, GetUClassQueryKey);
+				Algo::TransformIf(Interface->Metadata.UClassOptions, Keys, IsDefaultInterface, GetUClassQueryKey);
 			}
 
 			return Keys;
@@ -230,7 +230,7 @@ namespace Metasound::Frontend
 			FName Key;
 			if (ensure(InEntry.Value.IsType<FMetasoundFrontendInterface>()))
 			{
-				Key = InEntry.Value.Get<FMetasoundFrontendInterface>().Version.Name;
+				Key = InEntry.Value.Get<FMetasoundFrontendInterface>().Metadata.Version.Name;
 			}
 			return FFrontendQueryKey{Key};
 		}
diff --git a/Engine/Plugins/Runtime/Metasound/Source/MetasoundFrontend/Private/MetasoundFrontendTransform.cpp b/Engine/Plugins/Runtime/Metasound/Source/MetasoundFrontend/Private/MetasoundFrontendTransform.cpp
--- a/Engine/Plugins/Runtime/Metasound/Source/MetasoundFrontend/Private/MetasoundFrontendTransform.cpp
+++ b/Engine/Plugins/Runtime/Metasound/Source/MetasoundFrontend/Private/MetasoundFrontendTransform.cpp
@@ -431,12 +431,12 @@ namespace Metasound
 		{
 			for (const FMetasoundFrontendInterface& Interface : InterfacesToRemove)
 			{
-				DocumentHandle->RemoveInterfaceVersion(Interface.Version);
+				DocumentHandle->RemoveInterfaceVersion(Interface.Metadata.Version);
 			}
 
 			for (const FMetasoundFrontendInterface& Interface : InterfacesToAdd)
 			{
-				DocumentHandle->AddInterfaceVersion(Interface.Version);
+				DocumentHandle->AddInterfaceVersion(Interface.Metadata.Version);
 			}
 
 			return !InterfacesToRemove.IsEmpty() || !InterfacesToAdd.IsEmpty();
diff --git a/Engine/Plugins/Runtime/Metasound/Source/MetasoundFrontend/Public/Interfaces/MetasoundFrontendInterface.h b/Engine/Plugins/Runtime/Metasound/Source/MetasoundFrontend/Public/Interfaces/MetasoundFrontendInterface.h
--- a/Engine/Plugins/Runtime/Metasound/Source/MetasoundFrontend/Public/Interfaces/MetasoundFrontendInterface.h
+++ b/Engine/Plugins/Runtime/Metasound/Source/MetasoundFrontend/Public/Interfaces/MetasoundFrontendInterface.h
@@ -20,6 +20,7 @@ namespace Metasound
 
 		METASOUNDFRONTEND_API int32 InputOutputDifferenceCount(const FMetasoundFrontendInterface& InInterfaceA, const FMetasoundFrontendInterface& InInterfaceB);
 
+		UE_DEPRECATED(5.6, "All interface functionality is in the MetaSoundFrontendDocumentBuilder API")
 		METASOUNDFRONTEND_API bool DeclaredInterfaceVersionsMatch(const FMetasoundFrontendDocument& InDocumentA, const FMetasoundFrontendDocument& InDocumentB);
 
 		METASOUNDFRONTEND_API void GatherRequiredEnvironmentVariables(const FMetasoundFrontendClass& InClass, TArray<FMetasoundFrontendClassEnvironmentVariable>& OutEnvironmentVariables);
diff --git a/Engine/Plugins/Runtime/Metasound/Source/MetasoundFrontend/Public/MetasoundAssetManager.h b/Engine/Plugins/Runtime/Metasound/Source/MetasoundFrontend/Public/MetasoundAssetManager.h
--- a/Engine/Plugins/Runtime/Metasound/Source/MetasoundFrontend/Public/MetasoundAssetManager.h
+++ b/Engine/Plugins/Runtime/Metasound/Source/MetasoundFrontend/Public/MetasoundAssetManager.h
@@ -16,6 +16,8 @@
 #include "UObject/UObjectGlobals.h"
 #include "UObject/WeakObjectPtrTemplates.h"
 
+#include "MetasoundAssetManager.generated.h"
+
 
 // Forward Declarations
 class FMetasoundAssetBase;
@@ -28,33 +30,57 @@ struct FMetasoundFrontendClassName;
 struct FMetaSoundFrontendDocumentBuilder;
 
 
+USTRUCT(BlueprintType, DisplayName = "MetaSound Document Info")
+struct METASOUNDFRONTEND_API FMetaSoundDocumentInfo
+{
+	GENERATED_BODY()
+
+	FMetaSoundDocumentInfo();
+	FMetaSoundDocumentInfo(const IMetaSoundDocumentInterface& InDocInterface);
+	FMetaSoundDocumentInfo(const FAssetData& InAssetData, bool& bOutIsValid);
+
+	void ExportToContext(FAssetRegistryTagsContext& OutContext) const;
+
+	// Version of document
+	UPROPERTY(BlueprintReadOnly, Category = DocumentInfo)
+	FMetasoundFrontendVersionNumber DocumentVersion;
+
+	// Versions of referenced asset class keys
+	UPROPERTY()
+	TArray<FMetaSoundAssetKey> ReferencedAssetKeys;
+
+	// Whether asset is a preset or not.
+	UPROPERTY(BlueprintReadOnly, Category = DocumentInfo)
+	uint8 bIsPreset : 1;
+};
+
 namespace Metasound::Frontend
 {
 	namespace AssetTags
 	{
-		UE_DEPRECATED(5.6, "AssetTags no longer public. Construct FAssetClassQueryInfo from AssetData to parse tag fields.")
+		UE_DEPRECATED(5.6, "AssetTags no longer public. Construct FMetaSoundAssetClassInfo from AssetData to parse tag fields.")
 		extern const FString METASOUNDFRONTEND_API ArrayDelim;
 
 #if WITH_EDITORONLY_DATA
-		UE_DEPRECATED(5.6, "AssetTags no longer public. Construct FAssetClassQueryInfo from AssetData to parse tag fields.")
+		UE_DEPRECATED(5.6, "AssetTags no longer public. Construct FMetaSoundAssetClassInfo from AssetData to parse tag fields.")
 		extern const FName METASOUNDFRONTEND_API IsPreset;
 #endif // WITH_EDITORONLY_DATA
 
 
-		UE_DEPRECATED(5.6, "AssetTags no longer public. Construct FAssetClassQueryInfo from AssetData to parse tag fields.")
+		UE_DEPRECATED(5.6, "AssetTags no longer public. Construct FMetaSoundAssetClassInfo from AssetData to parse tag fields.")
 		extern const FName METASOUNDFRONTEND_API AssetClassID;
 
-		UE_DEPRECATED(5.6, "AssetTags no longer public. Construct FAssetClassQueryInfo from AssetData to parse tag fields.")
+		UE_DEPRECATED(5.6, "AssetTags no longer public. Construct FMetaSoundAssetClassInfo from AssetData to parse tag fields.")
 		extern const FName METASOUNDFRONTEND_API RegistryVersionMajor;
 
-		UE_DEPRECATED(5.6, "AssetTags no longer public. Construct FAssetClassQueryInfo from AssetData to parse tag fields.")
+		UE_DEPRECATED(5.6, "AssetTags no longer public. Construct FMetaSoundAssetClassInfo from AssetData to parse tag fields.")
 		extern const FName METASOUNDFRONTEND_API RegistryVersionMinor;
 
 #if WITH_EDITORONLY_DATA
-		UE_DEPRECATED(5.6, "AssetTags no longer public. Construct FAssetClassQueryInfo from AssetData to parse tag fields.")
+		UE_DEPRECATED(5.6, "AssetTags no longer public. Construct FMetaSoundAssetClassInfo from AssetData to parse tag fields.")
 		extern const FName METASOUNDFRONTEND_API RegistryInputTypes;
 
-		UE_DEPRECATED(5.6, "AssetTags no longer public. Construct FAssetClassQueryInfo from AssetData to parse tag fields.")
+		UE_DEPRECATED(5.6, "AssetTags no longer public. Construct FMetaSoundAssetClassInfo from AssetData to parse tag fields.")
 		extern const FName METASOUNDFRONTEND_API RegistryOutputTypes;
 #endif // WITH_EDITORONLY_DATA
 	} // namespace AssetTags
@@ -94,50 +120,53 @@ namespace Metasound::Frontend
 	};
 
 
-	/** Contains a minimal set of information needed to display or query additional
-	  * class data from the internal MetaSound Asset or loaded Node Class Registries. This data is
-	  * serializable to asset tags, and is available even if asset is unloaded / not registered with
-	  * the MetaSound Node Class Registry. */
-	struct METASOUNDFRONTEND_API FAssetClassQueryInfo : public FClassQueryInfo
+	/** At runtime, contains a minimal set of information needed to further query additional
+	  * class data from the Node Class Registries. At edit-time, contains this plus additional
+	  * data useful for informing user of applicable asset classes within given edit contexts
+	  * (see 'FMetaSoundEditorGraphSchema' for usage).
+	  */
+	struct METASOUNDFRONTEND_API FMetaSoundAssetClassInfo : public FMetaSoundClassInfo
 	{
-		FAssetClassQueryInfo();
-		FAssetClassQueryInfo(const IMetaSoundDocumentInterface& InDocInterface);
+		FMetaSoundAssetClassInfo() = default;
+		FMetaSoundAssetClassInfo(const IMetaSoundDocumentInterface& InDocInterface);
 
 		// Attempts to return tag data from the given AssetData, which may or may not be loaded.
 		// Invalid tag constructed if object isn't loaded and any fields fail to load/are not serialized
 		// (i.e. tags are out-of-date).
-		FAssetClassQueryInfo(const FAssetData& InAssetData);
+		FMetaSoundAssetClassInfo(const FAssetData& InAssetData);
 
 		// Exports tag data to the given RegistryContext
-		void ExportToContext(FAssetRegistryTagsContext& OutContext) const;
+		virtual void ExportToContext(FAssetRegistryTagsContext& OutContext) const override;
+
+		// If asset is loaded, retrieves asset key from loaded data. Otherwise, parses just the
+		// tag data necessary to get the given asset's asset key. Does not attempt to load asset,
+		// and will fail if data is not found and asset isn't loaded.
+		static bool TryGetAssetKey(const FAssetData& InAssetData, FMetaSoundAssetKey& OutKey);
+
+		// If asset is loaded, retrieves asset class name from loaded data. Otherwise, parses just the
+		// tag data necessary to get the given asset's class name. Does not attempt to load asset,
+		// and will fail if data is not found and asset isn't loaded.
+		static bool TryGetAssetClassName(const FAssetData& InAssetData, FMetasoundFrontendClassName& OutClassName);
+
+	protected:
+		virtual void InitFromDocument(const IMetaSoundDocumentInterface& InDocInterface) override;
 
 	private:
-		void InitFromDocument(const IMetaSoundDocumentInterface& InDocInterface);
+		static bool TryGetAssetClassTag(const FAssetData& InAssetData, FString& OutClassIDString);
 
 	public:
 		// Path to asset containing graph if external type and references asset class.
 		FTopLevelAssetPath AssetPath;
 
 #if WITH_EDITORONLY_DATA
-		// Document Version of asset
-		FMetasoundFrontendVersionNumber DocumentVersion;
-
-		// Versions of referenced asset class keys
-		TArray<FMetaSoundAssetKey> ReferencedAssetKeys;
-
-		// Whether asset is a preset or not.
-		uint8 bIsPreset : 1;
+		FMetaSoundDocumentInfo DocInfo;
 #endif // WITH_EDITORONLY_DATA
-
-		// True, asset tag data is invalid due to provided asset not being loaded and out-of-date
-		uint8 bIsValid : 1;
 	};
 
 
 	class METASOUNDFRONTEND_API IMetaSoundAssetManager
 	{
 	public:
-
 		virtual ~IMetaSoundAssetManager() = default;
 
 		static IMetaSoundAssetManager* Get();
@@ -200,9 +229,9 @@ namespace Metasound::Frontend
 		// on entry update (synchronously).
 		virtual void AddOrLoadAndUpdateFromObjectAsync(const FAssetData& InAssetData, TUniqueFunction<void(FAssetKey, UObject&)>&& OnAssetRegistered) = 0;
 
-		// Add or Update a MetaSound Asset's entry data from AssetTags. Potentially loads asset and added asynchronously if asset is or asset tags
-		// are out-of-date.
-		virtual void AddOrUpdateFromAssetTags(const FAssetData& InAssetData) = 0;
+		// Add or Update a MetaSound Asset's entry data from AssetData. Potentially loads asset
+		// and adds asynchronously if asset or associated tag schema is out-of-date.
+		virtual void AddOrUpdateFromAssetData(const FAssetData& InAssetData) = 0;
 
 		// Whether or not the class is eligible for auto-update
 		virtual bool CanAutoUpdate(const FMetasoundFrontendClassName& InClassName) const = 0;
@@ -246,7 +275,7 @@ namespace Metasound::Frontend
 		// returns whether or not the given class is defined as a registered asset
 		virtual bool IsAssetClass(const FMetasoundFrontendClassMetadata& ClassMetadata) const = 0;
 
-		UE_DEPRECATED(5.6, "Moved to Iterate ClassQueryInfo")
+		UE_DEPRECATED(5.6, "Moved to Iterate ClassInfo")
 		virtual void IterateAssets(TFunctionRef<void(const FAssetKey, const TArray<FTopLevelAssetPath>&)> Iter) const { }
 
 		UE_DEPRECATED(5.5, "Rescan no longer supported nor required by Frontend")
diff --git a/Engine/Plugins/Runtime/Metasound/Source/MetasoundFrontend/Public/MetasoundFrontendDocument.h b/Engine/Plugins/Runtime/Metasound/Source/MetasoundFrontend/Public/MetasoundFrontendDocument.h
--- a/Engine/Plugins/Runtime/Metasound/Source/MetasoundFrontend/Public/MetasoundFrontendDocument.h
+++ b/Engine/Plugins/Runtime/Metasound/Source/MetasoundFrontend/Public/MetasoundFrontendDocument.h
@@ -1276,7 +1276,7 @@ struct METASOUNDFRONTEND_API FMetasoundFrontendInterfaceBinding
 
 // Options used to restrict a corresponding UClass that interface may be applied to.
 // If unspecified, interface is assumed to be applicable to any arbitrary UClass.
-USTRUCT()
+USTRUCT(BlueprintType)
 struct METASOUNDFRONTEND_API FMetasoundFrontendInterfaceUClassOptions
 {
 	GENERATED_BODY()
@@ -1285,17 +1285,36 @@ struct METASOUNDFRONTEND_API FMetasoundFrontendInterfaceUClassOptions
 	FMetasoundFrontendInterfaceUClassOptions(const Audio::FParameterInterface::FClassOptions& InOptions);
 	FMetasoundFrontendInterfaceUClassOptions(const FTopLevelAssetPath& InClassPath, bool bInIsModifiable = true, bool bInIsDefault = false);
 
-	UPROPERTY()
+	// Path to MetaSound class interface can be added to (ex. UMetaSoundSource or UMetaSound)
+	UPROPERTY(BlueprintReadOnly, meta = (Category = Interface))
 	FTopLevelAssetPath ClassPath;
 
+	// True if user can add or remove the given class directly to or from the inherited interface UI, false if not.
 	UPROPERTY()
 	bool bIsModifiable = true;
 
+	// True if interface should be added by default to newly created MetaSound assets, false if not.
 	UPROPERTY()
 	bool bIsDefault = false;
 };
 
 
+USTRUCT(BlueprintType)
+struct METASOUNDFRONTEND_API FMetasoundFrontendInterfaceMetadata
+{
+	GENERATED_BODY()
+
+	// Name and version number of the interface
+	UPROPERTY(BlueprintReadOnly, meta = (Category = Interface))
+	FMetasoundFrontendVersion Version;
+
+	// If specified, options used to restrict a corresponding UClass that interface may be
+	// applied to.  If unspecified, interface is assumed to be applicable to any arbitrary UClass.
+	UPROPERTY(BlueprintReadOnly, meta = (Category = Interface))
+	TArray<FMetasoundFrontendInterfaceUClassOptions> UClassOptions;
+};
+
+
 // Definition of an interface that an FMetasoundFrontendClass adheres to in part or full.
 USTRUCT()
 struct METASOUNDFRONTEND_API FMetasoundFrontendInterface : public FMetasoundFrontendClassInterface
@@ -1305,16 +1324,19 @@ struct METASOUNDFRONTEND_API FMetasoundFrontendInterface : public FMetasoundFron
 	FMetasoundFrontendInterface() = default;
 	FMetasoundFrontendInterface(Audio::FParameterInterfacePtr InInterface);
 
-	// Name and version number of the interface
 	UPROPERTY()
+	FMetasoundFrontendInterfaceMetadata Metadata;
+
+#if WITH_EDITORONLY_DATA
+	UPROPERTY(Transient, meta = (DeprecatedProperty, DeprecationMessage = "5.6 - Field never serialized but will be in the future, and moved to Metadata. Will be removed in subsequent release."))
 	FMetasoundFrontendVersion Version;
 
-	// If specified, options used to restrict a corresponding UClass that interface may be
-	// applied to.  If unspecified, interface is assumed to be applicable to any arbitrary UClass.
-	UPROPERTY()
+	UPROPERTY(Transient, meta = (DeprecatedProperty, DeprecationMessage = "5.6 - Field never serialized but will be in the future, and moved to Metadata. Will be removed in subsequent release."))
 	TArray<FMetasoundFrontendInterfaceUClassOptions> UClassOptions;
 
+	UE_DEPRECATED(5.6, "Inlined where necessary. Use desired predicate look-up on UClassOptions from now-shared Metadata struct above.")
 	const FMetasoundFrontendInterfaceUClassOptions* FindClassOptions(const FTopLevelAssetPath& InClassPath) const;
+#endif // WITH_EDITORONLY_DATA
 };
 
 
@@ -1661,10 +1683,8 @@ struct METASOUNDFRONTEND_API FMetasoundFrontendClass
 	FMetasoundFrontendClassInterface Interface;
 
 #if WITH_EDITORONLY_DATA
-
 	UPROPERTY()
 	FMetasoundFrontendClassStyle Style;
-
 #endif // WITH_EDITORONLY_DATA
 
 #if WITH_EDITOR
diff --git a/Engine/Plugins/Runtime/Metasound/Source/MetasoundFrontend/Public/MetasoundFrontendQuery.h b/Engine/Plugins/Runtime/Metasound/Source/MetasoundFrontend/Public/MetasoundFrontendQuery.h
--- a/Engine/Plugins/Runtime/Metasound/Source/MetasoundFrontend/Public/MetasoundFrontendQuery.h
+++ b/Engine/Plugins/Runtime/Metasound/Source/MetasoundFrontend/Public/MetasoundFrontendQuery.h
@@ -53,60 +53,101 @@
 // Forward Declarations
 class FAssetRegistryTagsContext;
 class IMetaSoundDocumentInterface;
+struct FAssetData;
 
 
-// Minimal Json-serializable vertex data which characterizes an assets interface vertex for use in tag data.
-USTRUCT()
-struct METASOUNDFRONTEND_API FClassQueryVertexInfo
+// Condensed set of class vertex data that is serialized to editor-only asset
+// tag data, allowing editor scripts and code to query MetaSounds without loading
+// them in entirety.
+USTRUCT(BlueprintType)
+struct METASOUNDFRONTEND_API FMetaSoundClassVertexInfo
 {
 	GENERATED_BODY()
 
-#if WITH_EDITORONLY_DATA
-	UPROPERTY()
+	UPROPERTY(BlueprintReadOnly, Category = ClassInfo)
 	FName Name;
 
-	UPROPERTY()
+	UPROPERTY(BlueprintReadOnly, Category = ClassInfo)
 	FName TypeName;
 
-	UPROPERTY()
+	UPROPERTY(BlueprintReadOnly, Category = ClassInfo)
 	EMetasoundFrontendVertexAccessType AccessType = EMetasoundFrontendVertexAccessType::Unset;
-#endif // WITH_EDITORONLY_DATA
 };
 
-USTRUCT()
-struct METASOUNDFRONTEND_API FClassQueryDisplayInfo
+// Condensed set of class metadata that is serialized to editor-only asset
+// tag data, allowing editor scripts and code to search and display MetaSounds
+// without loading in asset selection contexts without loading them in entirety.
+USTRUCT(BlueprintType)
+struct METASOUNDFRONTEND_API FMetaSoundClassSearchInfo
 {
 	GENERATED_BODY()
 
-#if WITH_EDITORONLY_DATA
+	FMetaSoundClassSearchInfo() = default;
+	FMetaSoundClassSearchInfo(const FMetasoundFrontendClassMetadata& InClassMetadata);
+
 	// Human readable DisplayName of Class (optional, overrides the
 	// package name in the editor if specified by MetaSound Asset Author).
-	UPROPERTY()
+	UPROPERTY(BlueprintReadOnly, Category = ClassInfo)
 	FText ClassDisplayName;
 
-	UPROPERTY()
+	UPROPERTY(BlueprintReadOnly, Category = ClassInfo)
 	FText ClassDescription;
 
-	UPROPERTY()
+	UPROPERTY(BlueprintReadOnly, Category = ClassInfo)
 	TArray<FText> Hierarchy;
 
-	UPROPERTY()
+	UPROPERTY(BlueprintReadOnly, Category = ClassInfo)
 	TArray<FText> Keywords;
-#endif // WITH_EDITORONLY_DATA
 };
 
-// Minimal Json-serializable array of vertex data, which characterizes an assets interface for use in tag data.
-USTRUCT()
-struct METASOUNDFRONTEND_API FClassQueryVertexCollection
+// Condensed set of vertex data that is serialized to editor-only asset tag data,
+// allowing editor scripts and code to query MetaSounds without loading them
+// in entirety.
+USTRUCT(BlueprintType)
+struct METASOUNDFRONTEND_API FMetaSoundClassVertexCollectionInfo
 {
 	GENERATED_BODY()
 
-#if WITH_EDITORONLY_DATA
-	UPROPERTY()
-	TArray<FClassQueryVertexInfo> ClassVertexInfo;
-#endif // WITH_EDITORONLY_DATA
+	UPROPERTY(BlueprintReadOnly, Category = ClassInfo)
+	TArray<FMetaSoundClassVertexInfo> ClassVertexInfo;
 };
 
+// Condensed set of class data that is serialized to editor-only asset tag data,
+// allowing editor scripts and code to query MetaSounds without loading them
+// in entirety.
+USTRUCT(BlueprintType, DisplayName = "MetaSound Class Interface Info")
+struct METASOUNDFRONTEND_API FMetaSoundClassInterfaceInfo
+{
+	GENERATED_BODY()
+
+	FMetaSoundClassInterfaceInfo() = default;
+	FMetaSoundClassInterfaceInfo(const IMetaSoundDocumentInterface& InDocInterface);
+	FMetaSoundClassInterfaceInfo(const FAssetData& InAssetData, bool& bOutIsValid);
+
+	void ExportToContext(FAssetRegistryTagsContext& OutContext) const;
+
+	// Interfaces metadata associated with interfaces defined by this class.
+	UPROPERTY(BlueprintReadOnly, Category = ClassInfo)
+	TArray<FMetasoundFrontendInterfaceMetadata> DefinedInterfaces;
+
+	// Editor-only search info
+	UPROPERTY(BlueprintReadOnly, Category = ClassInfo)
+	FMetaSoundClassSearchInfo SearchInfo;
+
+	// Collection of identifiable input vertex data cached in query for fast access & serializability
+	// (ex. in asset tags)
+	UPROPERTY(BlueprintReadOnly, Category = ClassInfo)
+	TArray<FMetaSoundClassVertexInfo> Inputs;
+
+	// Collection of identifiable output vertex data cached in query for fast access & serializability
+	// (ex. in asset tags)
+	UPROPERTY(BlueprintReadOnly, Category = ClassInfo)
+	TArray<FMetaSoundClassVertexInfo> Outputs;
+
+	// Interfaces metadata associated with a given class definition.
+	UPROPERTY(BlueprintReadOnly, Category = ClassInfo)
+	TArray<FMetasoundFrontendVersion> InheritedInterfaces;
+};
 
 namespace Metasound
 {
@@ -116,8 +157,25 @@ namespace Metasound
 		// that may or may not be loaded and/or registered (ex. in editor, class info may be supplied by
 		// asset registry tags or class info could be provided by a cloud service, enabling class browsing
 		// without a local asset required to be loaded or registered with the node class registry).
-		struct METASOUNDFRONTEND_API FClassQueryInfo
+		struct METASOUNDFRONTEND_API FMetaSoundClassInfo
 		{
+			FMetaSoundClassInfo();
+			FMetaSoundClassInfo(const FAssetData& InAssetData);
+
+			virtual ~FMetaSoundClassInfo() = default;
+
+			// Exports class query info as tag data to the given RegistryContext
+			virtual void ExportToContext(FAssetRegistryTagsContext& OutContext) const;
+
+			// If asset is loaded, retrieves version number from loaded data. Otherwise, parses just the
+			// tag data necessary to get the given asset's version number. Does not attempt to load asset,
+			// and will fail if data is not found and asset isn't loaded.
+			static bool TryGetClassVersion(const FAssetData& InAssetData, FMetasoundFrontendVersionNumber& OutKey);
+
+		protected:
+			virtual void InitFromDocument(const IMetaSoundDocumentInterface& InDocInterface);
+
+		public:
 			// ClassName of class
 			FMetasoundFrontendClassName ClassName;
 
@@ -125,17 +183,11 @@ namespace Metasound
 			FMetasoundFrontendVersionNumber Version;
 
 #if WITH_EDITORONLY_DATA
-			// Editor-only display info, included for fast look-up and 
-			FClassQueryDisplayInfo DisplayInfo;
-
-			// Collection of identifiable input vertex data cached in query for fast access & serializability
-			// (ex. in asset tags)
-			FClassQueryVertexCollection Inputs;
-
-			// Collection of identifiable output vertex data cached in query for fast access & serializability
-			// (ex. in asset tags)
-			FClassQueryVertexCollection Outputs;
+			FMetaSoundClassInterfaceInfo InterfaceInfo;
 #endif // WITH_EDITORONLY_DATA
+
+			// If true, class query data was successfully loaded from manager, false if not.
+			uint8 bIsValid : 1;
 		};
 	} // namespace Frontend
 
