Adding asset scopes to metasound building and rendering

#swarm https://p4-swarm.epicgames.net/reviews/40808519

diff --git a/Engine/Plugins/Runtime/Metasound/Source/MetasoundEngine/Private/MetasoundSource.cpp b/Engine/Plugins/Runtime/Metasound/Source/MetasoundEngine/Private/MetasoundSource.cpp
--- a/Engine/Plugins/Runtime/Metasound/Source/MetasoundEngine/Private/MetasoundSource.cpp
+++ b/Engine/Plugins/Runtime/Metasound/Source/MetasoundEngine/Private/MetasoundSource.cpp
@@ -1057,7 +1057,9 @@ ISoundGeneratorPtr UMetaSoundSource::CreateSoundGenerator(const FSoundGeneratorI
 				MoveTemp(InDefaultParameters),
 				bBuildSynchronous,
 				DataChannel,
-				GraphRenderCost
+				GraphRenderCost,
+				FName(GetClass()->GetName()),
+				GetAssetPathChecked().GetPackageName()
 			},
 			DynamicTransactor->CreateTransformQueue(InSettings, Environment, GraphRenderCost) // Create transaction queue
 		};
@@ -1096,7 +1098,10 @@ ISoundGeneratorPtr UMetaSoundSource::CreateSoundGenerator(const FSoundGeneratorI
 				GetOutputAudioChannelOrder(),
 				MoveTemp(InDefaultParameters),
 				bBuildSynchronous,
-				DataChannel
+				DataChannel,
+				{}, // RelativeRenderCost
+				FName(GetClass()->GetName()),
+				GetAssetPathChecked().GetPackageName()
 			};
 
 			Generator = MakeShared<FMetasoundConstGraphGenerator>(MoveTemp(InitParams));
diff --git a/Engine/Plugins/Runtime/Metasound/Source/MetasoundGenerator/MetasoundGenerator.Build.cs b/Engine/Plugins/Runtime/Metasound/Source/MetasoundGenerator/MetasoundGenerator.Build.cs
--- a/Engine/Plugins/Runtime/Metasound/Source/MetasoundGenerator/MetasoundGenerator.Build.cs
+++ b/Engine/Plugins/Runtime/Metasound/Source/MetasoundGenerator/MetasoundGenerator.Build.cs
@@ -14,6 +14,7 @@ namespace UnrealBuildTool.Rules
 					"SignalProcessing",
 					"Engine",
 					"AudioExtensions",
+					"TraceLog"
 				}
 			);
 
diff --git a/Engine/Plugins/Runtime/Metasound/Source/MetasoundGenerator/Private/MetasoundGenerator.cpp b/Engine/Plugins/Runtime/Metasound/Source/MetasoundGenerator/Private/MetasoundGenerator.cpp
--- a/Engine/Plugins/Runtime/Metasound/Source/MetasoundGenerator/Private/MetasoundGenerator.cpp
+++ b/Engine/Plugins/Runtime/Metasound/Source/MetasoundGenerator/Private/MetasoundGenerator.cpp
@@ -22,12 +22,14 @@
 #include "Interfaces/MetasoundFrontendSourceInterface.h"
 #include "Modules/ModuleManager.h"
 #include "AutoRTFM/AutoRTFM.h"
+#include "ProfilingDebugging/AssetMetadataTrace.h"
 
 #ifndef ENABLE_METASOUNDGENERATOR_INVALID_SAMPLE_VALUE_LOGGING 
 #define ENABLE_METASOUNDGENERATOR_INVALID_SAMPLE_VALUE_LOGGING !UE_BUILD_SHIPPING
 #endif
 
 CSV_DECLARE_CATEGORY_MODULE_EXTERN(METASOUNDGRAPHCORE_API, Audio_Metasound);
+LLM_DEFINE_TAG(Audio_MetaSound_Generator_Execute);
 
 namespace Metasound
 {
@@ -656,7 +658,9 @@ PRAGMA_ENABLE_DEPRECATION_WARNINGS
 			// Create a scoped timed section for the bulk of the processing...
 			{
 				MetasoundGeneratorPrivate::FBlockRenderScope RuntimelBlockRenderScope(RenderTimer.Get());
-
+				LLM_SCOPE_BYTAG(Audio_MetaSound_Generator_Execute);
+				UE_TRACE_METADATA_SCOPE_ASSET_FNAME(FName(MetasoundName), NAME_None, NAME_None);
+				
 				ApplyPendingUpdatesToInputs();
 
 				// Call metasound graph operator.
diff --git a/Engine/Plugins/Runtime/Metasound/Source/MetasoundGenerator/Private/MetasoundGeneratorBuilder.cpp b/Engine/Plugins/Runtime/Metasound/Source/MetasoundGenerator/Private/MetasoundGeneratorBuilder.cpp
--- a/Engine/Plugins/Runtime/Metasound/Source/MetasoundGenerator/Private/MetasoundGeneratorBuilder.cpp
+++ b/Engine/Plugins/Runtime/Metasound/Source/MetasoundGenerator/Private/MetasoundGeneratorBuilder.cpp
@@ -17,9 +17,12 @@
 #include "Analysis/MetasoundFrontendAnalyzerRegistry.h"
 #include "Analysis/MetasoundFrontendVertexAnalyzer.h"
 #include "Interfaces/MetasoundFrontendSourceInterface.h"
+#include "ProfilingDebugging/AssetMetadataTrace.h"
 
 namespace Metasound
 {
+	LLM_DEFINE_TAG(Audio_MetaSound_Generator_Build);
+	
 	FAsyncMetaSoundBuilderBase::FAsyncMetaSoundBuilderBase(FMetasoundGenerator* InGenerator, bool bInTriggerGenerator)
 	: Generator(InGenerator)
 	, bTriggerGenerator(bInTriggerGenerator)
@@ -268,6 +271,16 @@ namespace Metasound
 			// Create an instance of the new graph operator
 			FBuildGraphOperatorParams BuildParams { *InInitParams.Graph, InOperatorSettings, OpAndInputs.Inputs, InInitParams.Environment, InInitParams.GraphRenderCost.Get()};
 			FOperatorBuilder Builder(InInitParams.BuilderSettings);
+			
+			LLM_SCOPE_BYTAG(Audio_MetaSound_Generator_Build);
+
+			// asset-level trace asset path
+			UE_TRACE_METADATA_SCOPE_ASSET_FNAME(
+				  FName(InInitParams.MetaSoundName)
+				, InInitParams.ClassName
+				, InInitParams.AssetPath
+			);
+			
 			OpAndInputs.Operator = Builder.BuildGraphOperator(BuildParams, OutBuildResults);
 
 			return OpAndInputs;
diff --git a/Engine/Plugins/Runtime/Metasound/Source/MetasoundGenerator/Public/MetasoundGenerator.h b/Engine/Plugins/Runtime/Metasound/Source/MetasoundGenerator/Public/MetasoundGenerator.h
--- a/Engine/Plugins/Runtime/Metasound/Source/MetasoundGenerator/Public/MetasoundGenerator.h
+++ b/Engine/Plugins/Runtime/Metasound/Source/MetasoundGenerator/Public/MetasoundGenerator.h
@@ -117,6 +117,8 @@ namespace Metasound
 		bool bBuildSynchronous = false;
 		TSharedPtr<TSpscQueue<FMetaSoundParameterTransmitter::FParameter>> DataChannel;
 		TSharedPtr<FGraphRenderCost> GraphRenderCost;
+		FName ClassName;
+		FName AssetPath;
 
 		static void Reset(FMetasoundGeneratorInitParams& InParams);
 	};
diff --git a/Engine/Plugins/Runtime/Metasound/Source/MetasoundGraphCore/MetasoundGraphCore.Build.cs b/Engine/Plugins/Runtime/Metasound/Source/MetasoundGraphCore/MetasoundGraphCore.Build.cs
--- a/Engine/Plugins/Runtime/Metasound/Source/MetasoundGraphCore/MetasoundGraphCore.Build.cs
+++ b/Engine/Plugins/Runtime/Metasound/Source/MetasoundGraphCore/MetasoundGraphCore.Build.cs
@@ -21,6 +21,7 @@ namespace UnrealBuildTool.Rules
 				new string[]
 				{
 					"MathCore",
+					"TraceLog"
 					// ... add private dependencies that you statically link with here ...	
 				}
 			);
diff --git a/Engine/Plugins/Runtime/Metasound/Source/MetasoundGraphCore/Private/MetasoundOperatorBuilder.cpp b/Engine/Plugins/Runtime/Metasound/Source/MetasoundGraphCore/Private/MetasoundOperatorBuilder.cpp
--- a/Engine/Plugins/Runtime/Metasound/Source/MetasoundGraphCore/Private/MetasoundOperatorBuilder.cpp
+++ b/Engine/Plugins/Runtime/Metasound/Source/MetasoundGraphCore/Private/MetasoundOperatorBuilder.cpp
@@ -20,6 +20,7 @@
 #include "MetasoundRenderCost.h"
 #include "MetasoundThreadLocalDebug.h"
 #include "MetasoundTrace.h"
+#include "ProfilingDebugging/AssetMetadataTrace.h"
 #include "Templates/PimplPtr.h"
 
 
diff --git a/Engine/Plugins/Runtime/Metasound/Source/MetasoundGraphCore/Public/MetasoundBuilderInterface.h b/Engine/Plugins/Runtime/Metasound/Source/MetasoundGraphCore/Public/MetasoundBuilderInterface.h
--- a/Engine/Plugins/Runtime/Metasound/Source/MetasoundGraphCore/Public/MetasoundBuilderInterface.h
+++ b/Engine/Plugins/Runtime/Metasound/Source/MetasoundGraphCore/Public/MetasoundBuilderInterface.h
@@ -124,6 +124,10 @@ namespace Metasound
 		/** Runtime render cost tracker. */
 		FGraphRenderCost* GraphRenderCost;
 
+		FName ClassName;
+		FName AssetName;
+		FName AssetPath;
+		
 		/** Convert operator building params to graph building params */
 		static FBuildGraphOperatorParams FromBuildOperatorParams(const IGraph& InGraph, const FBuildOperatorParams& InParams)
 		{
