NodeConfigurationFrontend registration and serialization
#rnx

diff --git a/Engine/Plugins/Runtime/Metasound/Source/MetasoundFrontend/Private/MetasoundAssetBase.cpp b/Engine/Plugins/Runtime/Metasound/Source/MetasoundFrontend/Private/MetasoundAssetBase.cpp
--- a/Engine/Plugins/Runtime/Metasound/Source/MetasoundFrontend/Private/MetasoundAssetBase.cpp
+++ b/Engine/Plugins/Runtime/Metasound/Source/MetasoundFrontend/Private/MetasoundAssetBase.cpp
@@ -117,6 +117,12 @@ namespace Metasound
 					return CreateVertexInterfaceFromClassNoProxy(FrontendClass);
 				}
 
+				virtual TInstancedStruct<FMetaSoundFrontendNodeConfiguration> CreateFrontendNodeConfiguration() const override
+				{
+					// Document nodes do not currently support node configuration.
+					return TInstancedStruct<FMetaSoundFrontendNodeConfiguration>();
+				}
+
 			private:
 				TSet<FMetasoundFrontendVersion> Interfaces;
 				FMetasoundFrontendClass FrontendClass;
diff --git a/Engine/Plugins/Runtime/Metasound/Source/MetasoundFrontend/Private/MetasoundFrontendDataTypeRegistry.cpp b/Engine/Plugins/Runtime/Metasound/Source/MetasoundFrontend/Private/MetasoundFrontendDataTypeRegistry.cpp
--- a/Engine/Plugins/Runtime/Metasound/Source/MetasoundFrontend/Private/MetasoundFrontendDataTypeRegistry.cpp
+++ b/Engine/Plugins/Runtime/Metasound/Source/MetasoundFrontend/Private/MetasoundFrontendDataTypeRegistry.cpp
@@ -137,6 +137,12 @@ namespace Metasound
 					return FVertexInterface();
 				}
 
+				virtual TInstancedStruct<FMetaSoundFrontendNodeConfiguration> CreateFrontendNodeConfiguration() const override
+				{
+					// By default, data type related nodes do not offer any extensions. 
+					return TInstancedStruct<FMetaSoundFrontendNodeConfiguration>();
+				}
+
 			protected:
 				void UpdateNodeClassInfo(const FMetasoundFrontendClass& InClass)
 				{
diff --git a/Engine/Plugins/Runtime/Metasound/Source/MetasoundFrontend/Private/MetasoundFrontendDocument.cpp b/Engine/Plugins/Runtime/Metasound/Source/MetasoundFrontend/Private/MetasoundFrontendDocument.cpp
--- a/Engine/Plugins/Runtime/Metasound/Source/MetasoundFrontend/Private/MetasoundFrontendDocument.cpp
+++ b/Engine/Plugins/Runtime/Metasound/Source/MetasoundFrontend/Private/MetasoundFrontendDocument.cpp
@@ -43,6 +43,18 @@ namespace Metasound
 
 				return Frontend::DefaultPageID;
 			}
+
+			FMetasoundFrontendNodeInterface MakeNodeInterface(const FMetasoundFrontendClass& InClass, const TInstancedStruct<FMetaSoundFrontendNodeConfiguration>& InConfiguration)
+			{
+				if (const FMetaSoundFrontendNodeConfiguration* ConfigurationPtr = InConfiguration.GetPtr())
+				{
+					return ConfigurationPtr->GetNodeInterface(InClass);
+				}
+				else
+				{
+					return FMetasoundFrontendNodeInterface(InClass.Interface);
+				}
+			}
 		} // namespace DocumentPrivate
 
 #if WITH_EDITORONLY_DATA
@@ -346,12 +358,27 @@ FMetasoundFrontendNodeInterface::FMetasoundFrontendNodeInterface(const FMetasoun
 	}
 }
 
+FMetasoundFrontendNodeInterface FMetaSoundFrontendNodeConfiguration::GetNodeInterface(const FMetasoundFrontendClass& InNodeClass) const
+{
+	return FMetasoundFrontendNodeInterface{InNodeClass.Interface};
+}
+
+TSharedPtr<const Metasound::IOperatorData> FMetaSoundFrontendNodeConfiguration::GetOperatorData() const
+{
+	return {};
+}
+
 FMetasoundFrontendNode::FMetasoundFrontendNode(const FMetasoundFrontendClass& InClass)
+: FMetasoundFrontendNode(InClass, TInstancedStruct<FMetaSoundFrontendNodeConfiguration>())
+{
+}
+
+FMetasoundFrontendNode::FMetasoundFrontendNode(const FMetasoundFrontendClass& InClass, TInstancedStruct<FMetaSoundFrontendNodeConfiguration> InConfiguration)
 : ClassID(InClass.ID)
 , Name(InClass.Metadata.GetClassName().Name.ToString())
-, Interface(InClass.Interface)
+, Interface(Metasound::Frontend::DocumentPrivate::MakeNodeInterface(InClass, InConfiguration))
+, Configuration(MoveTemp(InConfiguration))
 {
-
 }
 
 FString FMetasoundFrontendVersion::ToString() const
diff --git a/Engine/Plugins/Runtime/Metasound/Source/MetasoundFrontend/Private/MetasoundFrontendDocumentBuilder.cpp b/Engine/Plugins/Runtime/Metasound/Source/MetasoundFrontend/Private/MetasoundFrontendDocumentBuilder.cpp
--- a/Engine/Plugins/Runtime/Metasound/Source/MetasoundFrontend/Private/MetasoundFrontendDocumentBuilder.cpp
+++ b/Engine/Plugins/Runtime/Metasound/Source/MetasoundFrontend/Private/MetasoundFrontendDocumentBuilder.cpp
@@ -1811,10 +1811,11 @@ FMetasoundFrontendNode* FMetaSoundFrontendDocumentBuilder::AddNodeInternal(const
 
 	if (const FMetasoundFrontendClass* Dependency = DocumentCache->FindDependency(InClassKey))
 	{
+		TInstancedStruct<FMetaSoundFrontendNodeConfiguration> NodeConfiguration = INodeClassRegistry::Get()->CreateFrontendNodeConfiguration(InClassKey);
 		FMetasoundFrontendDocument& Document = GetDocumentChecked();
 		FMetasoundFrontendGraph& Graph = Document.RootGraph.FindGraphChecked(InPageID);
 		TArray<FMetasoundFrontendNode>& Nodes = Graph.Nodes;
-		FMetasoundFrontendNode& Node = Nodes.Emplace_GetRef(*Dependency);
+		FMetasoundFrontendNode& Node = Nodes.Emplace_GetRef(*Dependency, MoveTemp(NodeConfiguration));
 		Node.UpdateID(InNodeID);
 		FinalizeNode(Node, InClassKey);
 
@@ -2622,7 +2623,7 @@ bool FMetaSoundFrontendDocumentBuilder::FindNodeClassInterfaces(const FGuid& InN
 		if (const FMetasoundFrontendClass* NodeClass = DocumentCache->FindDependency(Node->ClassID))
 		{
 			const FNodeRegistryKey NodeClassRegistryKey = FNodeRegistryKey(NodeClass->Metadata);
-			return FMetasoundFrontendRegistryContainer::Get()->FindImplementedInterfacesFromRegistered(NodeClassRegistryKey, OutInterfaces);
+			return INodeClassRegistry::Get()->FindImplementedInterfacesFromRegistered(NodeClassRegistryKey, OutInterfaces);
 		}
 	}
 
diff --git a/Engine/Plugins/Runtime/Metasound/Source/MetasoundFrontend/Private/MetasoundFrontendGraphBuilder.cpp b/Engine/Plugins/Runtime/Metasound/Source/MetasoundFrontend/Private/MetasoundFrontendGraphBuilder.cpp
--- a/Engine/Plugins/Runtime/Metasound/Source/MetasoundFrontend/Private/MetasoundFrontendGraphBuilder.cpp
+++ b/Engine/Plugins/Runtime/Metasound/Source/MetasoundFrontend/Private/MetasoundFrontendGraphBuilder.cpp
@@ -392,7 +392,12 @@ namespace Metasound::Frontend
 
 		FNodeData CreateNodeData(const FBuildContext& InContext, const FMetasoundFrontendNode& InNode, const FMetasoundFrontendClass& InClass, const FMetasoundFrontendClassInput* InOwningGraphClassInput=nullptr, const FMetasoundFrontendClassOutput* InOwningGraphClassOutput=nullptr)
 		{
-			return FNodeData(InNode.Name, InNode.GetID(), CreateVertexInterface(InContext, InNode, InClass, InOwningGraphClassInput, InOwningGraphClassOutput));
+			TSharedPtr<const IOperatorData> Config;
+			if (InNode.Configuration.IsValid())
+			{
+				Config = InNode.Configuration.Get().GetOperatorData();
+			}
+			return FNodeData(InNode.Name, InNode.GetID(), CreateVertexInterface(InContext, InNode, InClass, InOwningGraphClassInput, InOwningGraphClassOutput), MoveTemp(Config));
 		}
 
 		const FMetasoundFrontendVariable* FindVariableForVariableNode(const FMetasoundFrontendNode& InVariableNode, const FMetasoundFrontendGraph& InGraph)
diff --git a/Engine/Plugins/Runtime/Metasound/Source/MetasoundFrontend/Private/MetasoundFrontendGraphController.cpp b/Engine/Plugins/Runtime/Metasound/Source/MetasoundFrontend/Private/MetasoundFrontendGraphController.cpp
--- a/Engine/Plugins/Runtime/Metasound/Source/MetasoundFrontend/Private/MetasoundFrontendGraphController.cpp
+++ b/Engine/Plugins/Runtime/Metasound/Source/MetasoundFrontend/Private/MetasoundFrontendGraphController.cpp
@@ -10,6 +10,7 @@
 #include "MetasoundFrontendDocumentBuilder.h"
 #include "MetasoundFrontendDocumentIdGenerator.h"
 #include "MetasoundFrontendGraph.h"
+#include "MetasoundFrontendNodeClassRegistry.h"
 #include "MetasoundFrontendNodeController.h"
 #include "MetasoundFrontendNodeTemplateRegistry.h"
 #include "MetasoundFrontendProxyDataCache.h"
@@ -1004,7 +1005,7 @@ namespace Metasound
 
 		FNodeHandle FGraphController::AddInputVertex(const FMetasoundFrontendClassInput& InClassInput)
 		{
-			using FRegistry = FMetasoundFrontendRegistryContainer;
+			using FRegistry = INodeClassRegistry;
 
 			if (FMetasoundFrontendGraphClass* GraphClass = GraphClassPtr.Get())
 			{
@@ -1019,8 +1020,9 @@ namespace Metasound
 						{
 							const FVertexName NewName = InClassInput.Name;
 
+							TInstancedStruct<FMetaSoundFrontendNodeConfiguration> Configuration = FRegistry::Get()->CreateFrontendNodeConfiguration(Key);
 							// Setup input node
-							FMetasoundFrontendNode& Node = FindBuildGraphChecked(*GraphClass).Nodes.Emplace_GetRef(*InputClass);
+							FMetasoundFrontendNode& Node = FindBuildGraphChecked(*GraphClass).Nodes.Emplace_GetRef(*InputClass, MoveTemp(Configuration));
 							Node.Name = NewName;
 
 							FGuid NodeID = InClassInput.NodeID;
@@ -1126,7 +1128,7 @@ namespace Metasound
 
 		FNodeHandle FGraphController::AddOutputVertex(const FMetasoundFrontendClassOutput& InClassOutput)
 		{
-			using FRegistry = FMetasoundFrontendRegistryContainer;
+			using FRegistry = INodeClassRegistry;
 
 			if (FMetasoundFrontendGraphClass* GraphClass = GraphClassPtr.Get())
 			{
@@ -1141,7 +1143,9 @@ namespace Metasound
 						{
 							const FVertexName NewName = InClassOutput.Name;
 
-							FMetasoundFrontendNode& Node = FindBuildGraphChecked(*GraphClass).Nodes.Add_GetRef(*OutputClass);
+							TInstancedStruct<FMetaSoundFrontendNodeConfiguration> Configuration = FRegistry::Get()->CreateFrontendNodeConfiguration(Key);
+
+							FMetasoundFrontendNode& Node = FindBuildGraphChecked(*GraphClass).Nodes.Emplace_GetRef(*OutputClass, MoveTemp(Configuration));
 							Node.Name = NewName;
 
 							FGuid NodeID = InClassOutput.NodeID;
@@ -1475,7 +1479,8 @@ namespace Metasound
 				{
 					if (const FMetasoundFrontendClass* NodeClass = Class.Get())
 					{
-						FMetasoundFrontendNode& Node = FindBuildGraphChecked(*GraphClass).Nodes.Emplace_GetRef(*NodeClass);
+						// Currently template nodes do not have node configurations so an invalid one is supplied here. 
+						FMetasoundFrontendNode& Node = FindBuildGraphChecked(*GraphClass).Nodes.Emplace_GetRef(*NodeClass, TInstancedStruct<FMetaSoundFrontendNodeConfiguration>());
 						Node.UpdateID(InNodeGuid);
 						Node.Interface = InTemplate.GenerateNodeInterface(MoveTemp(Params));
 						FNodeAccessPtr NodePtr = GraphClassPtr.GetNodeWithNodeID(Node.GetID());
@@ -1684,7 +1689,12 @@ namespace Metasound
 			{
 				if (const FMetasoundFrontendClass* NodeClass = InExistingDependency.Get())
 				{
-					FMetasoundFrontendNode& Node = FindBuildGraphChecked(*GraphClass).Nodes.Emplace_GetRef(*NodeClass);
+					TInstancedStruct<FMetaSoundFrontendNodeConfiguration> Configuration;
+					if (NodeClass->Metadata.GetType() != EMetasoundFrontendClassType::Graph) // Subgraphs do not have extensions
+					{
+						Configuration = INodeClassRegistry::Get()->CreateFrontendNodeConfiguration(FNodeClassRegistryKey(NodeClass->Metadata));
+					}
+					FMetasoundFrontendNode& Node = FindBuildGraphChecked(*GraphClass).Nodes.Emplace_GetRef(*NodeClass, MoveTemp(Configuration));
 
 					// Cache the asset name on the node if it node is reference to asset-defined graph.
 					// AssetManager may not exist if this is called from a build that does not load the
diff --git a/Engine/Plugins/Runtime/Metasound/Source/MetasoundFrontend/Private/MetasoundFrontendNodeClassRegistry.cpp b/Engine/Plugins/Runtime/Metasound/Source/MetasoundFrontend/Private/MetasoundFrontendNodeClassRegistry.cpp
--- a/Engine/Plugins/Runtime/Metasound/Source/MetasoundFrontend/Private/MetasoundFrontendNodeClassRegistry.cpp
+++ b/Engine/Plugins/Runtime/Metasound/Source/MetasoundFrontend/Private/MetasoundFrontendNodeClassRegistry.cpp
@@ -11,6 +11,7 @@
 #include "HAL/PlatformTime.h"
 #include "Misc/Guid.h"
 #include "Misc/ScopeLock.h"
+#include "StructUtils/InstancedStruct.h"
 #include "Tasks/Pipe.h"
 #include "Tasks/Task.h"
 #include "Templates/SharedPointer.h"
@@ -200,6 +201,11 @@ namespace Metasound::Frontend
 				}
 			}
 
+			virtual TSharedPtr<const IOperatorData> GetOperatorData() const override
+			{
+				return NodeData.OperatorData;
+			}
+
 			virtual FOperatorFactorySharedRef GetDefaultOperatorFactory() const override
 			{
 				return Factory;
@@ -289,6 +295,11 @@ namespace Metasound::Frontend
 				}
 			}
 
+			virtual TInstancedStruct<FMetaSoundFrontendNodeConfiguration> CreateFrontendNodeConfiguration() const override
+			{
+				// Document based nodes do not current offer extensions. 
+				return TInstancedStruct<FMetaSoundFrontendNodeConfiguration>();
+			}
 		private:
 			FMetasoundFrontendClass FrontendClass;
 			TSet<FMetasoundFrontendVersion> Interfaces;
@@ -317,6 +328,18 @@ PRAGMA_DISABLE_DEPRECATION_WARNINGS
 PRAGMA_ENABLE_DEPRECATION_WARNINGS
 	}
 
+#if !UE_METASOUND_PURE_VIRTUAL_CREATE_FRONTEND_NODE_EXTENSION
+	TInstancedStruct<FMetaSoundFrontendNodeConfiguration> INodeClassRegistryEntry::CreateFrontendNodeConfiguration() const
+	{
+		static bool bDidWarn = false;
+		if (!bDidWarn)
+		{
+			UE_LOG(LogMetaSound, Warning, TEXT("Please implement INodeClassRegistryEntry::CreateFrontendNodeConfiguration for the registry entry class representing node %s. This method will become pure virtual in future releases. Define UE_METASOUND_PURE_VIRTUAL_CREATE_FRONTEND_NODE_EXTENSION in order to build with this method as a pure virtual on the interface."), *GetClassInfo().ClassName.ToString());
+		}
+		return TInstancedStruct<FMetaSoundFrontendNodeConfiguration>();
+	}
+#endif
+
 	FNodeClassRegistryTransaction::FNodeClassRegistryTransaction(ETransactionType InType, const FNodeClassInfo& InNodeClassInfo, FNodeClassRegistryTransaction::FTimeType InTimestamp)
 	: Type(InType)
 	, NodeClassInfo(InNodeClassInfo)
@@ -1159,7 +1182,7 @@ PRAGMA_ENABLE_DEPRECATION_WARNINGS
 		return false;
 	}
 
-	bool FNodeClassRegistry::FindFrontendClassFromRegistered(const FNodeClassRegistryKey& InKey, FMetasoundFrontendClass& OutClass)
+	bool FNodeClassRegistry::FindFrontendClassFromRegistered(const FNodeClassRegistryKey& InKey, FMetasoundFrontendClass& OutClass) const
 	{
 		auto SetFrontendClass = [&OutClass](const INodeClassRegistryEntry& Entry)
 		{
@@ -1180,6 +1203,22 @@ PRAGMA_ENABLE_DEPRECATION_WARNINGS
 		return false;
 	}
 
+	TInstancedStruct<FMetaSoundFrontendNodeConfiguration> FNodeClassRegistry::CreateFrontendNodeConfiguration(const FNodeClassRegistryKey& InKey) const
+	{
+		TInstancedStruct<FMetaSoundFrontendNodeConfiguration> NodeConfiguration;
+		auto SetNodeConfiguration = [&NodeConfiguration](const INodeClassRegistryEntry& Entry)
+		{
+			NodeConfiguration = Entry.CreateFrontendNodeConfiguration();
+		};
+
+		AccessNodeEntryThreadSafe(InKey, SetNodeConfiguration);
+
+		// Currently node configuration on template nodes is not supported. To enable that, the node template registry will need 
+		// to provide a creation mechanisms for making related FMetaSoundFrontendNodeConfigurations
+
+		return NodeConfiguration;
+	}
+
 	bool FNodeClassRegistry::FindImplementedInterfacesFromRegistered(const Metasound::Frontend::FNodeClassRegistryKey& InKey, TSet<FMetasoundFrontendVersion>& OutInterfaceVersions) const 
 	{
 		bool bDidCopy = false;
diff --git a/Engine/Plugins/Runtime/Metasound/Source/MetasoundFrontend/Private/MetasoundFrontendNodeClassRegistryPrivate.h b/Engine/Plugins/Runtime/Metasound/Source/MetasoundFrontend/Private/MetasoundFrontendNodeClassRegistryPrivate.h
--- a/Engine/Plugins/Runtime/Metasound/Source/MetasoundFrontend/Private/MetasoundFrontendNodeClassRegistryPrivate.h
+++ b/Engine/Plugins/Runtime/Metasound/Source/MetasoundFrontend/Private/MetasoundFrontendNodeClassRegistryPrivate.h
@@ -96,7 +96,8 @@ namespace Metasound::Frontend
 		// Find Frontend class data.
 		virtual bool FindDefaultVertexInterface(const FNodeClassRegistryKey& InKey, FVertexInterface& OutVertexInterface) const override;
 		
-		virtual bool FindFrontendClassFromRegistered(const FNodeClassRegistryKey& InKey, FMetasoundFrontendClass& OutClass) override;
+		virtual bool FindFrontendClassFromRegistered(const FNodeClassRegistryKey& InKey, FMetasoundFrontendClass& OutClass) const override;
+		virtual TInstancedStruct<FMetaSoundFrontendNodeConfiguration> CreateFrontendNodeConfiguration(const FNodeClassRegistryKey& InKey) const override;
 
 		virtual bool FindImplementedInterfacesFromRegistered(const Metasound::Frontend::FNodeClassRegistryKey& InKey, TSet<FMetasoundFrontendVersion>& OutInterfaceVersions) const override;
 		
diff --git a/Engine/Plugins/Runtime/Metasound/Source/MetasoundFrontend/Private/MetasoundNodeRegistrationMacro.cpp b/Engine/Plugins/Runtime/Metasound/Source/MetasoundFrontend/Private/MetasoundNodeRegistrationMacro.cpp
--- a/Engine/Plugins/Runtime/Metasound/Source/MetasoundFrontend/Private/MetasoundNodeRegistrationMacro.cpp
+++ b/Engine/Plugins/Runtime/Metasound/Source/MetasoundFrontend/Private/MetasoundNodeRegistrationMacro.cpp
@@ -42,6 +42,11 @@ namespace Metasound::Frontend
 			return ClassMetadata->DefaultInterface;
 		}
 
+		TInstancedStruct<FMetaSoundFrontendNodeConfiguration> FNodeRegistryEntryBase::CreateFrontendNodeConfiguration() const
+		{
+			return TInstancedStruct<FMetaSoundFrontendNodeConfiguration>();
+		}
+
 		TSharedRef<const FNodeClassMetadata> FNodeRegistryEntryBase::GetNodeClassMetadata() const
 		{
 			return ClassMetadata;
diff --git a/Engine/Plugins/Runtime/Metasound/Source/MetasoundFrontend/Public/MetasoundFrontendDocument.h b/Engine/Plugins/Runtime/Metasound/Source/MetasoundFrontend/Public/MetasoundFrontendDocument.h
--- a/Engine/Plugins/Runtime/Metasound/Source/MetasoundFrontend/Public/MetasoundFrontendDocument.h
+++ b/Engine/Plugins/Runtime/Metasound/Source/MetasoundFrontend/Public/MetasoundFrontendDocument.h
@@ -8,9 +8,12 @@
 #include "Internationalization/Text.h"
 #include "MetasoundAccessPtr.h"
 #include "MetasoundFrontendLiteral.h"
+#include "MetasoundOperatorData.h"
 #include "MetasoundNodeInterface.h"
 #include "MetasoundVertex.h"
 #include "Misc/Guid.h"
+#include "StructUtils/InstancedStruct.h"
+#include "StructUtils/SharedStruct.h"
 #include "Templates/Function.h"
 #include "Templates/Invoke.h"
 #include "Templates/TypeHash.h"
@@ -508,6 +511,24 @@ struct METASOUNDFRONTEND_API FMetasoundFrontendNodeInterface
 	TArray<FMetasoundFrontendVertex> Environment;
 };
 
+USTRUCT()
+struct METASOUNDFRONTEND_API FMetaSoundFrontendNodeConfiguration
+{
+	GENERATED_BODY()
+
+	virtual ~FMetaSoundFrontendNodeConfiguration() = default;
+
+	/* Get the current interface for the class based upon the node extension */
+	virtual FMetasoundFrontendNodeInterface GetNodeInterface(const FMetasoundFrontendClass& InNodeClass) const;
+
+	/* Reassign existing connections if the node's interface has changed */
+	//virtual ReassignConnections(const TMap<FName, FGuid>& OriginalInputs, TMap<FName, FGuid>& NewInputs, OriginalOutputs, NewOutputs);
+
+	/** Provide any data needed by IOperators instantiated from this node. */
+	virtual TSharedPtr<const Metasound::IOperatorData> GetOperatorData() const;
+};
+
+
 // DEPRECATED in Document Model v1.1
 UENUM()
 enum class EMetasoundFrontendNodeStyleDisplayVisibility : uint8
@@ -575,9 +596,12 @@ struct METASOUNDFRONTEND_API FMetasoundFrontendNode
 
 	FMetasoundFrontendNode() = default;
 
-	// Construct node to satisfy class. 
+	UE_DEPRECATED(5.6, "Please use constructor which accepts a node extension.")
 	FMetasoundFrontendNode(const FMetasoundFrontendClass& InClass);
 
+	// Construct a node from a node class and an optional node extension 
+	FMetasoundFrontendNode(const FMetasoundFrontendClass& InClass, TInstancedStruct<FMetaSoundFrontendNodeConfiguration> InConfiguration);
+
 private:
 	// Unique ID of this node.
 	UPROPERTY()
@@ -600,6 +624,9 @@ public:
 	UPROPERTY()
 	TArray<FMetasoundFrontendVertexLiteral> InputLiterals;
 
+	UPROPERTY()
+	TInstancedStruct<FMetaSoundFrontendNodeConfiguration> Configuration;
+
 #if WITH_EDITORONLY_DATA
 	// Style info related to a node.
 	UPROPERTY()
diff --git a/Engine/Plugins/Runtime/Metasound/Source/MetasoundFrontend/Public/MetasoundFrontendNodeClassRegistry.h b/Engine/Plugins/Runtime/Metasound/Source/MetasoundFrontend/Public/MetasoundFrontendNodeClassRegistry.h
--- a/Engine/Plugins/Runtime/Metasound/Source/MetasoundFrontend/Public/MetasoundFrontendNodeClassRegistry.h
+++ b/Engine/Plugins/Runtime/Metasound/Source/MetasoundFrontend/Public/MetasoundFrontendNodeClassRegistry.h
@@ -15,6 +15,10 @@
 #define UE_METASOUND_ENABLE_5_6_NODE_REGISTRATION_DEPRECATION_WARNINGS (0)
 #endif
 
+#ifndef UE_METASOUND_PURE_VIRTUAL_CREATE_FRONTEND_NODE_EXTENSION
+#define UE_METASOUND_PURE_VIRTUAL_CREATE_FRONTEND_NODE_EXTENSION (0)
+#endif
+
 // Forward declare
 class UObject;
 
@@ -108,6 +112,13 @@ namespace Metasound::Frontend
 		UE_DEPRECATED(5.6, "Node class registry no longer tracks nature of implementation (i.e. if was registered "
 			"from an asset or c++. Use the MetaSoundAssetManager to determine if the class has been defined within an asset(s).")
 		virtual bool IsNative() const { return false; }
+
+		/** Optionally create the node extension associated with the node. */
+#if UE_METASOUND_PURE_VIRTUAL_CREATE_FRONTEND_NODE_EXTENSION
+		virtual TInstancedStruct<FMetaSoundFrontendNodeConfiguration> CreateFrontendNodeConfiguration() const = 0;
+#else
+		virtual TInstancedStruct<FMetaSoundFrontendNodeConfiguration> CreateFrontendNodeConfiguration() const;
+#endif
 	};
 
 	using INodeRegistryEntry UE_DEPRECATED(5.6, "Use INodeClassRegistryEntry instead") = INodeClassRegistryEntry;
@@ -323,7 +334,10 @@ namespace Metasound::Frontend
 		// returning true on success, false otherwise 
 		virtual bool FindDefaultVertexInterface(const FNodeClassRegistryKey& InKey, FVertexInterface& OutVertexInterface) const = 0;
 
-		virtual bool FindFrontendClassFromRegistered(const FNodeClassRegistryKey& InKey, FMetasoundFrontendClass& OutClass) = 0;
+		virtual bool FindFrontendClassFromRegistered(const FNodeClassRegistryKey& InKey, FMetasoundFrontendClass& OutClass) const = 0;
+
+		/** Return the node extension associated with the node. If there is no extension assoicated with the node, the returned TInstancedStruct will be invalid. */
+		virtual TInstancedStruct<FMetaSoundFrontendNodeConfiguration> CreateFrontendNodeConfiguration(const FNodeClassRegistryKey& InKey) const = 0;
 
 		virtual bool FindImplementedInterfacesFromRegistered(const FNodeClassRegistryKey& InKey, TSet<FMetasoundFrontendVersion>& OutInterfaceVersions) const = 0;
 		
diff --git a/Engine/Plugins/Runtime/Metasound/Source/MetasoundFrontend/Public/MetasoundNodeRegistrationMacro.h b/Engine/Plugins/Runtime/Metasound/Source/MetasoundFrontend/Public/MetasoundNodeRegistrationMacro.h
--- a/Engine/Plugins/Runtime/Metasound/Source/MetasoundFrontend/Public/MetasoundNodeRegistrationMacro.h
+++ b/Engine/Plugins/Runtime/Metasound/Source/MetasoundFrontend/Public/MetasoundNodeRegistrationMacro.h
@@ -41,6 +41,8 @@ namespace Metasound::Frontend
 
 			virtual FVertexInterface GetDefaultVertexInterface() const override;
 
+			virtual TInstancedStruct<FMetaSoundFrontendNodeConfiguration> CreateFrontendNodeConfiguration() const override;
+
 		protected:
 
 			TSharedRef<const FNodeClassMetadata> GetNodeClassMetadata() const;
@@ -52,7 +54,7 @@ namespace Metasound::Frontend
 		};
 
 		template<typename TNodeType>
-		class TNodeRegistryEntry : public FNodeRegistryEntryBase
+		class TNodeRegistryEntryBase : public FNodeRegistryEntryBase
 		{
 		public:
 			// Expose FNodeRegistryEntryBase constructors. 
@@ -82,6 +84,31 @@ namespace Metasound::Frontend
 			}
 		};
 
+		// A node registry entry which also provides a node extension. 
+		template<typename NodeType, typename ConfigurationType> 
+		class TNodeRegistryEntry : public TNodeRegistryEntryBase<NodeType>
+		{
+			static_assert(std::is_base_of_v<FMetaSoundFrontendNodeConfiguration, ConfigurationType>, "Configurations must inherit from FMetaSoundFrontendNodeConfiguration");
+		public:
+			// Inherit constructor
+			using TNodeRegistryEntryBase<NodeType>::TNodeRegistryEntryBase;
+
+			virtual TInstancedStruct<FMetaSoundFrontendNodeConfiguration> CreateFrontendNodeConfiguration() const override
+			{
+				return TInstancedStruct<FMetaSoundFrontendNodeConfiguration>::Make<ConfigurationType>();
+			}
+
+		};
+
+		// A partial template specialization for scenario where no node extension is provided
+		template<typename NodeType>
+		class TNodeRegistryEntry<NodeType, void> : public TNodeRegistryEntryBase<NodeType>
+		{
+		public:
+			// Inherit constructor
+			using TNodeRegistryEntryBase<NodeType>::TNodeRegistryEntryBase;
+		};
+
 #if UE_METASOUND_ENABLE_5_6_NODE_REGISTRATION_DEPRECATION_WARNINGS
 		template<typename T>
 		UE_DEPRECATED(5.6, "Update the node's constructor to be of the form FMyNode(FNodeData InNodeData, TSharedRef<const FNodeClassMetadata> InClassMetadata)")
@@ -136,46 +163,54 @@ namespace Metasound::Frontend
 				}
 			}
 		};
-
-	}
+	} // namespace NodeRegistrationPrivate
 
 
-	template <typename TNodeType>
-	bool RegisterNode(const Metasound::FNodeClassMetadata& InMetadata)
+	template <typename TNodeType, typename ConfigurationType=void>
+	bool RegisterNode(const FNodeClassMetadata& InMetadata)
 	{
 		using namespace NodeRegistrationPrivate;
 		if constexpr (TIsOnlyDeprecatedNodeConstructorProvided<TNodeType>::Value)
 		{
-			return RegisterNode<TDeprecatedNodeAPIAdapter<TNodeType>>(InMetadata);
+			return RegisterNode<TDeprecatedNodeAPIAdapter<TNodeType>, ConfigurationType>(InMetadata);
 		}
 		else
 		{
-			// if we reenter this code (because DECLARE_METASOUND_DATA_REFERENCE_TYPES was called twice with the same type),
-			// we catch it here.
-			static bool bAlreadyRegisteredThisDataType = false;
-			if (bAlreadyRegisteredThisDataType)
-			{
-				UE_LOG(LogMetaSound, Display, TEXT("Tried to call METASOUND_REGISTER_NODE twice with the same class. ignoring the second call. Likely because METASOUND_REGISTER_NODE is in a header that's used in multiple modules. Consider moving it to a private header or cpp file."));
-				return false;
-			}
-
-			bAlreadyRegisteredThisDataType = true;
-
-
-			Frontend::FNodeRegistryKey Key = FMetasoundFrontendRegistryContainer::Get()->RegisterNode(MakeUnique<TNodeRegistryEntry<TNodeType>>(InMetadata));
+			Frontend::FNodeRegistryKey Key = INodeClassRegistry::Get()->RegisterNode(MakeUnique<TNodeRegistryEntry<TNodeType, ConfigurationType>>(InMetadata));
 			const bool bSuccessfullyRegisteredNode = Key.IsValid();
 			ensureAlwaysMsgf(bSuccessfullyRegisteredNode, TEXT("Registering node class failed. Please check the logs."));
 
 			return bSuccessfullyRegisteredNode;
 		}
-	}
+	} 
 
-	template <typename TNodeType>
+	template <typename TNodeType, typename ConfigurationType=void>
 	bool RegisterNode()
 	{
+		 static_assert(::Metasound::TIsNodeConstructorSupported<TNodeType>::Value, "In order to be registered as a MetaSound node, the node needs to implement the following public constructor: Constructor(Metasound::FNodeData InNodeData) or Construct(Metasound::FNodeData InNodeData, TSharedRef<const FNodeClassMetadata> InMetadata)"); 
+
+		if constexpr(TIsCreateNodeClassMetadataDeclared<TNodeType>::Value)
+		{
+			return RegisterNode<TNodeType, ConfigurationType>(TNodeType::CreateNodeClassMetadata());
+		}
+		else
+		{
+			using namespace NodeRegistrationPrivate;
+			TriggerMissingCreateNodeClassMetadataWarning<TNodeType>();
+			// Register a node using a prototype node.
+			FNodeInitData InitData;
+			TUniquePtr<Metasound::INode> Node = MakeUnique<TNodeType>(InitData);
+			return RegisterNode<TNodeType, ConfigurationType>(Node->GetMetadata());
+		}
+	}
+
+	template<typename TNodeType>
+	bool UnregisterNode()
+	{
+		FNodeClassRegistryKey RegistryKey;
 		if constexpr(TIsCreateNodeClassMetadataDeclared<TNodeType>::Value)
 		{
-			return RegisterNode<TNodeType>(TNodeType::CreateNodeClassMetadata());
+			RegistryKey = FNodeClassRegistryKey(TNodeType::CreateNodeClassMetadata());
 		}
 		else
 		{
@@ -184,9 +219,31 @@ namespace Metasound::Frontend
 			// Register a node using a prototype node.
 			FNodeInitData InitData;
 			TUniquePtr<Metasound::INode> Node = MakeUnique<TNodeType>(InitData);
-			return RegisterNode<TNodeType>(Node->GetMetadata());
+			RegistryKey = FNodeClassRegistryKey(Node->GetMetadata());
 		}
+
+		return INodeClassRegistry::Get()->UnregisterNode(RegistryKey);
 	}
+
+	namespace NodeRegistrationPrivate
+	{
+		template<typename NodeType, typename ConfigurationType=void>
+		bool RegisterNodeUsingMacro()
+		{
+			// if we reenter this code (because METASOUND_REGISTER_NODE was called twice with the same node type),
+			// we catch it here.
+			static bool bAlreadyRegisteredThisDataType = false;
+			if (bAlreadyRegisteredThisDataType)
+			{
+				UE_LOG(LogMetaSound, Display, TEXT("Tried to call METASOUND_REGISTER_NODE twice with the same class. ignoring the second call. Likely because METASOUND_REGISTER_NODE is in a header that's used in multiple modules. Consider moving it to a private header or cpp file."));
+				return false;
+			}
+
+			bAlreadyRegisteredThisDataType = true;
+
+			return RegisterNode<NodeType, ConfigurationType>();
+		}
+	} // namespace NodeRegistrationPrivate
 }
 
 namespace Metasound
@@ -206,26 +263,27 @@ namespace Metasound
 		static const bool Value = sizeof(TestForConstructor<NodeClass>(nullptr)) == sizeof(uint16);
 	};
 
-	template<typename T>
+	template<typename T, typename ConfigurationType=void>
 	UE_DEPRECATED(5.6, "Use Frontend::RegisterNode()")
 	bool RegisterNodeWithFrontend()
 	{
-		return Frontend::RegisterNode<T>();
+		return Frontend::RegisterNode<T, ConfigurationType>();
 	}
 
-	template<typename T>
+	template<typename T, typename ConfigurationType=void>
 	UE_DEPRECATED(5.6, "Use Frontend::RegisterNode(const Metasound::FNodeClassMetadata&)")
 	bool RegisterNodeWithFrontend(const Metasound::FNodeClassMetadata& InMetadata)
 	{
-		return Frontend::RegisterNode<T>(InMetadata);
+		return Frontend::RegisterNode<T, ConfigurationType>(InMetadata);
 	}
 }
 
-#define METASOUND_REGISTER_NODE(NodeClass, ...) \
+#define METASOUND_REGISTER_NODE_AND_EXTENSION(NodeClass, ConfigurationClass) \
 	 static_assert(std::is_base_of<::Metasound::INodeBase, NodeClass>::value, "To be registered as a  Metasound Node," #NodeClass "need to be a derived class from Metasound::INodeBase, Metasound::INode, or Metasound::FNode."); \
 	 static_assert(::Metasound::TIsNodeConstructorSupported<NodeClass>::Value, "In order to be registered as a Metasound Node, " #NodeClass " needs to implement the following public constructor: " #NodeClass "(Metasound::FNodeData InNodeData);"); \
-	 static bool bSuccessfullyRegistered##NodeClass  = ::Metasound::Frontend::INodeClassRegistry::Get()->EnqueueInitCommand([](){ ::Metasound::Frontend::RegisterNode<NodeClass>(__VA_ARGS__); }); // This static bool is useful for debugging, but also is the only way the compiler will let us call this function outside of an expression.
+	 static bool bSuccessfullyRegistered##NodeClass  = ::Metasound::Frontend::INodeClassRegistry::Get()->EnqueueInitCommand([](){ ::Metasound::Frontend::NodeRegistrationPrivate::RegisterNodeUsingMacro<NodeClass, ConfigurationClass>(); }); // This static bool is useful for debugging, but also is the only way the compiler will let us call this function outside of an expression.
 
+#define METASOUND_REGISTER_NODE(NodeClass) METASOUND_REGISTER_NODE_AND_EXTENSION(NodeClass, void)
 /*
 Macros to help define various FText node fields.
 */
diff --git a/Engine/Plugins/Runtime/Metasound/Source/MetasoundGraphCore/Private/MetasoundBasicNode.cpp b/Engine/Plugins/Runtime/Metasound/Source/MetasoundGraphCore/Private/MetasoundBasicNode.cpp
--- a/Engine/Plugins/Runtime/Metasound/Source/MetasoundGraphCore/Private/MetasoundBasicNode.cpp
+++ b/Engine/Plugins/Runtime/Metasound/Source/MetasoundGraphCore/Private/MetasoundBasicNode.cpp
@@ -35,6 +35,11 @@ namespace Metasound
 	{
 		return NodeData.Interface;
 	}
+	
+	TSharedPtr<const IOperatorData> FBasicNode::GetOperatorData() const 
+	{
+		return NodeData.OperatorData;
+	}
 
 	void FBasicNode::SetDefaultInput(const FVertexName& InVertexName, const FLiteral& InLiteral)
 	{
diff --git a/Engine/Plugins/Runtime/Metasound/Source/MetasoundGraphCore/Private/MetasoundGraph.cpp b/Engine/Plugins/Runtime/Metasound/Source/MetasoundGraphCore/Private/MetasoundGraph.cpp
--- a/Engine/Plugins/Runtime/Metasound/Source/MetasoundGraphCore/Private/MetasoundGraph.cpp
+++ b/Engine/Plugins/Runtime/Metasound/Source/MetasoundGraphCore/Private/MetasoundGraph.cpp
@@ -20,10 +20,11 @@ namespace Metasound
 		return MutableNode;
 	}
 
-	FGraph::FGraph(const FString& InInstanceName, const FGuid& InInstanceID)
+	FGraph::FGraph(const FString& InInstanceName, const FGuid& InInstanceID, TSharedPtr<const IOperatorData> InOperatorData)
 	: InstanceName(InInstanceName)
 	, InstanceID(InInstanceID)
 	, Metadata(FNodeClassMetadata::GetEmpty())
+	, OperatorData(MoveTemp(InOperatorData))
 	{
 	}
 
@@ -384,6 +385,11 @@ PRAGMA_ENABLE_DEPRECATION_WARNINGS
 		}
 	}
 
+	TSharedPtr<const IOperatorData> FGraph::GetOperatorData() const
+	{
+		return OperatorData;
+	}
+
 	bool FGraph::SetVertexInterface(const FVertexInterface& InInterface)
 	{
 		return InInterface == Metadata.DefaultInterface;
diff --git a/Engine/Plugins/Runtime/Metasound/Source/MetasoundGraphCore/Private/MetasoundNode.cpp b/Engine/Plugins/Runtime/Metasound/Source/MetasoundGraphCore/Private/MetasoundNode.cpp
--- a/Engine/Plugins/Runtime/Metasound/Source/MetasoundGraphCore/Private/MetasoundNode.cpp
+++ b/Engine/Plugins/Runtime/Metasound/Source/MetasoundGraphCore/Private/MetasoundNode.cpp
@@ -6,10 +6,11 @@
 
 namespace Metasound
 {
-	FNode::FNode(const FVertexName& InInstanceName, const FGuid& InInstanceID, const FNodeClassMetadata& InInfo)
+	FNode::FNode(const FVertexName& InInstanceName, const FGuid& InInstanceID, const FNodeClassMetadata& InInfo, TSharedPtr<const IOperatorData> InOperatorData)
 	: InstanceName(InInstanceName)
 	, InstanceID(InInstanceID)
 	, Info(InInfo)
+	, OperatorData(MoveTemp(InOperatorData))
 	{
 	}
 
@@ -47,4 +48,9 @@ namespace Metasound
 			UE_LOG(LogMetaSound, Error, TEXT("Cannot set default input. Vertex %s does not exist on node %s"), *InVertexName.ToString(), *InstanceName.ToString());
 		}
 	}
+
+	TSharedPtr<const IOperatorData> FNode::GetOperatorData() const 
+	{
+		return OperatorData;
+	}
 }
diff --git a/Engine/Plugins/Runtime/Metasound/Source/MetasoundGraphCore/Private/MetasoundNodeInterface.cpp b/Engine/Plugins/Runtime/Metasound/Source/MetasoundGraphCore/Private/MetasoundNodeInterface.cpp
--- a/Engine/Plugins/Runtime/Metasound/Source/MetasoundGraphCore/Private/MetasoundNodeInterface.cpp
+++ b/Engine/Plugins/Runtime/Metasound/Source/MetasoundGraphCore/Private/MetasoundNodeInterface.cpp
@@ -2,12 +2,23 @@
 
 #include "MetasoundNodeInterface.h"
 
+#include "Containers/UnrealString.h"
+#include "Internationalization/Text.h"
+#include "Misc/Guid.h"
+#include "UObject/NameTypes.h"
+
 namespace Metasound
 {
 	FNodeData::FNodeData(const FName& InName, const FGuid& InID, FVertexInterface InInterface)
+	: FNodeData(InName, InID, MoveTemp(InInterface), TSharedPtr<const IOperatorData>())
+	{
+	}
+
+	FNodeData::FNodeData(const FName& InName, const FGuid& InID, FVertexInterface InInterface, TSharedPtr<const IOperatorData> InOperatorData)
 	: Name(InName)
 	, ID(InID)
 	, Interface(MoveTemp(InInterface))
+	, OperatorData(MoveTemp(InOperatorData))
 	{
 	}
 
@@ -168,15 +179,28 @@ namespace Metasound
 		}
 	}
 
-#if !METASOUND_PURE_VIRTUAL_SETDEFAULTINPUT
+#if !UE_METASOUND_PURE_VIRTUAL_SET_DEFAULT_INPUT
 	void INodeBase::SetDefaultInput(const FVertexName& InVertexName, const FLiteral& InLiteral)
 	{
 		static bool bDidWarn = false;
 		if (!bDidWarn)
 		{
-			UE_LOG(LogMetaSound, Warning, TEXT("Ignoring input default for vertex %s. Please implement INodeInterface::SetDefaultInput(...) for the class representing node %s. This method will become pure virtual on or after release 5.8. Define METASOUND_PURE_VIRTUAL_SETDEFAULTINPUT in order to build with this method as a pure virtual on the interface."), *InVertexName.ToString(), *GetInstanceName().ToString());
+			UE_LOG(LogMetaSound, Warning, TEXT("Ignoring input default for vertex %s. Please implement INodeInterface::SetDefaultInput(...) for the class representing node %s. This method will become pure virtual in future releases. Define UE_METASOUND_PURE_VIRTUAL_SET_DEFAULT_INPUT in order to build with this method as a pure virtual on the interface."), *InVertexName.ToString(), *GetInstanceName().ToString());
+			bDidWarn = true;
+		}
+	}
+#endif
+
+#if !UE_METASOUND_PURE_VIRTUAL_GET_OPERATOR_DATA
+	TSharedPtr<const IOperatorData> INodeBase::GetOperatorData() const
+	{
+		static bool bDidWarn = false;
+		if (!bDidWarn)
+		{
+			UE_LOG(LogMetaSound, Warning, TEXT("Please implement INodeInterface::GetOperatorData(...) for the class representing node %s. This method will become pure virtual in future releases. Define UE_METASOUND_PURE_VIRTUAL_GET_OPERATOR_DATA in order to build with this method as a pure virtual on the interface."), *GetInstanceName().ToString());
 			bDidWarn = true;
 		}
+		return {};
 	}
 #endif
 }
diff --git a/Engine/Plugins/Runtime/Metasound/Source/MetasoundGraphCore/Public/MetasoundBasicNode.h b/Engine/Plugins/Runtime/Metasound/Source/MetasoundGraphCore/Public/MetasoundBasicNode.h
--- a/Engine/Plugins/Runtime/Metasound/Source/MetasoundGraphCore/Public/MetasoundBasicNode.h
+++ b/Engine/Plugins/Runtime/Metasound/Source/MetasoundGraphCore/Public/MetasoundBasicNode.h
@@ -34,6 +34,9 @@ namespace Metasound
 			/** Set the default literal for an input vertex. */
 			virtual void SetDefaultInput(const FVertexName& InVertexName, const FLiteral& InLiteral) override;
 
+			/** Return the configuration for this node. */
+			virtual TSharedPtr<const IOperatorData> GetOperatorData() const override;
+
 		private:
 
 			FNodeData NodeData;
diff --git a/Engine/Plugins/Runtime/Metasound/Source/MetasoundGraphCore/Public/MetasoundBuilderInterface.h b/Engine/Plugins/Runtime/Metasound/Source/MetasoundGraphCore/Public/MetasoundBuilderInterface.h
--- a/Engine/Plugins/Runtime/Metasound/Source/MetasoundGraphCore/Public/MetasoundBuilderInterface.h
+++ b/Engine/Plugins/Runtime/Metasound/Source/MetasoundGraphCore/Public/MetasoundBuilderInterface.h
@@ -64,6 +64,16 @@ namespace Metasound
 	 */
 	struct FBuildOperatorParams
 	{
+		FBuildOperatorParams(const INode& InNode, const FOperatorSettings& InOperatorSettings, const FInputVertexInterfaceData& InInputData, const FMetasoundEnvironment& InEnvironment, const IOperatorBuilder* InBuilder=nullptr, FGraphRenderCost* InGraphRenderCost=nullptr)
+		: Node(InNode)
+		, OperatorSettings(InOperatorSettings)
+		, InputData(InInputData)
+		, Environment(InEnvironment)
+		, Builder(InBuilder)
+		, GraphRenderCost(InGraphRenderCost)
+		{
+		}
+
 		/** The node associated with this factory and the desired IOperator. */
 		const INode& Node;
 
@@ -89,7 +99,7 @@ namespace Metasound
 		}
 	};
 
-	struct UE_DEPRECATED(5.6, "Use FBuildGrpahOperatorParams instead") FBuildGraphParams;
+	struct UE_DEPRECATED(5.6, "Use FBuildGraphOperatorParams instead") FBuildGraphParams;
 	/** Parameters for building an operator from a graph.  */
 	struct FBuildGraphParams
 	{
@@ -109,6 +119,15 @@ namespace Metasound
 	/** Parameters for building an operator from a graph. */
 	struct FBuildGraphOperatorParams
 	{
+		FBuildGraphOperatorParams(const IGraph& InGraph, const FOperatorSettings& InOperatorSettings, const FInputVertexInterfaceData& InInputData, const FMetasoundEnvironment& InEnvironment, FGraphRenderCost* InGraphRenderCost=nullptr)
+		: Graph(InGraph)
+		, OperatorSettings(InOperatorSettings)
+		, InputData(InInputData)
+		, Environment(InEnvironment)
+		, GraphRenderCost(InGraphRenderCost)
+		{
+		}
+
 		/** Reference to graph being built. */
 		const IGraph& Graph;
 
@@ -122,7 +141,7 @@ namespace Metasound
 		const FMetasoundEnvironment& Environment;
 
 		/** Runtime render cost tracker. */
-		FGraphRenderCost* GraphRenderCost;
+		FGraphRenderCost* GraphRenderCost = nullptr;
 
 		FName ClassName;
 		FName AssetName;
diff --git a/Engine/Plugins/Runtime/Metasound/Source/MetasoundGraphCore/Public/MetasoundFacade.h b/Engine/Plugins/Runtime/Metasound/Source/MetasoundGraphCore/Public/MetasoundFacade.h
--- a/Engine/Plugins/Runtime/Metasound/Source/MetasoundGraphCore/Public/MetasoundFacade.h
+++ b/Engine/Plugins/Runtime/Metasound/Source/MetasoundGraphCore/Public/MetasoundFacade.h
@@ -46,6 +46,10 @@ namespace Metasound
 				template<typename T>
 				static uint8 Check(Helper<const FNodeClassMetadata&(*)(), &T::GetNodeInfo>*);
 
+				// Check for "static FNodeClassMetadata U::GetNodeInfo()"
+				template<typename T>
+				static uint8 Check(Helper<FNodeClassMetadata(*)(), &T::GetNodeInfo>*);
+
 				template<typename T> static uint16 Check(...);
 
 			public:
diff --git a/Engine/Plugins/Runtime/Metasound/Source/MetasoundGraphCore/Public/MetasoundGraph.h b/Engine/Plugins/Runtime/Metasound/Source/MetasoundGraphCore/Public/MetasoundGraph.h
--- a/Engine/Plugins/Runtime/Metasound/Source/MetasoundGraphCore/Public/MetasoundGraph.h
+++ b/Engine/Plugins/Runtime/Metasound/Source/MetasoundGraphCore/Public/MetasoundGraph.h
@@ -18,7 +18,7 @@ namespace Metasound
 	{
 		public:
 
-			FGraph(const FString& InInstanceName, const FGuid& InInstanceID);
+			FGraph(const FString& InInstanceName, const FGuid& InInstanceID, TSharedPtr<const IOperatorData> InNodeConfig={});
 
 			virtual ~FGraph() = default;
 
@@ -40,6 +40,9 @@ namespace Metasound
 			/** Set the default value associated with an input vertex. */
 			virtual void SetDefaultInput(const FVertexName& InVertexName, const FLiteral& InLiteral) override;
 
+			/** Return the configuration for this graph. */
+			virtual TSharedPtr<const IOperatorData> GetOperatorData() const override;
+
 			/** Set the vertex interface. If the vertex was successfully changed, returns true. 
 			 *
 			 * @param InInterface - New interface for node. 
@@ -246,6 +249,7 @@ namespace Metasound
 			FGuid InstanceID;
 
 			FNodeClassMetadata Metadata;
+			TSharedPtr<const IOperatorData> OperatorData;
 
 			TArray<FDataEdge> Edges;
 
diff --git a/Engine/Plugins/Runtime/Metasound/Source/MetasoundGraphCore/Public/MetasoundNode.h b/Engine/Plugins/Runtime/Metasound/Source/MetasoundGraphCore/Public/MetasoundNode.h
--- a/Engine/Plugins/Runtime/Metasound/Source/MetasoundGraphCore/Public/MetasoundNode.h
+++ b/Engine/Plugins/Runtime/Metasound/Source/MetasoundGraphCore/Public/MetasoundNode.h
@@ -4,14 +4,17 @@
 #include "CoreMinimal.h"
 #include "MetasoundNodeInterface.h"
 #include "MetasoundVertex.h"
+#include "Templates/SharedPointer.h"
 
 
 namespace Metasound
 {
+	class IOperatorData;
+
 	class METASOUNDGRAPHCORE_API FNode : public INode
 	{
 		public:
-			FNode(const FVertexName& InInstanceName, const FGuid& InInstanceID, const FNodeClassMetadata& InInfo);
+			FNode(const FVertexName& InInstanceName, const FGuid& InInstanceID, const FNodeClassMetadata& InInfo, TSharedPtr<const IOperatorData> InOperatorData={});
 
 			virtual ~FNode() = default;
 
@@ -28,10 +31,12 @@ namespace Metasound
 
 			virtual void SetDefaultInput(const FVertexName& InVertexName, const FLiteral& InLiteral) override;
 
+			virtual TSharedPtr<const IOperatorData> GetOperatorData() const override;
 		private:
 
 			FVertexName InstanceName;
 			FGuid InstanceID;
 			FNodeClassMetadata Info;
+			TSharedPtr<const IOperatorData> OperatorData;
 	};
 }
diff --git a/Engine/Plugins/Runtime/Metasound/Source/MetasoundGraphCore/Public/MetasoundNodeInterface.h b/Engine/Plugins/Runtime/Metasound/Source/MetasoundGraphCore/Public/MetasoundNodeInterface.h
--- a/Engine/Plugins/Runtime/Metasound/Source/MetasoundGraphCore/Public/MetasoundNodeInterface.h
+++ b/Engine/Plugins/Runtime/Metasound/Source/MetasoundGraphCore/Public/MetasoundNodeInterface.h
@@ -12,14 +12,24 @@
 #include "Misc/Guid.h"
 
 // This define helps in migrating code which has not yet implemented the purevirtual
-// INode::SetDefaultInput(...) which was introduced in UE 5.6. By UE 5.8 or later,
+// INode::SetDefaultInput(...) which was introduced in UE 5.6. In later releases
 // the method will become pure virtual and this macro will no longer exist. 
-#ifndef METASOUND_PURE_VIRTUAL_SETDEFAULTINPUT
-#define METASOUND_PURE_VIRTUAL_SETDEFAULTINPUT 0
+#ifndef UE_METASOUND_PURE_VIRTUAL_SET_DEFAULT_INPUT
+#define UE_METASOUND_PURE_VIRTUAL_SET_DEFAULT_INPUT 0
+#endif
+
+// This define helps in migrating code which has not yet implemented the purevirtual
+// INode::GetOperatorData(...) which was introduced in UE 5.6. In later releases
+// the method will become pure virtual and this macro will no longer exist. 
+#ifndef UE_METASOUND_PURE_VIRTUAL_GET_OPERATOR_DATA
+#define UE_METASOUND_PURE_VIRTUAL_GET_OPERATOR_DATA 0
 #endif
 
 namespace Metasound
 {
+	// Forward declare
+	class IOperatorData;
+
 	extern const FString METASOUNDGRAPHCORE_API PluginAuthor;
 	extern const FText METASOUNDGRAPHCORE_API PluginNodeMissingPrompt;
 
@@ -28,10 +38,12 @@ namespace Metasound
 	{
 		FNodeData() = default;
 		FNodeData(const FName& InName, const FGuid& InID, FVertexInterface InInterface);
+		FNodeData(const FName& InName, const FGuid& InID, FVertexInterface InInterface, TSharedPtr<const IOperatorData> InOperatorData);
 
 		FName Name;
 		FGuid ID;
 		FVertexInterface Interface;
+		TSharedPtr<const IOperatorData> OperatorData;
 	};
 
 	/**
@@ -174,12 +186,19 @@ namespace Metasound
 
 			/** Set the default input literal for a vertex. This literal will be 
 			 * used if nothing is connected to the node's input vertex. */
-#if METASOUND_PURE_VIRTUAL_SETDEFAULTINPUT
+#if UE_METASOUND_PURE_VIRTUAL_SET_DEFAULT_INPUT
 			virtual void SetDefaultInput(const FVertexName& InVertexName, const FLiteral& InLiteral) = 0;
 #else 
 			virtual METASOUNDGRAPHCORE_API void SetDefaultInput(const FVertexName& InVertexName, const FLiteral& InLiteral);
 #endif
 
+			/** Returns the configuration for the node. */
+#if UE_METASOUND_PURE_VIRTUAL_GET_OPERATOR_DATA
+			virtual TSharedPtr<const IOperatorData> GetOperatorData() const = 0;
+#else 
+			virtual METASOUNDGRAPHCORE_API TSharedPtr<const IOperatorData> GetOperatorData() const;
+#endif
+
 			/** Set the vertex interface. If the vertex was successfully changed, returns true. 
 			 *
 			 * @param InInterface - New interface for node. 
@@ -212,7 +231,6 @@ namespace Metasound
 		return MakeShared<FactoryType, ESPMode::ThreadSafe>(Forward<ArgTypes>(Args)...);
 	}
 
-
 	/** INode 
 	 * 
 	 * Interface for all nodes that can create operators. 
@@ -347,4 +365,5 @@ namespace Metasound
 			/** Get vertices which contain output parameters. */
 			virtual const FOutputDataSourceCollection& GetOutputDataSources() const = 0;
 	};
+
 }
diff --git a/Engine/Plugins/Runtime/Metasound/Source/MetasoundGraphCore/Public/MetasoundOperatorData.h b/Engine/Plugins/Runtime/Metasound/Source/MetasoundGraphCore/Public/MetasoundOperatorData.h
new file mode 100644
--- /dev/null
+++ b/Engine/Plugins/Runtime/Metasound/Source/MetasoundGraphCore/Public/MetasoundOperatorData.h
@@ -0,0 +1,116 @@
+// Copyright Epic Games, Inc. All Rights Reserved.
+
+#pragma once
+
+#include "UObject/NameTypes.h"
+#include "Templates/SharedPointer.h"
+#include <type_traits>
+
+namespace Metasound
+{
+	namespace OperatorDataPrivate
+	{
+		// Helper template for determining whether a static OperatorDataTypeName exists on the node configuration class declaration. 
+		template<typename OperatorDataType> 
+		struct THasOperatorDataTypeNameMember
+		{
+		private:
+			template<typename T> 
+			static std::negation<std::is_member_pointer<decltype(&T::OperatorDataTypeName)>> HasStaticMemberHelper(int32);
+
+			template<typename T>
+			static std::false_type HasStaticMemberHelper(...);
+
+			// If there is a static member named T::OperatorDataTypeName, than THasStaticMember<T> will be of type std::true_type
+			template<typename T>
+			using THasStaticMember = decltype(HasStaticMemberHelper<T>(0));
+
+			template<typename T>
+			static std::is_convertible<decltype(T::OperatorDataTypeName), FName> IsConvertibleToFNameHelper(int32);
+
+			template<typename T>
+			static std::false_type IsConvertibleToFNameHelper(...);
+
+			// If the member T::OperatorDataTypeName exists and is convertible to an FName, this will be of std::true_type
+			template<typename T>
+			using TIsConvertibleToFName = decltype(IsConvertibleToFNameHelper<T>(0));
+
+		public:
+			static constexpr bool Value = TIsConvertibleToFName<OperatorDataType>::value && THasStaticMember<OperatorDataType>::value;
+		};
+	}
+
+	/** Base class for node configurations. */
+	class IOperatorData
+	{
+	public:
+		IOperatorData() = default;
+		virtual ~IOperatorData() = default;
+
+		/** Return a typename use for safe downcasting of the IOperatorData. */
+		virtual FName GetOperatorDataTypeName() const = 0;
+	};
+
+	/** Node configuration CRTP. Node configurations should derived from this class
+	 * to enable downcasting. 
+	 *
+	 * class FMyOperatorData : public TOperatorData<FMyOperatorData>
+	 * {
+	 * public:
+	 *    static const FLazyName OperatorDataTypeName;
+	 * };
+	 */
+	template<typename DerivedType>
+	class TOperatorData : public IOperatorData
+	{
+	public:
+
+		TOperatorData() = default;
+		virtual ~TOperatorData() = default;
+
+		virtual FName GetOperatorDataTypeName() const override final 
+		{
+			static_assert(OperatorDataPrivate::THasOperatorDataTypeNameMember<DerivedType>::Value, "Types deriving from TOperatorData must define a static class member which can describes the name of the data type");
+			return DerivedType::OperatorDataTypeName;
+		}
+	};
+
+	/** Returns the type name of the node configuration object. */
+	template<typename OperatorDataType>
+	FName GetStaticOperatorDataTypeName()
+	{
+		return OperatorDataType::OperatorDataTypeName;
+	}
+
+	/** Returns true if the provided node configuration is of the type T. */
+	template<
+		typename OperatorDataType, 
+		typename U=std::enable_if_t<std::is_base_of_v<TOperatorData<std::decay_t<OperatorDataType>>, OperatorDataType>>
+	>
+	bool IsOperatorDataOfType(const IOperatorData& InNodeConfig)
+	{
+		return InNodeConfig.GetOperatorDataTypeName() == GetStaticOperatorDataTypeName<OperatorDataType>();
+	}
+
+	/** Returns a non-null pointer to a downcast node configuration object if the 
+	 * provided node configuration is non-null and of the same derived type. */
+	template< typename DesiredOperatorDataType, typename ProvidedOperatorDataType >
+	std::enable_if_t<
+		std::conjunction_v< 
+			std::is_base_of<TOperatorData<std::decay_t<DesiredOperatorDataType>>, DesiredOperatorDataType>, 
+			std::is_base_of<IOperatorData, ProvidedOperatorDataType>
+		>, 
+	DesiredOperatorDataType*> 
+	CastOperatorData(ProvidedOperatorDataType* InOperatorData)
+	{
+		if (InOperatorData)
+		{
+			if (IsOperatorDataOfType<DesiredOperatorDataType>(*InOperatorData))
+			{
+				return static_cast<DesiredOperatorDataType*>(InOperatorData);
+			}
+		}
+		return nullptr;
+	}
+}
+
