- Migrate remainder of FAssetKey to FMetaSoundAssetKey
- Deprecate IMetaSoundAssetManager::FAssetInfo and utilize smaller FAssetRef struct in older migration code for clarity (FMetaSoundAssetClassInfo/FMetaSoundClassInfo are used in completely different context for caching asset tag data)
#rnx

diff --git a/Engine/Plugins/Runtime/Metasound/Source/MetasoundEditor/Private/MetasoundEditor.cpp b/Engine/Plugins/Runtime/Metasound/Source/MetasoundEditor/Private/MetasoundEditor.cpp
--- a/Engine/Plugins/Runtime/Metasound/Source/MetasoundEditor/Private/MetasoundEditor.cpp
+++ b/Engine/Plugins/Runtime/Metasound/Source/MetasoundEditor/Private/MetasoundEditor.cpp
@@ -2689,7 +2689,7 @@ namespace Metasound
 							IMetasoundEditorModule& MetaSoundEditorModule = FModuleManager::GetModuleChecked<IMetasoundEditorModule>("MetaSoundEditor");
 							if (!MetaSoundEditorModule.IsRestrictedMode())
 							{
-								const FAssetKey AssetKey(Class->Metadata);
+								const FMetaSoundAssetKey AssetKey(Class->Metadata);
 								if (const FMetasoundAssetBase* Asset = IMetaSoundAssetManager::GetChecked().FindAsset(AssetKey))
 								{
 									AssetSubsystem->OpenEditorForAsset(Asset->GetOwningAsset());
diff --git a/Engine/Plugins/Runtime/Metasound/Source/MetasoundEditor/Private/MetasoundEditorGraphBuilder.cpp b/Engine/Plugins/Runtime/Metasound/Source/MetasoundEditor/Private/MetasoundEditorGraphBuilder.cpp
--- a/Engine/Plugins/Runtime/Metasound/Source/MetasoundEditor/Private/MetasoundEditorGraphBuilder.cpp
+++ b/Engine/Plugins/Runtime/Metasound/Source/MetasoundEditor/Private/MetasoundEditorGraphBuilder.cpp
@@ -283,7 +283,7 @@ namespace Metasound
 					const bool bIsAssetClass = IMetaSoundAssetManager::GetChecked().IsAssetClass(Metadata);
 					if (bIsAssetClass)
 					{
-						const FTopLevelAssetPath Path = IMetaSoundAssetManager::GetChecked().FindAssetPath(FAssetKey(Metadata));
+						const FTopLevelAssetPath Path = IMetaSoundAssetManager::GetChecked().FindAssetPath(FMetaSoundAssetKey(Metadata));
 						if (Path.IsValid())
 						{
 							DisplayName = FText::FromName(Path.GetAssetName());
diff --git a/Engine/Plugins/Runtime/Metasound/Source/MetasoundEditor/Private/MetasoundEditorGraphNode.cpp b/Engine/Plugins/Runtime/Metasound/Source/MetasoundEditor/Private/MetasoundEditorGraphNode.cpp
--- a/Engine/Plugins/Runtime/Metasound/Source/MetasoundEditor/Private/MetasoundEditorGraphNode.cpp
+++ b/Engine/Plugins/Runtime/Metasound/Source/MetasoundEditor/Private/MetasoundEditorGraphNode.cpp
@@ -12,6 +12,7 @@
 #include "HAL/IConsoleManager.h"
 #include "Logging/TokenizedMessage.h"
 #include "Metasound.h"
+#include "MetasoundAssetKey.h"
 #include "MetasoundAssetManager.h"
 #include "MetasoundDocumentBuilderRegistry.h"
 #include "MetasoundEditorCommands.h"
@@ -1208,7 +1209,7 @@ void UMetasoundEditorGraphExternalNode::Validate(Metasound::Editor::FGraphNodeVa
 		case EMetasoundFrontendClassType::External:
 		case EMetasoundFrontendClassType::Graph:
 		{
-			const FAssetKey AssetKey(Metadata);
+			const FMetaSoundAssetKey AssetKey(Metadata);
 			if (FMetasoundAssetBase* MetaSoundAsset = IMetaSoundAssetManager::GetChecked().TryLoadAssetFromKey(AssetKey))
 			{
 				if (const UMetasoundEditorGraph* NodeGraph = Cast<UMetasoundEditorGraph>(MetaSoundAsset->GetGraph()))
diff --git a/Engine/Plugins/Runtime/Metasound/Source/MetasoundEditor/Private/MetasoundEditorGraphSchema.cpp b/Engine/Plugins/Runtime/Metasound/Source/MetasoundEditor/Private/MetasoundEditorGraphSchema.cpp
--- a/Engine/Plugins/Runtime/Metasound/Source/MetasoundEditor/Private/MetasoundEditorGraphSchema.cpp
+++ b/Engine/Plugins/Runtime/Metasound/Source/MetasoundEditor/Private/MetasoundEditorGraphSchema.cpp
@@ -16,6 +16,7 @@
 #include "Layout/SlateRect.h"
 #include "Metasound.h"
 #include "MetasoundAssetBase.h"
+#include "MetasoundAssetKey.h"
 #include "MetasoundAssetManager.h"
 #include "MetasoundAssetSubsystem.h"
 #include "MetasoundBuilderSubsystem.h"
@@ -468,9 +469,9 @@ PRAGMA_ENABLE_DEPRECATION_WARNINGS
 				return bCanConnect;
 			}
 
-			virtual Frontend::FAssetKey GetAssetKey() const override
+			virtual FMetaSoundAssetKey GetAssetKey() const override
 			{
-				return Frontend::FAssetKey(ClassInfo.ClassName, ClassInfo.Version);
+				return FMetaSoundAssetKey(ClassInfo.ClassName, ClassInfo.Version);
 			}
 
 			virtual const TArray<FText>& GetCategoryHierarchy() const override
@@ -597,9 +598,9 @@ PRAGMA_ENABLE_DEPRECATION_WARNINGS
 				return bCanConnect;
 			}
 
-			virtual Frontend::FAssetKey GetAssetKey() const override
+			virtual FMetaSoundAssetKey GetAssetKey() const override
 			{
-				return Frontend::FAssetKey(Class.Metadata);
+				return FMetaSoundAssetKey(Class.Metadata);
 			}
 
 			virtual const TArray<FText>& GetCategoryHierarchy() const override
@@ -617,7 +618,7 @@ PRAGMA_ENABLE_DEPRECATION_WARNINGS
 					DisplayName = Metadata.GetDisplayName();
 					if (DisplayName.IsEmptyOrWhitespace())
 					{
-						const FTopLevelAssetPath Path = IMetaSoundAssetManager::GetChecked().FindAssetPath(FAssetKey(Metadata));
+						const FTopLevelAssetPath Path = IMetaSoundAssetManager::GetChecked().FindAssetPath(FMetaSoundAssetKey(Metadata));
 						if (Path.IsValid())
 						{
 							DisplayName = FText::FromName(Path.GetAssetName());
@@ -2658,7 +2659,7 @@ void UMetasoundEditorGraphSchema::GetFunctionActions(FGraphActionMenuBuilder& Ac
 	using namespace Metasound::Frontend;
 
 	const FMetaSoundAssetManager& AssetManager = FMetaSoundAssetManager::GetChecked();
-	const FAssetKey ParentAssetKey(InGraphHandle->GetGraphMetadata());
+	const FMetaSoundAssetKey ParentAssetKey(InGraphHandle->GetGraphMetadata());
 	QueryNodeClasses([&](TUniquePtr<ISchemaQueryResult>&& Result)
 	{
 		if (Result->GetRegistryClassType() != EMetasoundFrontendClassType::External)
@@ -2683,7 +2684,7 @@ void UMetasoundEditorGraphSchema::GetFunctionActions(FGraphActionMenuBuilder& Ac
 		}
 
 		bool bCausesLoop = false;
-		auto IsMatchingKey = [&ParentAssetKey, &bCausesLoop](const FAssetKey& AssetKey)
+		auto IsMatchingKey = [&ParentAssetKey, &bCausesLoop](const FMetaSoundAssetKey& AssetKey)
 		{
 			bCausesLoop = ParentAssetKey == AssetKey;
 		};
@@ -2887,7 +2888,7 @@ void UMetasoundEditorGraphSchema::QueryNodeClasses(TFunctionRef<void(TUniquePtr<
 			{
 				using namespace Metasound::Frontend;
 
-				const FAssetKey AssetKey(ClassInfo.ClassName, ClassInfo.Version);
+				const FMetaSoundAssetKey AssetKey(ClassInfo.ClassName, ClassInfo.Version);
 				const FNodeRegistryKey RegKey(AssetKey);
 				if (!FoundKeys.Contains(RegKey))
 				{
diff --git a/Engine/Plugins/Runtime/Metasound/Source/MetasoundEditor/Public/MetasoundEditorGraphSchema.h b/Engine/Plugins/Runtime/Metasound/Source/MetasoundEditor/Public/MetasoundEditorGraphSchema.h
--- a/Engine/Plugins/Runtime/Metasound/Source/MetasoundEditor/Public/MetasoundEditorGraphSchema.h
+++ b/Engine/Plugins/Runtime/Metasound/Source/MetasoundEditor/Public/MetasoundEditorGraphSchema.h
@@ -47,7 +47,7 @@ namespace Metasound::Editor
 		virtual bool CanConnectInputOfTypeAndAccess(FName TypeName, EMetasoundFrontendVertexAccessType AccessType) const = 0;
 		virtual bool CanConnectOutputOfTypeAndAccess(FName TypeName, EMetasoundFrontendVertexAccessType AccessType) const = 0;
 
-		virtual Frontend::FAssetKey GetAssetKey() const = 0;
+		virtual FMetaSoundAssetKey GetAssetKey() const = 0;
 		virtual const TArray<FText>& GetCategoryHierarchy() const = 0;
 		virtual FText GetDisplayName() const = 0;
 		virtual const TArray<FText>& GetKeywords() const = 0;
diff --git a/Engine/Plugins/Runtime/Metasound/Source/MetasoundEngine/Private/Metasound.cpp b/Engine/Plugins/Runtime/Metasound/Source/MetasoundEngine/Private/Metasound.cpp
--- a/Engine/Plugins/Runtime/Metasound/Source/MetasoundEngine/Private/Metasound.cpp
+++ b/Engine/Plugins/Runtime/Metasound/Source/MetasoundEngine/Private/Metasound.cpp
@@ -200,9 +200,9 @@ void UMetaSoundPatch::PostLoad()
 }
 
 #if WITH_EDITOR
-void UMetaSoundPatch::SetReferencedAssetClasses(TSet<Metasound::Frontend::IMetaSoundAssetManager::FAssetInfo>&& InAssetClasses)
+void UMetaSoundPatch::SetReferencedAssets(TSet<Metasound::Frontend::IMetaSoundAssetManager::FAssetRef>&& InAssetRefs)
 {
-	Metasound::Engine::FAssetHelper::SetReferencedAssetClasses(*this, MoveTemp(InAssetClasses));
+	Metasound::Engine::FAssetHelper::SetReferencedAssets(*this, MoveTemp(InAssetRefs));
 }
 #endif
 
diff --git a/Engine/Plugins/Runtime/Metasound/Source/MetasoundEngine/Private/MetasoundAssetSubsystem.cpp b/Engine/Plugins/Runtime/Metasound/Source/MetasoundEngine/Private/MetasoundAssetSubsystem.cpp
--- a/Engine/Plugins/Runtime/Metasound/Source/MetasoundEngine/Private/MetasoundAssetSubsystem.cpp
+++ b/Engine/Plugins/Runtime/Metasound/Source/MetasoundEngine/Private/MetasoundAssetSubsystem.cpp
@@ -8,8 +8,6 @@
 #include "Engine/AssetManager.h"
 #include "HAL/CriticalSection.h"
 #include "Metasound.h"
-#include "MetasoundAssetBase.h"
-#include "MetasoundAssetManager.h"
 #include "MetasoundBuilderSubsystem.h"
 #include "MetasoundDocumentBuilderRegistry.h"
 #include "MetasoundEngineAsset.h"
@@ -24,9 +22,7 @@
 #include "MetasoundUObjectRegistry.h"
 #include "Misc/CoreDelegates.h"
 #include "Modules/ModuleManager.h"
-#include "UObject/AssetRegistryTagsContext.h"
 #include "UObject/NoExportTypes.h"
-#include "UObject/Object.h"
 
 
 #include UE_INLINE_GENERATED_CPP_BY_NAME(MetasoundAssetSubsystem)
@@ -40,7 +36,7 @@ namespace Metasound::Engine
 
 		std::atomic<bool> InitialAssetScanComplete = false;
 
-		bool RemoveClassInfo(FCriticalSection& MapCritSec, FAssetClassInfoMap& Map, const Frontend::FAssetKey& AssetKey, const FTopLevelAssetPath& InAssetPath)
+		bool RemoveClassInfo(FCriticalSection& MapCritSec, FAssetClassInfoMap& Map, const FMetaSoundAssetKey& AssetKey, const FTopLevelAssetPath& InAssetPath)
 		{
 			using namespace Frontend;
 
@@ -71,11 +67,11 @@ namespace Metasound::Engine
 			return false;
 		}
 
-		Frontend::FAssetKey AddClassInfo(FCriticalSection& MapCritSec, FAssetClassInfoMap& Map, Frontend::FMetaSoundAssetClassInfo&& ClassInfo)
+		FMetaSoundAssetKey AddClassInfo(FCriticalSection& MapCritSec, FAssetClassInfoMap& Map, Frontend::FMetaSoundAssetClassInfo&& ClassInfo)
 		{
 			using namespace Frontend;
 
-			const Frontend::FAssetKey AssetKey(ClassInfo.ClassName, ClassInfo.Version);
+			const FMetaSoundAssetKey AssetKey(ClassInfo.ClassName, ClassInfo.Version);
 			if (!AssetKey.IsValid())
 			{
 				return AssetKey;
@@ -110,7 +106,7 @@ namespace Metasound::Engine
 
 		if (bLogActiveAssetsOnShutdown)
 		{
-			TMap<FAssetKey, TArray<FMetaSoundAssetClassInfo>> InfoOnShutdown;
+			TMap<FMetaSoundAssetKey, TArray<FMetaSoundAssetClassInfo>> InfoOnShutdown;
 			{
 				FScopeLock Lock(&TagDataMapCriticalSection);
 				InfoOnShutdown = MoveTemp(ClassInfoMap);
@@ -120,7 +116,7 @@ namespace Metasound::Engine
 			if (!InfoOnShutdown.IsEmpty())
 			{
 				UE_LOG(LogMetaSound, Display, TEXT("AssetManager is shutting down with the following %i assets active:"), InfoOnShutdown.Num());
-				for (const TPair<FAssetKey, TArray<FMetaSoundAssetClassInfo>>& Pair : InfoOnShutdown)
+				for (const TPair<FMetaSoundAssetKey, TArray<FMetaSoundAssetClassInfo>>& Pair : InfoOnShutdown)
 				{
 					for (const FMetaSoundAssetClassInfo& ClassInfo : Pair.Value)
 					{
@@ -148,7 +144,7 @@ namespace Metasound::Engine
 
 		{
 			const FMetasoundFrontendDocument& Document = InAssetBase.GetConstDocumentChecked();
-			const FAssetKey AssetKey(Document.RootGraph.Metadata);
+			const FMetaSoundAssetKey AssetKey(Document.RootGraph.Metadata);
 			if (!ContainsKey(AssetKey))
 			{
 				AddOrUpdateFromObject(*InAssetBase.GetOwningAsset());
@@ -162,7 +158,7 @@ namespace Metasound::Engine
 		{
 			FNodeRegistryKey RegistryKey;
 			const bool bIsKey = FNodeRegistryKey::Parse(KeyString, RegistryKey);
-			if (!bIsKey || !ContainsKey(FAssetKey(RegistryKey)))
+			if (!bIsKey || !ContainsKey(FMetaSoundAssetKey(RegistryKey)))
 			{
 				UE_LOG(LogMetaSound, Verbose, TEXT("Missing referenced class '%s' asset entry."), *KeyString);
 				bAddFromReferencedAssets = true;
@@ -183,7 +179,7 @@ namespace Metasound::Engine
 			if (Asset)
 			{
 				const FMetasoundFrontendDocument& RefDocument = Asset->GetConstDocumentChecked();
-				const FAssetKey ClassKey = FAssetKey(RefDocument.RootGraph);
+				const FMetaSoundAssetKey ClassKey(RefDocument.RootGraph);
 				if (!ContainsKey(ClassKey))
 				{
 					UE_LOG(LogMetaSound, Verbose,
@@ -208,7 +204,7 @@ namespace Metasound::Engine
 	}
 #endif // WITH_EDITORONLY_DATA
 
-	Frontend::FAssetKey FMetaSoundAssetManager::AddOrUpdateFromObject(const UObject& InObject)
+	FMetaSoundAssetKey FMetaSoundAssetManager::AddOrUpdateFromObject(const UObject& InObject)
 	{
 		using namespace AssetSubsystemPrivate;
 		using namespace Frontend;
@@ -221,10 +217,10 @@ namespace Metasound::Engine
 		TScriptInterface<const IMetaSoundDocumentInterface> ScriptInterface(&InObject);
 		check(ScriptInterface.GetObject() != nullptr);
 		const IMetaSoundDocumentInterface& DocInterface = *ScriptInterface;
-		return AssetSubsystemPrivate::AddClassInfo(TagDataMapCriticalSection, ClassInfoMap, FMetaSoundAssetClassInfo(DocInterface));
+		return AssetSubsystemPrivate::AddClassInfo(TagDataMapCriticalSection, ClassInfoMap, FMetaSoundAssetClassInfo (DocInterface));
 	}
 
-	void  FMetaSoundAssetManager::AddOrLoadAndUpdateFromObjectAsync(const FAssetData& InAssetData, TUniqueFunction<void(FAssetKey, UObject&)>&& OnAssetRegistered)
+	void  FMetaSoundAssetManager::AddOrLoadAndUpdateFromObjectAsync(const FAssetData& InAssetData, TUniqueFunction<void(FMetaSoundAssetKey, UObject&)>&& OnAssetRegistered)
 	{
 		using namespace Frontend;
 
@@ -241,7 +237,7 @@ namespace Metasound::Engine
 		if (UObject* Object = Path.ResolveObject())
 		{
 			UE_LOG(LogMetaSound, VeryVerbose, TEXT("FMetaSoundAssetManager::AddOrLoadAndUpdateFromObjectAsync called, but object ''%s' already loaded. Updating entry using in-memory object definition."), *Path.ToString());
-			const FAssetKey AssetKey = AddOrUpdateFromObject(*Object);
+			const FMetaSoundAssetKey AssetKey = AddOrUpdateFromObject(*Object);
 			OnAssetRegistered(AssetKey, *Object);
 			OnAssetRegistered.Reset();
 			return;
@@ -324,7 +320,7 @@ namespace Metasound::Engine
 		return !AutoUpdateDenyListCache.Contains(InClassName.GetFullName());
 	}
 
-	bool FMetaSoundAssetManager::ContainsKey(const Frontend::FAssetKey& InKey) const
+	bool FMetaSoundAssetManager::ContainsKey(const FMetaSoundAssetKey& InKey) const
 	{
 		FScopeLock Lock(&TagDataMapCriticalSection);
 		return ClassInfoMap.Contains(InKey);
@@ -332,26 +328,26 @@ namespace Metasound::Engine
 
 	bool FMetaSoundAssetManager::ContainsKey(const Frontend::FNodeRegistryKey& InRegistryKey) const
 	{
-		if (Frontend::FAssetKey::IsValidType(InRegistryKey.Type))
+		if (FMetaSoundAssetKey::IsValidType(InRegistryKey.Type))
 		{
-			return ContainsKey(Frontend::FAssetKey(InRegistryKey));
+			return ContainsKey(FMetaSoundAssetKey(InRegistryKey));
 		}
 
 		return false;
 	}
 
-	void FMetaSoundAssetManager::DepthFirstVisit_AssetKey(const FAssetKey& InKey, TFunctionRef<TSet<FAssetKey>(const FAssetKey&)> VisitFunction)
+	void FMetaSoundAssetManager::DepthFirstVisit_AssetKey(const FMetaSoundAssetKey& InKey, TFunctionRef<TSet<FMetaSoundAssetKey>(const FMetaSoundAssetKey&)> VisitFunction)
 	{
 		// Non recursive depth first traversal.
-		TArray<FAssetKey> Stack({ InKey });
-		TSet<FAssetKey> Visited;
+		TArray<FMetaSoundAssetKey> Stack({ InKey });
+		TSet<FMetaSoundAssetKey> Visited;
 
 		while (!Stack.IsEmpty())
 		{
-			const FAssetKey CurrentKey = Stack.Pop();
+			const FMetaSoundAssetKey CurrentKey = Stack.Pop();
 			if (!Visited.Contains(CurrentKey))
 			{
-				TArray<FAssetKey> Children = VisitFunction(CurrentKey).Array();
+				TArray<FMetaSoundAssetKey> Children = VisitFunction(CurrentKey).Array();
 				Stack.Append(MoveTemp(Children));
 				Visited.Add(CurrentKey);
 			}
@@ -378,7 +374,7 @@ namespace Metasound::Engine
 		return LoadingDependencies.FindByPredicate(IsEqualID);
 	}
 
-	FMetasoundAssetBase* FMetaSoundAssetManager::FindAsset(const Frontend::FAssetKey& InKey) const
+	FMetasoundAssetBase* FMetaSoundAssetManager::FindAsset(const FMetaSoundAssetKey& InKey) const
 	{
 		FTopLevelAssetPath AssetPath = FindAssetPath(InKey);
 		if (AssetPath.IsValid())
@@ -427,13 +423,13 @@ namespace Metasound::Engine
 		return { };
 	}
 
-	TArray<Frontend::FMetaSoundAssetClassInfo> FMetaSoundAssetManager::FindAssetClassInfoInternal(const Frontend::FAssetKey& InKey) const
+	TArray<Frontend::FMetaSoundAssetClassInfo> FMetaSoundAssetManager::FindAssetClassInfoInternal(const FMetaSoundAssetKey& InKey) const
 	{
 		FScopeLock Lock(&TagDataMapCriticalSection);
 		return ClassInfoMap.FindRef(InKey);
 	}
 
-	TScriptInterface<IMetaSoundDocumentInterface> FMetaSoundAssetManager::FindAssetAsDocumentInterface(const Frontend::FAssetKey& InKey) const
+	TScriptInterface<IMetaSoundDocumentInterface> FMetaSoundAssetManager::FindAssetAsDocumentInterface(const FMetaSoundAssetKey& InKey) const
 	{
 		const FTopLevelAssetPath AssetPath = FindAssetPath(InKey);
 		if (AssetPath.IsValid())
@@ -447,7 +443,7 @@ namespace Metasound::Engine
 		return nullptr;
 	}
 
-	FTopLevelAssetPath FMetaSoundAssetManager::FindAssetPath(const Frontend::FAssetKey& InKey) const
+	FTopLevelAssetPath FMetaSoundAssetManager::FindAssetPath(const FMetaSoundAssetKey& InKey) const
 	{
 		using namespace Frontend;
 
@@ -463,7 +459,7 @@ namespace Metasound::Engine
 		return nullptr;
 	}
 
-	TArray<FTopLevelAssetPath> FMetaSoundAssetManager::FindAssetPaths(const Frontend::FAssetKey& InKey) const
+	TArray<FTopLevelAssetPath> FMetaSoundAssetManager::FindAssetPaths(const FMetaSoundAssetKey& InKey) const
 	{
 		using namespace Frontend;
 
@@ -495,7 +491,7 @@ namespace Metasound::Engine
 		return IMetasoundUObjectRegistry::Get().GetObjectAsAssetBase(&InObject);
 	}
 
-	Frontend::FAssetKey FMetaSoundAssetManager::GetAssetKey(const FSoftObjectPath& InObjectPath) const
+	FMetaSoundAssetKey FMetaSoundAssetManager::GetAssetKey(const FSoftObjectPath& InObjectPath) const
 	{
 		using namespace Frontend;
 
@@ -602,7 +598,7 @@ namespace Metasound::Engine
 						}
 #endif // WITH_EDITORONLY_DATA
 
-						const FAssetKey AssetKey = AssetManager->AddOrUpdateFromObject(*MetaSoundObj);
+						const FMetaSoundAssetKey AssetKey = AssetManager->AddOrUpdateFromObject(*MetaSoundObj);
 						if (PackageLoadedArgs.OnAssetRegistered)
 						{
 							PackageLoadedArgs.OnAssetRegistered(AssetKey, *MetaSoundObj);
@@ -632,12 +628,14 @@ namespace Metasound::Engine
 	}
 
 #if WITH_EDITOR
-	TSet<UMetaSoundAssetSubsystem::FAssetInfo> FMetaSoundAssetManager::GetReferencedAssetClasses(const FMetasoundAssetBase& InAssetBase) const
+	TSet<Frontend::IMetaSoundAssetManager::FAssetRef> FMetaSoundAssetManager::GetReferencedAssets(const FMetasoundAssetBase& InAssetBase) const
 	{
-		METASOUND_TRACE_CPUPROFILER_EVENT_SCOPE(FMetaSoundAssetManager::GetReferencedAssetClasses);
+		METASOUND_TRACE_CPUPROFILER_EVENT_SCOPE(FMetaSoundAssetManager::GetReferencedAssets);
 		using namespace Frontend;
 
-		TSet<FAssetInfo> OutAssetInfos;
+		using FAssetRef = IMetaSoundAssetManager::FAssetRef;
+
+		TSet<FAssetRef> OutAssetRefs;
 		const FMetasoundFrontendDocument& Document = InAssetBase.GetConstDocumentChecked();
 		for (const FMetasoundFrontendClass& Class : Document.Dependencies)
 		{
@@ -646,12 +644,12 @@ namespace Metasound::Engine
 				continue;
 			}
 
-			const FAssetKey AssetKey(Class.Metadata);
+			const FMetaSoundAssetKey AssetKey(Class.Metadata);
 			FTopLevelAssetPath ObjectPath = FindAssetPath(AssetKey);
 			if (ObjectPath.IsValid())
 			{
-				FAssetInfo AssetInfo { FNodeRegistryKey(Class.Metadata), FSoftObjectPath(ObjectPath) };
-				OutAssetInfos.Add(MoveTemp(AssetInfo));
+				FAssetRef AssetRef { FMetaSoundAssetKey(Class.Metadata), ObjectPath };
+				OutAssetRefs.Add(MoveTemp(AssetRef));
 			}
 			else
 			{
@@ -681,7 +679,7 @@ namespace Metasound::Engine
 				}
 			}
 		}
-		return MoveTemp(OutAssetInfos);
+		return MoveTemp(OutAssetRefs);
 	}
 	
 	bool FMetaSoundAssetManager::GetReferencedPresetHierarchy(FMetasoundAssetBase& InAsset, TArray<FMetasoundAssetBase*>& OutReferencedAssets) const
@@ -749,7 +747,7 @@ namespace Metasound::Engine
 		AssetBase->UnregisterGraphWithFrontend();
 
 		{
-			const FAssetKey OldAssetKey(ClassMetadata.GetClassName(), ClassMetadata.GetVersion());
+			const FMetaSoundAssetKey OldAssetKey(ClassMetadata.GetClassName(), ClassMetadata.GetVersion());
 			AssetSubsystemPrivate::RemoveClassInfo(TagDataMapCriticalSection, ClassInfoMap, OldAssetKey, Path);
 		}
 
@@ -776,7 +774,7 @@ namespace Metasound::Engine
 			return false;
 		}
 
-		return ContainsKey(FAssetKey(ClassMetadata));
+		return ContainsKey(FMetaSoundAssetKey(ClassMetadata));
 	}
 
 #if WITH_EDITOR
@@ -812,9 +810,9 @@ namespace Metasound::Engine
 		}
 	}
 
-	void FMetaSoundAssetManager::IterateReferences(const FAssetKey& InKey, TFunctionRef<void(const FAssetKey&)> VisitFunction) const
+	void FMetaSoundAssetManager::IterateReferences(const FMetaSoundAssetKey& InKey, TFunctionRef<void(const FMetaSoundAssetKey&)> VisitFunction) const
 	{
-		DepthFirstVisit_AssetKey(InKey, [&](const FAssetKey& ReferencedKey)
+		DepthFirstVisit_AssetKey(InKey, [&](const FMetaSoundAssetKey& ReferencedKey)
 		{
 			using namespace Frontend;
 
@@ -827,7 +825,7 @@ namespace Metasound::Engine
 				}
 			}
 
-			TSet<FAssetKey> Children;
+			TSet<FMetaSoundAssetKey> Children;
 			for (FMetaSoundAssetClassInfo& ClassInfo : TagDatas)
 			{
 				// If loaded, this is the "freshest" reference list, so use that.
@@ -847,7 +845,7 @@ namespace Metasound::Engine
 
 							if (bIsAsset)
 							{
-								const FAssetKey RefKey = FAssetKey(Class.Metadata);
+								const FMetaSoundAssetKey RefKey(Class.Metadata);
 								VisitFunction(RefKey);
 								Children.Add(RefKey);
 							}
@@ -858,7 +856,7 @@ namespace Metasound::Engine
 				// Otherwise, use provided reference list from the asset tag data
 				else
 				{
-					for (const FAssetKey& RefKey : ClassInfo.DocInfo.ReferencedAssetKeys)
+					for (const FMetaSoundAssetKey& RefKey : ClassInfo.DocInfo.ReferencedAssetKeys)
 					{
 						VisitFunction(RefKey);
 					}
@@ -915,7 +913,7 @@ namespace Metasound::Engine
 
 		SearchAndIterateDirectoryAssets(Directories, [this](const FAssetData& AssetData)
 		{
-			AddOrLoadAndUpdateFromObjectAsync(AssetData, [](FAssetKey, UObject& AssetObject)
+			AddOrLoadAndUpdateFromObjectAsync(AssetData, [](FMetaSoundAssetKey, UObject& AssetObject)
 			{
 				FMetasoundAssetBase* MetaSoundAsset = Metasound::IMetasoundUObjectRegistry::Get().GetObjectAsAssetBase(&AssetObject);
 				check(MetaSoundAsset);
@@ -946,7 +944,7 @@ namespace Metasound::Engine
 			BuilderRegistry->FinishBuilding(Metadata.GetClassName(), AssetPath, bForceUnregister);
 		}
 
-		const FAssetKey AssetKey(Metadata.GetClassName(), Metadata.GetVersion());
+		const FMetaSoundAssetKey AssetKey(Metadata.GetClassName(), Metadata.GetVersion());
 		AssetSubsystemPrivate::RemoveClassInfo(TagDataMapCriticalSection, ClassInfoMap, AssetKey, AssetPath);
 	}
 
@@ -963,7 +961,7 @@ namespace Metasound::Engine
 				BuilderRegistry->FinishBuilding(ClassInfo.ClassName, ClassInfo.AssetPath, bForceUnregister);
 			}
 
-			const FAssetKey AssetKey(ClassInfo.ClassName, ClassInfo.Version);
+			const FMetaSoundAssetKey AssetKey(ClassInfo.ClassName, ClassInfo.Version);
 			AssetSubsystemPrivate::RemoveClassInfo(TagDataMapCriticalSection, ClassInfoMap, AssetKey, ClassInfo.AssetPath);
 		}
 	}
@@ -987,7 +985,7 @@ namespace Metasound::Engine
 		FMetaSoundAssetClassInfo ClassInfo(InAssetData);
 		if (ensureAlways(ClassInfo.bIsValid))
 		{
-			const FAssetKey AssetKey(ClassInfo.ClassName, ClassInfo.Version);
+			const FMetaSoundAssetKey AssetKey(ClassInfo.ClassName, ClassInfo.Version);
 			const FTopLevelAssetPath OldPath(InOldObjectPath);
 			AssetSubsystemPrivate::RemoveClassInfo(TagDataMapCriticalSection, ClassInfoMap, AssetKey, OldPath);
 			AssetSubsystemPrivate::AddClassInfo(TagDataMapCriticalSection, ClassInfoMap, MoveTemp(ClassInfo));
@@ -1094,7 +1092,7 @@ namespace Metasound::Engine
 
 		IMetasoundUObjectRegistry& ObjectRegistry = IMetasoundUObjectRegistry::Get();
 		TSet<FMetasoundAssetBase*> ToReregister;
-		for (const TPair<FAssetKey, TArray<FMetaSoundAssetClassInfo>>& Pair : ClassInfoMap)
+		for (const TPair<FMetaSoundAssetKey, TArray<FMetaSoundAssetClassInfo>>& Pair : ClassInfoMap)
 		{
 			if (!Pair.Value.IsEmpty())
 			{
@@ -1157,7 +1155,7 @@ namespace Metasound::Engine
 		bLogActiveAssetsOnShutdown = bInLogActiveAssetsOnShutdown;
 	}
 
-	FMetasoundAssetBase* FMetaSoundAssetManager::TryLoadAssetFromKey(const Frontend::FAssetKey& InAssetKey) const
+	FMetasoundAssetBase* FMetaSoundAssetManager::TryLoadAssetFromKey(const FMetaSoundAssetKey& InAssetKey) const
 	{
 		FTopLevelAssetPath ObjectPath = FindAssetPath(InAssetKey);
 		if (ObjectPath.IsValid())
@@ -1228,7 +1226,7 @@ namespace Metasound::Engine
 				const FMetaSoundAssetClassInfo ClassInfo(AssetData);
 				if (ensureAlways(ClassInfo.bIsValid))
 				{
-					const FAssetKey AssetKey(ClassInfo.ClassName, ClassInfo.Version);
+					const FMetaSoundAssetKey AssetKey(ClassInfo.ClassName, ClassInfo.Version);
 					const FNodeRegistryKey RegistryKey = FNodeRegistryKey(AssetKey);
 					const bool bIsRegistered = FMetasoundFrontendRegistryContainer::Get()->IsNodeRegistered(RegistryKey);
 					if (bIsRegistered)
@@ -1392,7 +1390,7 @@ void UMetaSoundAssetSubsystem::PostInitAssetScanInternal()
 		{
 			using namespace Metasound;
 
-			Manager.AddOrLoadAndUpdateFromObjectAsync(AssetData, [](FAssetKey, UObject& AssetObject)
+			Manager.AddOrLoadAndUpdateFromObjectAsync(AssetData, [](FMetaSoundAssetKey, UObject& AssetObject)
 			{
 				FMetasoundAssetBase* MetaSoundAsset = IMetasoundUObjectRegistry::Get().GetObjectAsAssetBase(&AssetObject);
 				check(MetaSoundAsset);
@@ -1439,7 +1437,7 @@ bool UMetaSoundAssetSubsystem::CanAutoUpdate(const FMetasoundFrontendClassName&
 bool UMetaSoundAssetSubsystem::ContainsKey(const Metasound::Frontend::FNodeRegistryKey& InRegistryKey) const
 {
 	using namespace Metasound::Frontend;
-	return IMetaSoundAssetManager::GetChecked().ContainsKey(FAssetKey(InRegistryKey));
+	return IMetaSoundAssetManager::GetChecked().ContainsKey(FMetaSoundAssetKey(InRegistryKey));
 }
 
 #if WITH_EDITOR
@@ -1497,7 +1495,7 @@ const FMetasoundAssetBase* UMetaSoundAssetSubsystem::GetAsAsset(const UObject& I
 TSet<UMetaSoundAssetSubsystem::FAssetInfo> UMetaSoundAssetSubsystem::GetReferencedAssetClasses(const FMetasoundAssetBase& InAssetBase) const
 {
 	using namespace Metasound::Frontend;
-	return IMetaSoundAssetManager::GetChecked().GetReferencedAssetClasses(InAssetBase);
+	return { };
 }
 #endif // WITH_EDITOR
 
diff --git a/Engine/Plugins/Runtime/Metasound/Source/MetasoundEngine/Private/MetasoundEngineAsset.h b/Engine/Plugins/Runtime/Metasound/Source/MetasoundEngine/Private/MetasoundEngineAsset.h
--- a/Engine/Plugins/Runtime/Metasound/Source/MetasoundEngine/Private/MetasoundEngineAsset.h
+++ b/Engine/Plugins/Runtime/Metasound/Source/MetasoundEngine/Private/MetasoundEngineAsset.h
@@ -90,23 +90,24 @@ namespace Metasound::Engine
 		}
 
 		template<typename TMetaSoundObject>
-		static void SetReferencedAssetClasses(TMetaSoundObject& InMetaSound, TSet<Frontend::IMetaSoundAssetManager::FAssetInfo>&& InAssetClasses)
+		static void SetReferencedAssets(TMetaSoundObject& InMetaSound, TSet<Frontend::IMetaSoundAssetManager::FAssetRef>&& InAssetRefs)
 		{
 			using namespace Frontend;
 			
 			InMetaSound.ReferencedAssetClassKeys.Reset();
 			InMetaSound.ReferencedAssetClassObjects.Reset();
 
-			for (const IMetaSoundAssetManager::FAssetInfo& AssetClass : InAssetClasses)
+			for (const IMetaSoundAssetManager::FAssetRef& AssetRef : InAssetRefs)
 			{
-				InMetaSound.ReferencedAssetClassKeys.Add(AssetClass.RegistryKey.ToString());
-				if (UObject* Object = AssetClass.AssetPath.TryLoad())
+				// Has to be serialized as node class registry key string for back compat
+				InMetaSound.ReferencedAssetClassKeys.Add(FNodeClassRegistryKey(AssetRef.Key).ToString());
+				if (UObject* Object = FSoftObjectPath(AssetRef.Path).TryLoad())
 				{
 					InMetaSound.ReferencedAssetClassObjects.Add(Object);
 				}
 				else
 				{
-					UE_LOG(LogMetaSound, Error, TEXT("Failed to load referenced asset %s from asset %s"), *AssetClass.AssetPath.ToString(), *InMetaSound.GetPathName());
+					UE_LOG(LogMetaSound, Error, TEXT("Failed to load referenced asset %s from asset %s"), *AssetRef.Path.ToString(), *InMetaSound.GetPathName());
 				}
 			}
 		}
@@ -124,7 +125,7 @@ namespace Metasound::Engine
 			AssetManager.SetLogActiveAssetsOnShutdown(false);
 
 			// Add error for multiple assets with the same class name
-			const FAssetKey Key(Document.RootGraph.Metadata);
+			const FMetaSoundAssetKey Key(Document.RootGraph.Metadata);
 			const TArray<FTopLevelAssetPath> AssetPaths = AssetManager.FindAssetPaths(Key);
 			if (AssetPaths.Num() > 1)
 			{
diff --git a/Engine/Plugins/Runtime/Metasound/Source/MetasoundEngine/Private/MetasoundEngineModule.cpp b/Engine/Plugins/Runtime/Metasound/Source/MetasoundEngine/Private/MetasoundEngineModule.cpp
--- a/Engine/Plugins/Runtime/Metasound/Source/MetasoundEngine/Private/MetasoundEngineModule.cpp
+++ b/Engine/Plugins/Runtime/Metasound/Source/MetasoundEngine/Private/MetasoundEngineModule.cpp
@@ -200,7 +200,7 @@ namespace Metasound::Engine
 		if (ModulePrivate::EnableMetaSoundEditorAssetAutoLoadAndRegister)
 		{
 			IMetaSoundAssetManager::GetChecked().AddOrLoadAndUpdateFromObjectAsync(InAssetData,
-			[](FAssetKey, UObject& AssetObject)
+			[](FMetaSoundAssetKey, UObject& AssetObject)
 			{
 				FModule& ThisModule = FModuleManager::GetModuleChecked<Metasound::Engine::FModule>("MetaSoundEngine");
 				ThisModule.GetOnGraphRegisteredDelegate().ExecuteIfBound(AssetObject, ERegistrationAssetContext::None);
@@ -227,7 +227,7 @@ namespace Metasound::Engine
 				if (ModulePrivate::EnableMetaSoundEditorAssetAutoLoadAndRegister)
 				{
 					IMetaSoundAssetManager::GetChecked().AddOrLoadAndUpdateFromObjectAsync(InAssetData,
-					[](FAssetKey AssetKey, UObject& AssetObject)
+					[](FMetaSoundAssetKey AssetKey, UObject& AssetObject)
 					{
 						// Have to re-register to avoid registry desync.
 						const bool bIsRegistered = FMetasoundFrontendRegistryContainer::Get()->IsNodeRegistered(FNodeRegistryKey(AssetKey));
diff --git a/Engine/Plugins/Runtime/Metasound/Source/MetasoundEngine/Private/MetasoundSource.cpp b/Engine/Plugins/Runtime/Metasound/Source/MetasoundEngine/Private/MetasoundSource.cpp
--- a/Engine/Plugins/Runtime/Metasound/Source/MetasoundEngine/Private/MetasoundSource.cpp
+++ b/Engine/Plugins/Runtime/Metasound/Source/MetasoundEngine/Private/MetasoundSource.cpp
@@ -696,9 +696,9 @@ void UMetaSoundSource::SerializeCookedQualitySettings(const FName PlatformName,
 }
 
 #if WITH_EDITOR
-void UMetaSoundSource::SetReferencedAssetClasses(TSet<Metasound::Frontend::IMetaSoundAssetManager::FAssetInfo>&& InAssetClasses)
+void UMetaSoundSource::SetReferencedAssets(TSet<Metasound::Frontend::IMetaSoundAssetManager::FAssetRef>&& InAssetRefs)
 {
-	Metasound::Engine::FAssetHelper::SetReferencedAssetClasses(*this, MoveTemp(InAssetClasses));
+	Metasound::Engine::FAssetHelper::SetReferencedAssets(*this, MoveTemp(InAssetRefs));
 }
 #endif // WITH_EDITOR
 
diff --git a/Engine/Plugins/Runtime/Metasound/Source/MetasoundEngine/Public/Metasound.h b/Engine/Plugins/Runtime/Metasound/Source/MetasoundEngine/Public/Metasound.h
--- a/Engine/Plugins/Runtime/Metasound/Source/MetasoundEngine/Public/Metasound.h
+++ b/Engine/Plugins/Runtime/Metasound/Source/MetasoundEngine/Public/Metasound.h
@@ -189,7 +189,7 @@ public:
 
 protected:
 #if WITH_EDITOR
-	virtual void SetReferencedAssetClasses(TSet<Metasound::Frontend::IMetaSoundAssetManager::FAssetInfo>&& InAssetClasses) override;
+	virtual void SetReferencedAssets(TSet<Metasound::Frontend::IMetaSoundAssetManager::FAssetRef>&& InAssetRefs) override;
 #endif // #if WITH_EDITOR
 
 	Metasound::Frontend::FDocumentAccessPtr GetDocumentAccessPtr() override;
diff --git a/Engine/Plugins/Runtime/Metasound/Source/MetasoundEngine/Public/MetasoundAssetSubsystem.h b/Engine/Plugins/Runtime/Metasound/Source/MetasoundEngine/Public/MetasoundAssetSubsystem.h
--- a/Engine/Plugins/Runtime/Metasound/Source/MetasoundEngine/Public/MetasoundAssetSubsystem.h
+++ b/Engine/Plugins/Runtime/Metasound/Source/MetasoundEngine/Public/MetasoundAssetSubsystem.h
@@ -4,6 +4,7 @@
 #include "Engine/Engine.h"
 #include "Engine/StreamableManager.h"
 #include "MetasoundAssetBase.h"
+#include "MetasoundAssetKey.h"
 #include "MetasoundAssetManager.h"
 #include "MetasoundBuilderBase.h"
 #include "MetasoundDocumentInterface.h"
@@ -35,9 +36,12 @@ namespace Metasound::Engine
 	public:
 		friend class UMetaSoundAssetSubsystem;
 
-		using FAssetClassInfoMap = TMap<Frontend::FAssetKey, TArray<Frontend::FMetaSoundAssetClassInfo>>;
-		using FAssetInfo = Frontend::IMetaSoundAssetManager::FAssetInfo;
-		using FAssetKey = Frontend::FAssetKey;
+		using FAssetClassInfoMap = TMap<FMetaSoundAssetKey, TArray<Frontend::FMetaSoundAssetClassInfo>>;
+
+PRAGMA_DISABLE_DEPRECATION_WARNINGS 
+		using FAssetInfo UE_DEPRECATED(5.6, "Use FAssetRef instead which is more compact and usage clear in context of other Info tag/query types") = Frontend::IMetaSoundAssetManager::FAssetInfo;
+PRAGMA_ENABLE_DEPRECATION_WARNINGS
+
 		using FNodeRegistryKey = Frontend::FNodeRegistryKey;
 
 #if WITH_EDITORONLY_DATA
@@ -69,7 +73,7 @@ namespace Metasound::Engine
 		 **/
 		Frontend::FMetaSoundAssetClassInfo FindAssetClassInfo(const FTopLevelAssetPath& InPath) const;
 
-		Frontend::FAssetKey GetAssetKey(const FSoftObjectPath& InObjectPath) const;
+		FMetaSoundAssetKey GetAssetKey(const FSoftObjectPath& InObjectPath) const;
 
 #if WITH_EDITOR
 		int32 GetActiveAsyncLoadRequestCount() const;
@@ -101,20 +105,20 @@ namespace Metasound::Engine
 #if WITH_EDITORONLY_DATA
 		virtual bool AddAssetReferences(FMetasoundAssetBase& InAssetBase) override;
 #endif // WITH_EDITORONLY_DATA
-		virtual FAssetKey AddOrUpdateFromObject(const UObject& InObject) override;
-		virtual void AddOrLoadAndUpdateFromObjectAsync(const FAssetData& InAssetData, TUniqueFunction<void(FAssetKey, UObject&)>&& OnAssetRegistered) override;
+		virtual FMetaSoundAssetKey AddOrUpdateFromObject(const UObject& InObject) override;
+		virtual void AddOrLoadAndUpdateFromObjectAsync(const FAssetData& InAssetData, TUniqueFunction<void(FMetaSoundAssetKey, UObject&)>&& OnAssetRegistered) override;
 		virtual void AddOrUpdateFromAssetData(const FAssetData& InAssetData) override;
 		virtual bool CanAutoUpdate(const FMetasoundFrontendClassName& InClassName) const override;
-		virtual bool ContainsKey(const FAssetKey& InKey) const override;
+		virtual bool ContainsKey(const FMetaSoundAssetKey& InKey) const override;
 		virtual bool ContainsKey(const Frontend::FNodeRegistryKey& InRegistryKey) const override;
-		virtual FMetasoundAssetBase* FindAsset(const FAssetKey& InKey) const override;
-		virtual TScriptInterface<IMetaSoundDocumentInterface> FindAssetAsDocumentInterface(const FAssetKey& InKey) const override;
-		virtual FTopLevelAssetPath FindAssetPath(const FAssetKey& InKey) const override;
-		virtual TArray<FTopLevelAssetPath> FindAssetPaths(const FAssetKey& InKey) const override;
+		virtual FMetasoundAssetBase* FindAsset(const FMetaSoundAssetKey& InKey) const override;
+		virtual TScriptInterface<IMetaSoundDocumentInterface> FindAssetAsDocumentInterface(const FMetaSoundAssetKey& InKey) const override;
+		virtual FTopLevelAssetPath FindAssetPath(const FMetaSoundAssetKey& InKey) const override;
+		virtual TArray<FTopLevelAssetPath> FindAssetPaths(const FMetaSoundAssetKey& InKey) const override;
 		virtual FMetasoundAssetBase* GetAsAsset(UObject& InObject) const override;
 		virtual const FMetasoundAssetBase* GetAsAsset(const UObject& InObject) const override;
 #if WITH_EDITOR
-		virtual TSet<FAssetInfo> GetReferencedAssetClasses(const FMetasoundAssetBase& InAssetBase) const override;
+		virtual TSet<Frontend::IMetaSoundAssetManager::FAssetRef> GetReferencedAssets(const FMetasoundAssetBase& InAssetBase) const override;
 		virtual bool GetReferencedPresetHierarchy(FMetasoundAssetBase& InAsset, TArray<FMetasoundAssetBase*>& OutReferencedAssets) const override;
 		virtual bool ReassignClassName(TScriptInterface<IMetaSoundDocumentInterface> DocInterface) override;
 #endif // WITH_EDITOR
@@ -127,7 +131,7 @@ namespace Metasound::Engine
 
 		// Iterates all references of a given asset key entry (Only recommended in certain editor contexts, as can be slow for
 		// deep reference trees and blocks access to reference map).
-		void IterateReferences(const FAssetKey& InKey, TFunctionRef<void(const FAssetKey&)> VisitFunction) const;
+		void IterateReferences(const FMetaSoundAssetKey& InKey, TFunctionRef<void(const FMetaSoundAssetKey&)> VisitFunction) const;
 #endif // WITH_EDITOR
 
 		virtual void ReloadMetaSoundAssets() const override;
@@ -135,7 +139,7 @@ namespace Metasound::Engine
 		virtual void RemoveAsset(const FAssetData& InAssetData) override;
 		virtual void RenameAsset(const FAssetData& InAssetData, const FString& InOldObjectPath) override;
 		virtual void SetLogActiveAssetsOnShutdown(bool bInLogActiveAssetsOnShutdown) override;
-		virtual FMetasoundAssetBase* TryLoadAssetFromKey(const FAssetKey& InKey) const override;
+		virtual FMetasoundAssetBase* TryLoadAssetFromKey(const FMetaSoundAssetKey& InKey) const override;
 		virtual bool TryGetAssetIDFromClassName(const FMetasoundFrontendClassName& InClassName, FGuid& OutGuid) const override;
 		virtual bool TryLoadReferencedAssets(const FMetasoundAssetBase& InAssetBase, TArray<FMetasoundAssetBase*>& OutReferencedAssets) const override;
 #if WITH_EDITORONLY_DATA
@@ -151,7 +155,7 @@ namespace Metasound::Engine
 	private:
 		TArray<FMetaSoundAsyncAssetDependencies> LoadingDependencies;
 
-		static void DepthFirstVisit_AssetKey(const FAssetKey& InKey, TFunctionRef<TSet<FAssetKey>(const FAssetKey&)> VisitFunction);
+		static void DepthFirstVisit_AssetKey(const FMetaSoundAssetKey& InKey, TFunctionRef<TSet<FMetaSoundAssetKey>(const FMetaSoundAssetKey&)> VisitFunction);
 
 		// Returns all asset query info with the given asset key (may be more than one due to errors in duplication
 		// of assets, improper plugin migration, etc. but should usually just be a single entry) registered
@@ -171,7 +175,7 @@ namespace Metasound::Engine
 			FName PackageName;
 			UPackage* Package = nullptr;
 			EAsyncLoadingResult::Type Result = EAsyncLoadingResult::Failed;
-			TUniqueFunction<void(FAssetKey, UObject&)> OnAssetRegistered;
+			TUniqueFunction<void(FMetaSoundAssetKey, UObject&)> OnAssetRegistered;
 			bool bMarkAsVersioned = false;
 		};
 
@@ -234,7 +238,9 @@ class METASOUNDENGINE_API UMetaSoundAssetSubsystem : public UEngineSubsystem
 	GENERATED_BODY()
 
 public:
-	using FAssetInfo = Metasound::Frontend::IMetaSoundAssetManager::FAssetInfo;
+PRAGMA_DISABLE_DEPRECATION_WARNINGS
+		using FAssetInfo UE_DEPRECATED(5.6, "No longer supported in favor of more compact and direct AssetRef struct") = Metasound::Frontend::IMetaSoundAssetManager::FAssetInfo;
+PRAGMA_ENABLE_DEPRECATION_WARNINGS
 
 	virtual void Initialize(FSubsystemCollectionBase& InCollection) override;
 
diff --git a/Engine/Plugins/Runtime/Metasound/Source/MetasoundEngine/Public/MetasoundBuilderBase.cpp b/Engine/Plugins/Runtime/Metasound/Source/MetasoundEngine/Public/MetasoundBuilderBase.cpp
--- a/Engine/Plugins/Runtime/Metasound/Source/MetasoundEngine/Public/MetasoundBuilderBase.cpp
+++ b/Engine/Plugins/Runtime/Metasound/Source/MetasoundEngine/Public/MetasoundBuilderBase.cpp
@@ -1014,7 +1014,7 @@ void UMetaSoundBuilderBase::OnDependencyAdded(int32 Index)
 	const FMetasoundFrontendClass& NewDependency = Builder.GetConstDocumentChecked().Dependencies[Index];
 	if (NewDependency.Metadata.GetType() == EMetasoundFrontendClassType::External)
 	{
-		const FAssetKey AssetKey(NewDependency.Metadata);
+		const FMetaSoundAssetKey AssetKey(NewDependency.Metadata);
 		if (TScriptInterface<IMetaSoundDocumentInterface> DocInterface = IMetaSoundAssetManager::GetChecked().FindAssetAsDocumentInterface(AssetKey))
 		{
 			OnAssetReferenceAdded(DocInterface);
@@ -1030,7 +1030,7 @@ void UMetaSoundBuilderBase::OnRemoveSwappingDependency(int32 Index, int32 LastIn
 	const FMetasoundFrontendClass& NewDependency = Builder.GetConstDocumentChecked().Dependencies[Index];
 	if (NewDependency.Metadata.GetType() == EMetasoundFrontendClassType::External)
 	{
-		const FAssetKey AssetKey(NewDependency.Metadata);
+		const FMetaSoundAssetKey AssetKey(NewDependency.Metadata);
 		if (TScriptInterface<IMetaSoundDocumentInterface> DocInterface = IMetaSoundAssetManager::GetChecked().FindAssetAsDocumentInterface(AssetKey))
 		{
 			OnRemovingAssetReference(DocInterface);
diff --git a/Engine/Plugins/Runtime/Metasound/Source/MetasoundEngine/Public/MetasoundSource.h b/Engine/Plugins/Runtime/Metasound/Source/MetasoundEngine/Public/MetasoundSource.h
--- a/Engine/Plugins/Runtime/Metasound/Source/MetasoundEngine/Public/MetasoundSource.h
+++ b/Engine/Plugins/Runtime/Metasound/Source/MetasoundEngine/Public/MetasoundSource.h
@@ -323,7 +323,7 @@ protected:
 	virtual bool GetAllDefaultParameters(TArray<FAudioParameter>& OutParameters) const override;
 
 #if WITH_EDITOR
-	virtual void SetReferencedAssetClasses(TSet<Metasound::Frontend::IMetaSoundAssetManager::FAssetInfo>&& InAssetClasses) override;
+	virtual void SetReferencedAssets(TSet<Metasound::Frontend::IMetaSoundAssetManager::FAssetRef>&& InAssetRefs) override;
 #endif // #if WITH_EDITOR
 
 private:
diff --git a/Engine/Plugins/Runtime/Metasound/Source/MetasoundFrontend/Private/MetasoundAssetBase.cpp b/Engine/Plugins/Runtime/Metasound/Source/MetasoundFrontend/Private/MetasoundAssetBase.cpp
--- a/Engine/Plugins/Runtime/Metasound/Source/MetasoundFrontend/Private/MetasoundAssetBase.cpp
+++ b/Engine/Plugins/Runtime/Metasound/Source/MetasoundFrontend/Private/MetasoundAssetBase.cpp
@@ -909,8 +909,8 @@ void FMetasoundAssetBase::RebuildReferencedAssetClasses()
 
 	IMetaSoundAssetManager& AssetManager = IMetaSoundAssetManager::GetChecked();
 	AssetManager.AddAssetReferences(*this);
-	TSet<IMetaSoundAssetManager::FAssetInfo> ReferencedAssetClasses = AssetManager.GetReferencedAssetClasses(*this);
-	SetReferencedAssetClasses(MoveTemp(ReferencedAssetClasses));
+	TSet<IMetaSoundAssetManager::FAssetRef> ReferencedAssetClasses = AssetManager.GetReferencedAssets(*this);
+	SetReferencedAssets(MoveTemp(ReferencedAssetClasses));
 }
 #endif // WITH_EDITOR
 
diff --git a/Engine/Plugins/Runtime/Metasound/Source/MetasoundFrontend/Private/MetasoundFrontendDocumentBuilder.cpp b/Engine/Plugins/Runtime/Metasound/Source/MetasoundFrontend/Private/MetasoundFrontendDocumentBuilder.cpp
--- a/Engine/Plugins/Runtime/Metasound/Source/MetasoundFrontend/Private/MetasoundFrontendDocumentBuilder.cpp
+++ b/Engine/Plugins/Runtime/Metasound/Source/MetasoundFrontend/Private/MetasoundFrontendDocumentBuilder.cpp
@@ -1355,7 +1355,7 @@ const FMetasoundFrontendVariable* FMetaSoundFrontendDocumentBuilder::AddGraphVar
 		using namespace Metasound::Frontend;
 
 		// Cache the asset name on the node if it node is reference to asset-defined graph.
-		const FTopLevelAssetPath Path = IMetaSoundAssetManager::GetChecked().FindAssetPath(FAssetKey(ClassKey.ClassName, ClassKey.Version));
+		const FTopLevelAssetPath Path = IMetaSoundAssetManager::GetChecked().FindAssetPath(FMetaSoundAssetKey(ClassKey.ClassName, ClassKey.Version));
 		if (Path.IsValid())
 		{
 			InOutNode.Name = Path.GetAssetName();
@@ -1694,7 +1694,7 @@ const FMetasoundFrontendNode* FMetaSoundFrontendDocumentBuilder::AddGraphNode(co
 		using namespace Metasound::Frontend;
 
 		// Cache the asset name on the node if it node is reference to asset-defined graph.
-		const FTopLevelAssetPath Path = IMetaSoundAssetManager::GetChecked().FindAssetPath(FAssetKey(ClassKey.ClassName, ClassKey.Version));
+		const FTopLevelAssetPath Path = IMetaSoundAssetManager::GetChecked().FindAssetPath(FMetaSoundAssetKey(ClassKey.ClassName, ClassKey.Version));
 		if (Path.IsValid())
 		{
 			InOutNode.Name = Path.GetAssetName();
@@ -2999,14 +2999,14 @@ FMetasoundAssetBase* FMetaSoundFrontendDocumentBuilder::GetReferencedPresetAsset
 	{
 		const FMetasoundFrontendClass* NodeClass = FindDependency(Node->ClassID);
 		check(NodeClass);
-		const FAssetKey NodeAssetKey(NodeClass->Metadata);
+		const FMetaSoundAssetKey NodeAssetKey(NodeClass->Metadata);
 		const TArray<FMetasoundAssetBase*> ReferencedAssets = GetMetasoundAsset().GetReferencedAssets();
 		for (FMetasoundAssetBase* RefAsset : ReferencedAssets)
 		{
 			TScriptInterface<IMetaSoundDocumentInterface> RefDocInterface = RefAsset->GetOwningAsset();
 			if (RefDocInterface.GetObject() != nullptr)
 			{
-				const FAssetKey AssetKey(RefDocInterface->GetConstDocument().RootGraph.Metadata);
+				const FMetaSoundAssetKey AssetKey(RefDocInterface->GetConstDocument().RootGraph.Metadata);
 				if (AssetKey == NodeAssetKey)
 				{
 					return RefAsset;
diff --git a/Engine/Plugins/Runtime/Metasound/Source/MetasoundFrontend/Private/MetasoundFrontendNodeClassRegistry.cpp b/Engine/Plugins/Runtime/Metasound/Source/MetasoundFrontend/Private/MetasoundFrontendNodeClassRegistry.cpp
--- a/Engine/Plugins/Runtime/Metasound/Source/MetasoundFrontend/Private/MetasoundFrontendNodeClassRegistry.cpp
+++ b/Engine/Plugins/Runtime/Metasound/Source/MetasoundFrontend/Private/MetasoundFrontendNodeClassRegistry.cpp
@@ -982,7 +982,7 @@ PRAGMA_ENABLE_DEPRECATION_WARNINGS
 		{
 			if (IMetaSoundAssetManager* AssetManager = IMetaSoundAssetManager::Get())
 			{
-				const TArray<FTopLevelAssetPath> AssetPaths = AssetManager->FindAssetPaths(Frontend::FAssetKey(Key));
+				const TArray<FTopLevelAssetPath> AssetPaths = AssetManager->FindAssetPaths(FMetaSoundAssetKey(Key));
 				TArray<FString> ExistingPaths;
 				Algo::Transform(AssetPaths, ExistingPaths, [](const FTopLevelAssetPath& AssetPath) { return AssetPath.ToString(); });
 				const FString ExistingAssetPaths = FString::Join(ExistingPaths, TEXT("\n"));
diff --git a/Engine/Plugins/Runtime/Metasound/Source/MetasoundFrontend/Private/MetasoundFrontendRegistryKey.cpp b/Engine/Plugins/Runtime/Metasound/Source/MetasoundFrontend/Private/MetasoundFrontendRegistryKey.cpp
--- a/Engine/Plugins/Runtime/Metasound/Source/MetasoundFrontend/Private/MetasoundFrontendRegistryKey.cpp
+++ b/Engine/Plugins/Runtime/Metasound/Source/MetasoundFrontend/Private/MetasoundFrontendRegistryKey.cpp
@@ -184,7 +184,7 @@ namespace Metasound::Frontend
 		checkf(InClassInfo.Type != EMetasoundFrontendClassType::Graph, TEXT("Cannot create key from 'graph' type. Likely meant to use FNodeClassRegistryKey ctor that is provided FMetasoundFrontendGraphClass"));
 	}
 
-	FNodeClassRegistryKey::FNodeClassRegistryKey(const FAssetKey& AssetKey)
+	FNodeClassRegistryKey::FNodeClassRegistryKey(const FMetaSoundAssetKey& AssetKey)
 		: Type(EMetasoundFrontendClassType::External)
 		, ClassName(AssetKey.ClassName)
 		, Version(AssetKey.Version)
diff --git a/Engine/Plugins/Runtime/Metasound/Source/MetasoundFrontend/Public/MetasoundAssetBase.h b/Engine/Plugins/Runtime/Metasound/Source/MetasoundFrontend/Public/MetasoundAssetBase.h
--- a/Engine/Plugins/Runtime/Metasound/Source/MetasoundFrontend/Public/MetasoundAssetBase.h
+++ b/Engine/Plugins/Runtime/Metasound/Source/MetasoundFrontend/Public/MetasoundAssetBase.h
@@ -213,7 +213,12 @@ protected:
 	void OnNotifyBeginDestroy();
 
 #if WITH_EDITOR
-	virtual void SetReferencedAssetClasses(TSet<Metasound::Frontend::IMetaSoundAssetManager::FAssetInfo>&& InAssetClasses) = 0;
+PRAGMA_DISABLE_DEPRECATION_WARNINGS
+	UE_DEPRECATED(5.6, "Use SetReferencedAssets instead")
+	virtual void SetReferencedAssetClasses(TSet<Metasound::Frontend::IMetaSoundAssetManager::FAssetInfo>&& InAssetClasses) { }
+PRAGMA_ENABLE_DEPRECATION_WARNINGS
+
+	virtual void SetReferencedAssets(TSet<Metasound::Frontend::IMetaSoundAssetManager::FAssetRef>&& InAssetRefs) = 0;
 #endif
 
 	// Get information for communicating asynchronously with MetaSound running instance.
diff --git a/Engine/Plugins/Runtime/Metasound/Source/MetasoundFrontend/Public/MetasoundAssetKey.h b/Engine/Plugins/Runtime/Metasound/Source/MetasoundFrontend/Public/MetasoundAssetKey.h
--- a/Engine/Plugins/Runtime/Metasound/Source/MetasoundFrontend/Public/MetasoundAssetKey.h
+++ b/Engine/Plugins/Runtime/Metasound/Source/MetasoundFrontend/Public/MetasoundAssetKey.h
@@ -64,6 +64,5 @@ struct METASOUNDFRONTEND_API FMetaSoundAssetKey
 
 namespace Metasound::Frontend
 {
-	// Struct moved to global namespace and made a USTRUCT. This alias is here for back compat.
-	using FAssetKey = FMetaSoundAssetKey;
+	using FAssetKey UE_DEPRECATED(5.6, "Moved to global namespace as 'FMetaSoundAssetKey' and made a USTRUCT") = FMetaSoundAssetKey;
 } // namespace Metasound::Frontend
diff --git a/Engine/Plugins/Runtime/Metasound/Source/MetasoundFrontend/Public/MetasoundAssetManager.h b/Engine/Plugins/Runtime/Metasound/Source/MetasoundFrontend/Public/MetasoundAssetManager.h
--- a/Engine/Plugins/Runtime/Metasound/Source/MetasoundFrontend/Public/MetasoundAssetManager.h
+++ b/Engine/Plugins/Runtime/Metasound/Source/MetasoundFrontend/Public/MetasoundAssetManager.h
@@ -180,7 +180,23 @@ namespace Metasound::Frontend
 		UE_DEPRECATED(5.5, "Use Initialize/Deinitialize instead")
 		static void Set(IMetaSoundAssetManager& InInterface) { checkNoEntry(); }
 
-		struct FAssetInfo
+		struct FAssetRef
+		{
+			FMetaSoundAssetKey Key;
+			FTopLevelAssetPath Path;
+
+			FORCEINLINE friend bool operator==(const FAssetRef& InLHS, const FAssetRef& InRHS)
+			{
+				return (InLHS.Key == InRHS.Key) && (InLHS.Path == InRHS.Path);
+			}
+
+			FORCEINLINE friend uint32 GetTypeHash(const FAssetRef& InInfo)
+			{
+				return HashCombineFast(GetTypeHash(InInfo.Key), GetTypeHash(InInfo.Path));
+			}
+		};
+
+		struct UE_DEPRECATED(5.6, "Use FAssetRef instead") FAssetInfo
 		{
 			FNodeRegistryKey RegistryKey;
 			FSoftObjectPath AssetPath;
@@ -190,7 +206,7 @@ namespace Metasound::Frontend
 				return (InLHS.RegistryKey == InRHS.RegistryKey) && (InLHS.AssetPath == InRHS.AssetPath);
 			}
 
-			FORCEINLINE friend uint32 GetTypeHash(const IMetaSoundAssetManager::FAssetInfo& InInfo)
+			FORCEINLINE friend uint32 GetTypeHash(const FAssetInfo& InInfo)
 			{
 				return HashCombineFast(GetTypeHash(InInfo.RegistryKey), GetTypeHash(InInfo.AssetPath));
 			}
@@ -216,18 +232,18 @@ namespace Metasound::Frontend
 #endif // WITH_EDITORONLY_DATA
 
 		UE_DEPRECATED(5.6, "Moved to AddOrUpdateFromObject")
-		virtual FAssetKey AddOrUpdateAsset(const UObject& InObject) { return { }; }
+		virtual FMetaSoundAssetKey AddOrUpdateAsset(const UObject& InObject) { return { }; }
 
 		UE_DEPRECATED(5.6, "Moved to AddOrLoadAndUpdateFromObjectAsync")
-		virtual FAssetKey AddOrUpdateAsset(const FAssetData& InAssetData) { return { }; }
+		virtual FMetaSoundAssetKey AddOrUpdateAsset(const FAssetData& InAssetData) { return { }; }
 
 		// Add or Update a MetaSound Asset's entry data from a loaded MetaSound asset UObject.
-		virtual FAssetKey AddOrUpdateFromObject(const UObject& InObject) = 0;
+		virtual FMetaSoundAssetKey AddOrUpdateFromObject(const UObject& InObject) = 0;
 
 		// Add or Update a MetaSound Asset's entry data from an object, loading it if it isn't already.  On initial call, requests object load
 		// asynchronously and runs provided function on successful completion. If asset is already loaded, runs provided function immediately
 		// on entry update (synchronously).
-		virtual void AddOrLoadAndUpdateFromObjectAsync(const FAssetData& InAssetData, TUniqueFunction<void(FAssetKey, UObject&)>&& OnAssetRegistered) = 0;
+		virtual void AddOrLoadAndUpdateFromObjectAsync(const FAssetData& InAssetData, TUniqueFunction<void(FMetaSoundAssetKey, UObject&)>&& OnAssetRegistered) = 0;
 
 		// Add or Update a MetaSound Asset's entry data from AssetData. Potentially loads asset
 		// and adds asynchronously if asset or associated tag schema is out-of-date.
@@ -237,7 +253,7 @@ namespace Metasound::Frontend
 		virtual bool CanAutoUpdate(const FMetasoundFrontendClassName& InClassName) const = 0;
 
 		// Whether or not the asset manager has loaded the given asset
-		virtual bool ContainsKey(const Metasound::Frontend::FAssetKey& InAssetKey) const = 0;
+		virtual bool ContainsKey(const FMetaSoundAssetKey& InAssetKey) const = 0;
 
 		// Whether or not the asset manager has loaded one or more assets with the given registry key.
 		// Returns false if key is not valid asset key (ex. input or output class key, variable, etc.)
@@ -245,26 +261,31 @@ namespace Metasound::Frontend
 
 		// Returns object (if loaded) associated with the given key (null if key not registered with the AssetManager)
 		// If multiple assets are associated with the given key, the last one is returned. 
-		virtual FMetasoundAssetBase* FindAsset(const Metasound::Frontend::FAssetKey& InAssetKey) const = 0;
+		virtual FMetasoundAssetBase* FindAsset(const FMetaSoundAssetKey& InAssetKey) const = 0;
 
 		// Returns object (if loaded) associated with the given key as a Document Interface (null if key not registered with the AssetManager)
-		virtual TScriptInterface<IMetaSoundDocumentInterface> FindAssetAsDocumentInterface(const Frontend::FAssetKey& InKey) const = 0;
+		virtual TScriptInterface<IMetaSoundDocumentInterface> FindAssetAsDocumentInterface(const FMetaSoundAssetKey& InKey) const = 0;
 
 		// Returns path associated with the given key (returns invalid asset path if key not registered with the AssetManager or was not loaded from asset)
 		// If multiple assets are associated with the given key, the last one is returned.
-		virtual FTopLevelAssetPath FindAssetPath(const Metasound::Frontend::FAssetKey& InAssetKey) const = 0;
+		virtual FTopLevelAssetPath FindAssetPath(const FMetaSoundAssetKey& InAssetKey) const = 0;
 		
 		// Returns all paths associated with the given key (returns empty array if key not registered with the AssetManager or was not loaded from asset)
-		virtual TArray<FTopLevelAssetPath> FindAssetPaths(const Metasound::Frontend::FAssetKey& InAssetKey) const = 0;
+		virtual TArray<FTopLevelAssetPath> FindAssetPaths(const FMetaSoundAssetKey& InAssetKey) const = 0;
 
 		// Converts an object to an AssetBase if its a registered asset
 		virtual FMetasoundAssetBase* GetAsAsset(UObject& InObject) const = 0;
 		virtual const FMetasoundAssetBase* GetAsAsset(const UObject& InObject) const = 0;
 
 #if WITH_EDITOR
+PRAGMA_DISABLE_DEPRECATION_WARNINGS
+		UE_DEPRECATED(5.6, "Use GetReferencedAssets instead")
+		virtual TSet<FAssetInfo> GetReferencedAssetClasses(const FMetasoundAssetBase& InAssetBase) const { return {  }; }
+PRAGMA_DISABLE_DEPRECATION_WARNINGS
+
 		// Generates all asset info associated with registered assets that are referenced by the provided asset's graph.
-		virtual TSet<FAssetInfo> GetReferencedAssetClasses(const FMetasoundAssetBase& InAssetBase) const = 0;
-		
+		virtual TSet<FAssetRef> GetReferencedAssets(const FMetasoundAssetBase& InAssetBase) const = 0;
+
 		// Get assets this asset is a preset of, recursively for presets of presets.
 		// This means finding all references including the last non preset, but not graphs used directly in non-preset composition. 
 		// Ex. MetaSound A is a preset of MetaSound B, which is a preset of MetaSound C, which references D by composition -> When called on A, will return B and C.
@@ -276,7 +297,7 @@ namespace Metasound::Frontend
 		virtual bool IsAssetClass(const FMetasoundFrontendClassMetadata& ClassMetadata) const = 0;
 
 		UE_DEPRECATED(5.6, "Moved to Iterate ClassInfo")
-		virtual void IterateAssets(TFunctionRef<void(const FAssetKey, const TArray<FTopLevelAssetPath>&)> Iter) const { }
+		virtual void IterateAssets(TFunctionRef<void(const FMetaSoundAssetKey, const TArray<FTopLevelAssetPath>&)> Iter) const { }
 
 		UE_DEPRECATED(5.5, "Rescan no longer supported nor required by Frontend")
 		virtual void RescanAutoUpdateDenyList() { }
@@ -291,7 +312,7 @@ namespace Metasound::Frontend
 		virtual FMetasoundAssetBase* TryLoadAsset(const FSoftObjectPath& InObjectPath) const = 0;
 
 		// Returns asset associated with the given key (null if key not registered with the AssetManager or was not loaded from asset)
-		virtual FMetasoundAssetBase* TryLoadAssetFromKey(const Metasound::Frontend::FAssetKey& InAssetKey) const = 0;
+		virtual FMetasoundAssetBase* TryLoadAssetFromKey(const FMetaSoundAssetKey& InAssetKey) const = 0;
 
 		// Try to load referenced assets of the given asset or return them if they are already loaded (non-recursive).
 		// @return - True if all referenced assets successfully loaded, false if not.
diff --git a/Engine/Plugins/Runtime/Metasound/Source/MetasoundFrontend/Public/MetasoundFrontendRegistryKey.h b/Engine/Plugins/Runtime/Metasound/Source/MetasoundFrontend/Public/MetasoundFrontendRegistryKey.h
--- a/Engine/Plugins/Runtime/Metasound/Source/MetasoundFrontend/Public/MetasoundFrontendRegistryKey.h
+++ b/Engine/Plugins/Runtime/Metasound/Source/MetasoundFrontend/Public/MetasoundFrontendRegistryKey.h
@@ -12,9 +12,6 @@ struct FMetaSoundAssetKey;
 
 namespace Metasound::Frontend
 {
-	// Forward Declarations
-	using FAssetKey = FMetaSoundAssetKey;
-
 	/** FNodeClassInfo contains a minimal set of information needed to find and query node classes. */
 	struct METASOUNDFRONTEND_API FNodeClassInfo
 	{
@@ -77,7 +74,7 @@ namespace Metasound::Frontend
 		FNodeClassRegistryKey(const FMetasoundFrontendClassMetadata& InNodeMetadata);
 		FNodeClassRegistryKey(const FMetasoundFrontendGraphClass& InGraphClass);
 		FNodeClassRegistryKey(const FNodeClassInfo& InClassInfo);
-		FNodeClassRegistryKey(const FAssetKey& AssetKey);
+		FNodeClassRegistryKey(const FMetaSoundAssetKey& AssetKey);
 
 		FORCEINLINE friend bool operator==(const FNodeClassRegistryKey& InLHS, const FNodeClassRegistryKey& InRHS)
 		{
