Fix unbalanced deprecation macro
#rnx

diff --git a/Engine/Plugins/Runtime/Metasound/Source/MetasoundFrontend/Public/MetasoundAssetManager.h b/Engine/Plugins/Runtime/Metasound/Source/MetasoundFrontend/Public/MetasoundAssetManager.h
--- a/Engine/Plugins/Runtime/Metasound/Source/MetasoundFrontend/Public/MetasoundAssetManager.h
+++ b/Engine/Plugins/Runtime/Metasound/Source/MetasoundFrontend/Public/MetasoundAssetManager.h
@@ -278,7 +278,7 @@ namespace Metasound::Frontend
 		virtual const FMetasoundAssetBase* GetAsAsset(const UObject& InObject) const = 0;
 
 #if WITH_EDITOR
-PRAGMA_DISABLE_DEPRECATION_WARNINGS
+PRAGMA_ENABLE_DEPRECATION_WARNINGS
 		UE_DEPRECATED(5.6, "Use GetReferencedAssets instead")
 		virtual TSet<FAssetInfo> GetReferencedAssetClasses(const FMetasoundAssetBase& InAssetBase) const { return {  }; }
 PRAGMA_DISABLE_DEPRECATION_WARNINGS
