Fix up a couple more spots where borked macro was covering up deprecation issues.
#rnx

diff --git a/Engine/Plugins/Runtime/Metasound/Source/MetasoundEngine/Private/MetasoundAssetSubsystem.cpp b/Engine/Plugins/Runtime/Metasound/Source/MetasoundEngine/Private/MetasoundAssetSubsystem.cpp
--- a/Engine/Plugins/Runtime/Metasound/Source/MetasoundEngine/Private/MetasoundAssetSubsystem.cpp
+++ b/Engine/Plugins/Runtime/Metasound/Source/MetasoundEngine/Private/MetasoundAssetSubsystem.cpp
@@ -1492,11 +1492,13 @@ const FMetasoundAssetBase* UMetaSoundAssetSubsystem::GetAsAsset(const UObject& I
 }
 
 #if WITH_EDITOR
+PRAGMA_DISABLE_DEPRECATION_WARNINGS
 TSet<UMetaSoundAssetSubsystem::FAssetInfo> UMetaSoundAssetSubsystem::GetReferencedAssetClasses(const FMetasoundAssetBase& InAssetBase) const
 {
 	using namespace Metasound::Frontend;
 	return { };
 }
+PRAGMA_ENABLE_DEPRECATION_WARNINGS
 #endif // WITH_EDITOR
 
 FMetasoundAssetBase* UMetaSoundAssetSubsystem::TryLoadAssetFromKey(const Metasound::Frontend::FNodeRegistryKey& RegistryKey) const
diff --git a/Engine/Plugins/Runtime/Metasound/Source/MetasoundEngine/Public/MetasoundAssetSubsystem.h b/Engine/Plugins/Runtime/Metasound/Source/MetasoundEngine/Public/MetasoundAssetSubsystem.h
--- a/Engine/Plugins/Runtime/Metasound/Source/MetasoundEngine/Public/MetasoundAssetSubsystem.h
+++ b/Engine/Plugins/Runtime/Metasound/Source/MetasoundEngine/Public/MetasoundAssetSubsystem.h
@@ -285,8 +285,10 @@ PRAGMA_ENABLE_DEPRECATION_WARNINGS
 	static UMetaSoundAssetSubsystem& GetChecked();
 
 #if WITH_EDITOR
+PRAGMA_DISABLE_DEPRECATION_WARNINGS
 	UE_DEPRECATED(5.5, "Moved to internal implementation, use IMetaSoundAssetManager::GetChecked() and analogous call")
 	virtual TSet<FAssetInfo> GetReferencedAssetClasses(const FMetasoundAssetBase& InAssetBase) const;
+PRAGMA_ENABLE_DEPRECATION_WARNINGS
 #endif
 
 	UE_DEPRECATED(5.5, "Moved to internal implementation, use IMetaSoundAssetManager::GetChecked() and analogous call")
diff --git a/Engine/Plugins/Runtime/Metasound/Source/MetasoundFrontend/Public/MetasoundAssetManager.h b/Engine/Plugins/Runtime/Metasound/Source/MetasoundFrontend/Public/MetasoundAssetManager.h
--- a/Engine/Plugins/Runtime/Metasound/Source/MetasoundFrontend/Public/MetasoundAssetManager.h
+++ b/Engine/Plugins/Runtime/Metasound/Source/MetasoundFrontend/Public/MetasoundAssetManager.h
@@ -278,10 +278,10 @@ namespace Metasound::Frontend
 		virtual const FMetasoundAssetBase* GetAsAsset(const UObject& InObject) const = 0;
 
 #if WITH_EDITOR
-PRAGMA_ENABLE_DEPRECATION_WARNINGS
+PRAGMA_DISABLE_DEPRECATION_WARNINGS
 		UE_DEPRECATED(5.6, "Use GetReferencedAssets instead")
 		virtual TSet<FAssetInfo> GetReferencedAssetClasses(const FMetasoundAssetBase& InAssetBase) const { return {  }; }
-PRAGMA_DISABLE_DEPRECATION_WARNINGS
+PRAGMA_ENABLE_DEPRECATION_WARNINGS
 
 		// Generates all asset info associated with registered assets that are referenced by the provided asset's graph.
 		virtual TSet<FAssetRef> GetReferencedAssets(const FMetasoundAssetBase& InAssetBase) const = 0;
