Make MetaSound node deprecation macros on by default
[RN] The MetaSound Node API has been updated to support node configuration. If you have many nodes with use deprecated code paths, you can use UE_METASOUND_DISABLE_5_6_NODE_REGISTRATION_DEPRECATION_WARNINGS=1 to temporarily disable the deprecation warnings.

diff --git a/Engine/Plugins/Runtime/Metasound/Source/MetasoundFrontend/MetasoundFrontend.Build.cs b/Engine/Plugins/Runtime/Metasound/Source/MetasoundFrontend/MetasoundFrontend.Build.cs
--- a/Engine/Plugins/Runtime/Metasound/Source/MetasoundFrontend/MetasoundFrontend.Build.cs
+++ b/Engine/Plugins/Runtime/Metasound/Source/MetasoundFrontend/MetasoundFrontend.Build.cs
@@ -27,7 +27,15 @@ namespace UnrealBuildTool.Rules
 			);
 
 			PublicDefinitions.Add("WITH_METASOUND_FRONTEND=1");
-			PublicDefinitions.Add("UE_METASOUND_ENABLE_5_6_NODE_REGISTRATION_DEPRECATION_WARNINGS=1");
+
+			// In UE 5.6, registered node are expected to support the constructor 
+			// signature Constructor(FNodeData, TSharedRef<const FNodeClassMetadata>)
+			// Because there are many existing nodes, it may take time to update 
+			// them. For convenience, the deprecations related to this change are
+			// configurable via a preprocessor macro so that the deprecation warnings 
+			// do not drown out other compiler errors and warnings.
+			PublicDefinitions.Add("UE_METASOUND_DISABLE_5_6_NODE_REGISTRATION_DEPRECATION_WARNINGS=0");
+
 		}
 	}
 }
diff --git a/Engine/Plugins/Runtime/Metasound/Source/MetasoundFrontend/Private/MetasoundFrontendDataTypeRegistry.cpp b/Engine/Plugins/Runtime/Metasound/Source/MetasoundFrontend/Private/MetasoundFrontendDataTypeRegistry.cpp
--- a/Engine/Plugins/Runtime/Metasound/Source/MetasoundFrontend/Private/MetasoundFrontendDataTypeRegistry.cpp
+++ b/Engine/Plugins/Runtime/Metasound/Source/MetasoundFrontend/Private/MetasoundFrontendDataTypeRegistry.cpp
@@ -646,58 +646,41 @@ PRAGMA_ENABLE_DEPRECATION_WARNINGS
 				virtual TSharedPtr<const FNodeClassMetadata> GetVariableDeferredAccessorClassMetadata(const FName& InDataType) const override;
 
 				// Create a new instance of a C++ implemented node from the registry.
-#if UE_METASOUND_ENABLE_5_6_NODE_REGISTRATION_DEPRECATION_WARNINGS
 				UE_DEPRECATED(5.6, "Create input nodes with FNodeData")
-#endif
 				virtual TUniquePtr<INode> CreateInputNode(const FName& InInputType, FInputNodeConstructorParams&& InParams) const override;
 				virtual TUniquePtr<INode> CreateInputNode(const FName& InInputType, FNodeData InNodeData) const override;
 
-#if UE_METASOUND_ENABLE_5_6_NODE_REGISTRATION_DEPRECATION_WARNINGS
 				UE_DEPRECATED(5.6, "Create constructor input nodes with FNodeData")
-#endif
 				virtual TUniquePtr<INode> CreateConstructorInputNode(const FName& InInputType, FInputNodeConstructorParams&& InParams) const override;
 				virtual TUniquePtr<INode> CreateConstructorInputNode(const FName& InInputType, FNodeData InNodeData) const override;
 
 				UE_DEPRECATED(5.6, "Direct creation of literal nodes will no longer be supported")
 				virtual TUniquePtr<INode> CreateLiteralNode(const FName& InLiteralType, FLiteralNodeConstructorParams&& InParams) const override;
 
-#if UE_METASOUND_ENABLE_5_6_NODE_REGISTRATION_DEPRECATION_WARNINGS
 				UE_DEPRECATED(5.6, "Create output nodes with FNodeData")
-#endif
 				virtual TUniquePtr<INode> CreateOutputNode(const FName& InDataTypeName, FOutputNodeConstructorParams&& InParams) const override;
 				virtual TUniquePtr<INode> CreateOutputNode(const FName& InDataTypeName, FNodeData InNodeData) const override;
-#if UE_METASOUND_ENABLE_5_6_NODE_REGISTRATION_DEPRECATION_WARNINGS
+
 				UE_DEPRECATED(5.6, "Create output nodes with FNodeData")
-#endif
 				virtual TUniquePtr<INode> CreateConstructorOutputNode(const FName& InDataTypeName, FOutputNodeConstructorParams&& InParams) const override;
 				virtual TUniquePtr<INode> CreateConstructorOutputNode(const FName& InDataTypeName, FNodeData InNodeData) const override;
 
-#if UE_METASOUND_ENABLE_5_6_NODE_REGISTRATION_DEPRECATION_WARNINGS
 				UE_DEPRECATED(5.6, "Direct creation of receive nodes will no longer be supported")
-#endif
 				virtual TUniquePtr<INode> CreateReceiveNode(const FName& InDataTypeName, const FNodeInitData& InParams) const override;
 
-#if UE_METASOUND_ENABLE_5_6_NODE_REGISTRATION_DEPRECATION_WARNINGS
 				UE_DEPRECATED(5.6, "Create Variable Node with FNodeData")
-#endif
 				virtual TUniquePtr<INode> CreateVariableNode(const FName& InDataType, FVariableNodeConstructorParams&&) const override;
 				virtual TUniquePtr<INode> CreateVariableNode(const FName& InDataType, FLiteral InLiteral, FNodeData InNodeData) const override;
 
-#if UE_METASOUND_ENABLE_5_6_NODE_REGISTRATION_DEPRECATION_WARNINGS
 				UE_DEPRECATED(5.6, "Create variable mutator nodes with FNodeData")
-#endif
 				virtual TUniquePtr<INode> CreateVariableMutatorNode(const FName& InDataType, const FNodeInitData&) const override;
 				virtual TUniquePtr<INode> CreateVariableMutatorNode(const FName& InDataType, FNodeData InNodeData) const override;
 
-#if UE_METASOUND_ENABLE_5_6_NODE_REGISTRATION_DEPRECATION_WARNINGS
 				UE_DEPRECATED(5.6, "Create variable accessor nodes with FNodeData")
-#endif
 				virtual TUniquePtr<INode> CreateVariableAccessorNode(const FName& InDataType, const FNodeInitData&) const override;
 				virtual TUniquePtr<INode> CreateVariableAccessorNode(const FName& InDataType, FNodeData InNodeData) const override;
 
-#if UE_METASOUND_ENABLE_5_6_NODE_REGISTRATION_DEPRECATION_WARNINGS
 				UE_DEPRECATED(5.6, "Create variable deferred accessor nodes with FNodeData")
-#endif
 				virtual TUniquePtr<INode> CreateVariableDeferredAccessorNode(const FName& InDataType, const FNodeInitData&) const override;
 				virtual TUniquePtr<INode> CreateVariableDeferredAccessorNode(const FName& InDataType, FNodeData InNodeData) const override;
 
diff --git a/Engine/Plugins/Runtime/Metasound/Source/MetasoundFrontend/Private/MetasoundFrontendNodeClassRegistry.cpp b/Engine/Plugins/Runtime/Metasound/Source/MetasoundFrontend/Private/MetasoundFrontendNodeClassRegistry.cpp
--- a/Engine/Plugins/Runtime/Metasound/Source/MetasoundFrontend/Private/MetasoundFrontendNodeClassRegistry.cpp
+++ b/Engine/Plugins/Runtime/Metasound/Source/MetasoundFrontend/Private/MetasoundFrontendNodeClassRegistry.cpp
@@ -30,6 +30,10 @@
 #include "MetasoundRouter.h"
 #include "MetasoundTrace.h"
 
+#ifndef UE_METASOUND_DISABLE_5_6_NODE_REGISTRATION_DEPRECATION_WARNINGS
+#define UE_METASOUND_DISABLE_5_6_NODE_REGISTRATION_DEPRECATION_WARNINGS (0)
+#endif
+
 namespace Metasound::Frontend
 {
 	namespace ConsoleVariables
@@ -311,7 +315,7 @@ namespace Metasound::Frontend
 	  
 	TUniquePtr<INode> INodeClassRegistryEntry::CreateNode(FNodeData InNodeData) const
 	{
-#if UE_METASOUND_ENABLE_5_6_NODE_REGISTRATION_DEPRECATION_WARNINGS
+#if !UE_METASOUND_DISABLE_5_6_NODE_REGISTRATION_DEPRECATION_WARNINGS
 		static bool bDidLogError = false;
 
 		if (!bDidLogError)
@@ -564,7 +568,6 @@ PRAGMA_ENABLE_DEPRECATION_WARNINGS
 	{
 		TUniquePtr<INode> Node;
 
-#if UE_METASOUND_ENABLE_5_6_NODE_REGISTRATION_DEPRECATION_WARNINGS
 		auto CreateNodeLambda = [&Node, &InInitData](const INodeClassRegistryEntry& Entry) mutable
 		{ 
 			const FMetasoundFrontendClass& FrontendClass = Entry.GetFrontendClass();
@@ -577,12 +580,6 @@ PRAGMA_ENABLE_DEPRECATION_WARNINGS
 			};
 			Node = Entry.CreateNode(MoveTemp(NodeData)); 
 		};
-#else 
-		auto CreateNodeLambda = [&Node, &InInitData](const INodeClassRegistryEntry& Entry) mutable
-		{ 
-			Node = Entry.CreateNode(InInitData); 
-		};
-#endif
 
 		if (!AccessNodeEntryThreadSafe(InKey, CreateNodeLambda))
 		{
diff --git a/Engine/Plugins/Runtime/Metasound/Source/MetasoundFrontend/Private/MetasoundFrontendNodeClassRegistryPrivate.h b/Engine/Plugins/Runtime/Metasound/Source/MetasoundFrontend/Private/MetasoundFrontendNodeClassRegistryPrivate.h
--- a/Engine/Plugins/Runtime/Metasound/Source/MetasoundFrontend/Private/MetasoundFrontendNodeClassRegistryPrivate.h
+++ b/Engine/Plugins/Runtime/Metasound/Source/MetasoundFrontend/Private/MetasoundFrontendNodeClassRegistryPrivate.h
@@ -106,9 +106,7 @@ namespace Metasound::Frontend
 		virtual bool FindOutputNodeRegistryKeyForDataType(const FName& InDataTypeName, const EMetasoundFrontendVertexAccessType InAccessType, FNodeClassRegistryKey& OutKey) override;
 
 		// Create a new instance of a C++ implemented node from the registry.
-#if UE_METASOUND_ENABLE_5_6_NODE_REGISTRATION_DEPRECATION_WARNINGS
 		UE_DEPRECATED(5.6, "Node classes should be constructed with FNodeData")
-#endif
 		virtual TUniquePtr<Metasound::INode> CreateNode(const FNodeClassRegistryKey& InKey, const Metasound::FNodeInitData& InInitData) const override;
 
 		// Create a new instance of a C++ implemented node from the registry.
diff --git a/Engine/Plugins/Runtime/Metasound/Source/MetasoundFrontend/Public/MetasoundFrontendDataTypeRegistry.h b/Engine/Plugins/Runtime/Metasound/Source/MetasoundFrontend/Public/MetasoundFrontendDataTypeRegistry.h
--- a/Engine/Plugins/Runtime/Metasound/Source/MetasoundFrontend/Public/MetasoundFrontendDataTypeRegistry.h
+++ b/Engine/Plugins/Runtime/Metasound/Source/MetasoundFrontend/Public/MetasoundFrontendDataTypeRegistry.h
@@ -13,14 +13,6 @@
 #include "UObject/Class.h"
 #include "UObject/NameTypes.h"
 
-// In UE 5.6, registered node are expected to support the constructor signature Constructor(FNodeData, TSharedRef<const FNodeClassMetadata>)
-// Because there are many existing nodes, it may take time to update them. For convenience, the deprecations related to this change are
-// configurable via a preprocessor macro so that the deprecation warnings do not drown out other compiler errors and warnings.
-#ifndef UE_METASOUND_ENABLE_5_6_NODE_REGISTRATION_DEPRECATION_WARNINGS
-#define UE_METASOUND_ENABLE_5_6_NODE_REGISTRATION_DEPRECATION_WARNINGS (0)
-#endif
-
-
 namespace Metasound
 {
 	class IDataChannel;
@@ -401,52 +393,41 @@ namespace Metasound
 			 * such as Inputs, Outputs, and Variables. 
 			 */
 
-#if UE_METASOUND_ENABLE_5_6_NODE_REGISTRATION_DEPRECATION_WARNINGS
 			UE_DEPRECATED(5.6, "Create input nodes with FNodeData")
-#endif
 			virtual TUniquePtr<INode> CreateInputNode(const FName& InInputType, FInputNodeConstructorParams&& InParams) const = 0;
 			virtual TUniquePtr<INode> CreateInputNode(const FName& InInputType, FNodeData InNodeData) const = 0;
-#if UE_METASOUND_ENABLE_5_6_NODE_REGISTRATION_DEPRECATION_WARNINGS
+
 			UE_DEPRECATED(5.6, "Create constructor input nodes with FNodeData")
-#endif
 			virtual TUniquePtr<INode> CreateConstructorInputNode(const FName& InInputType, FInputNodeConstructorParams&& InParams) const = 0;
 			virtual TUniquePtr<INode> CreateConstructorInputNode(const FName& InInputType, FNodeData InNodeData) const = 0;
-#if UE_METASOUND_ENABLE_5_6_NODE_REGISTRATION_DEPRECATION_WARNINGS
+
 			UE_DEPRECATED(5.6, "Direct creation of literal nodes will no longer be supported")
-#endif
 			virtual TUniquePtr<INode> CreateLiteralNode(const FName& InLiteralType, FLiteralNodeConstructorParams&& InParams) const = 0;
-#if UE_METASOUND_ENABLE_5_6_NODE_REGISTRATION_DEPRECATION_WARNINGS
+
 			UE_DEPRECATED(5.6, "Create output nodes with FNodeData")
-#endif
 			virtual TUniquePtr<INode> CreateOutputNode(const FName& InDataTypeName, FOutputNodeConstructorParams&& InParams) const = 0;
 			virtual TUniquePtr<INode> CreateOutputNode(const FName& InDataTypeName, FNodeData InNodeData) const = 0;
-#if UE_METASOUND_ENABLE_5_6_NODE_REGISTRATION_DEPRECATION_WARNINGS
+
 			UE_DEPRECATED(5.6, "Create output nodes with FNodeData")
-#endif
 			virtual TUniquePtr<INode> CreateConstructorOutputNode(const FName& InDataTypeName, FOutputNodeConstructorParams&& InParams) const = 0;
 			virtual TUniquePtr<INode> CreateConstructorOutputNode(const FName& InDataTypeName, FNodeData InNodeData) const = 0;
-#if UE_METASOUND_ENABLE_5_6_NODE_REGISTRATION_DEPRECATION_WARNINGS
+
 			UE_DEPRECATED(5.6, "Direct creation of receive nodes will no longer be supported")
-#endif
 			virtual TUniquePtr<INode> CreateReceiveNode(const FName& InDataType, const FNodeInitData&) const = 0;
-#if UE_METASOUND_ENABLE_5_6_NODE_REGISTRATION_DEPRECATION_WARNINGS
+
 			UE_DEPRECATED(5.6, "Create Variable Node with FNodeData")
-#endif
 			virtual TUniquePtr<INode> CreateVariableNode(const FName& InDataType, FVariableNodeConstructorParams&& InParams) const = 0;
 			virtual TUniquePtr<INode> CreateVariableNode(const FName& InDataType, FLiteral InLiteral, FNodeData InNodeData) const = 0;
-#if UE_METASOUND_ENABLE_5_6_NODE_REGISTRATION_DEPRECATION_WARNINGS
+
 			UE_DEPRECATED(5.6, "Create variable mutator nodes with FNodeData")
-#endif
 			virtual TUniquePtr<INode> CreateVariableMutatorNode(const FName& InDataType, const FNodeInitData& InParams) const = 0;
 			virtual TUniquePtr<INode> CreateVariableMutatorNode(const FName& InDataType, FNodeData InNodeData) const = 0;
-#if UE_METASOUND_ENABLE_5_6_NODE_REGISTRATION_DEPRECATION_WARNINGS
+
 			UE_DEPRECATED(5.6, "Create variable accessor nodes with FNodeData")
-#endif
 			virtual TUniquePtr<INode> CreateVariableAccessorNode(const FName& InDataType, const FNodeInitData& InParams) const = 0;
 			virtual TUniquePtr<INode> CreateVariableAccessorNode(const FName& InDataType, FNodeData InNodeData) const = 0;
-#if UE_METASOUND_ENABLE_5_6_NODE_REGISTRATION_DEPRECATION_WARNINGS
+
 			UE_DEPRECATED(5.6, "Create variable deferred accessor nodes with FNodeData")
-#endif
 			virtual TUniquePtr<INode> CreateVariableDeferredAccessorNode(const FName& InDataType, const FNodeInitData& InParams) const = 0;
 			virtual TUniquePtr<INode> CreateVariableDeferredAccessorNode(const FName& InDataType, FNodeData InNodeData) const = 0;
 		};
diff --git a/Engine/Plugins/Runtime/Metasound/Source/MetasoundFrontend/Public/MetasoundFrontendNodeClassRegistry.h b/Engine/Plugins/Runtime/Metasound/Source/MetasoundFrontend/Public/MetasoundFrontendNodeClassRegistry.h
--- a/Engine/Plugins/Runtime/Metasound/Source/MetasoundFrontend/Public/MetasoundFrontendNodeClassRegistry.h
+++ b/Engine/Plugins/Runtime/Metasound/Source/MetasoundFrontend/Public/MetasoundFrontendNodeClassRegistry.h
@@ -8,13 +8,6 @@
 #include "MetasoundVertex.h"
 #include "Templates/Function.h"
 
-// In UE 5.6, registered node are expected to support the constructor signature Constructor(FNodeData, TSharedRef<const FNodeClassMetadata>)
-// Because there are many existing nodes, it may take time to update them. For convenience, the deprecations related to this change are
-// configurable via a preprocessor macro so that the deprecation warnings do not drown out other compiler errors and warnings.
-#ifndef UE_METASOUND_ENABLE_5_6_NODE_REGISTRATION_DEPRECATION_WARNINGS
-#define UE_METASOUND_ENABLE_5_6_NODE_REGISTRATION_DEPRECATION_WARNINGS (0)
-#endif
-
 #ifndef UE_METASOUND_PURE_VIRTUAL_CREATE_FRONTEND_NODE_EXTENSION
 #define UE_METASOUND_PURE_VIRTUAL_CREATE_FRONTEND_NODE_EXTENSION (0)
 #endif
@@ -82,9 +75,7 @@ namespace Metasound::Frontend
 			* If a node can be created with FNodeInitData, this function
 			* should return a valid node pointer.
 			*/
-#if UE_METASOUND_ENABLE_5_6_NODE_REGISTRATION_DEPRECATION_WARNINGS
 		UE_DEPRECATED(5.6, "Node classes should be constructed with FNodeData")
-#endif
 		virtual TUniquePtr<INode> CreateNode(const FNodeInitData&) const = 0;
 
 		/** Create a node given FNodeData.
@@ -356,9 +347,7 @@ namespace Metasound::Frontend
 		 * 
 		 * Returns an invalid TUniquePtr<> if the key is not in the registry. 
 		 */
-#if UE_METASOUND_ENABLE_5_6_NODE_REGISTRATION_DEPRECATION_WARNINGS
 		UE_DEPRECATED(5.6, "Node classes should be constructed with FNodeData")
-#endif
 		virtual TUniquePtr<Metasound::INode> CreateNode(const FNodeClassRegistryKey& InKey, const Metasound::FNodeInitData&) const = 0;
 
 		UE_DEPRECATED(5.6, "Node classes should be constructed with FNodeData")
diff --git a/Engine/Plugins/Runtime/Metasound/Source/MetasoundFrontend/Public/MetasoundNodeRegistrationMacro.h b/Engine/Plugins/Runtime/Metasound/Source/MetasoundFrontend/Public/MetasoundNodeRegistrationMacro.h
--- a/Engine/Plugins/Runtime/Metasound/Source/MetasoundFrontend/Public/MetasoundNodeRegistrationMacro.h
+++ b/Engine/Plugins/Runtime/Metasound/Source/MetasoundFrontend/Public/MetasoundNodeRegistrationMacro.h
@@ -13,6 +13,13 @@
 #include "Traits/MetasoundNodeConstructorTraits.h"
 #include "Traits/MetasoundNodeStaticMemberTraits.h"
 
+// In UE 5.6, registered node are expected to support the constructor signature Constructor(FNodeData, TSharedRef<const FNodeClassMetadata>)
+// Because there are many existing nodes, it may take time to update them. For convenience, the deprecations related to this change are
+// configurable via a preprocessor macro so that the deprecation warnings do not drown out other compiler errors and warnings.
+#ifndef UE_METASOUND_DISABLE_5_6_NODE_REGISTRATION_DEPRECATION_WARNINGS
+#define UE_METASOUND_DISABLE_5_6_NODE_REGISTRATION_DEPRECATION_WARNINGS (0)
+#endif
+
 namespace Metasound::Frontend
 {
 	namespace NodeRegistrationPrivate
@@ -32,9 +39,7 @@ namespace Metasound::Frontend
 			// Unhide CreateNode overloads which exist on INodeClassRegistryEntry
 			using INodeClassRegistryEntry::CreateNode;
 
-#if UE_METASOUND_ENABLE_5_6_NODE_REGISTRATION_DEPRECATION_WARNINGS
-		UE_DEPRECATED(5.6, "Node classes should be constructed with FNodeData")
-#endif
+			UE_DEPRECATED(5.6, "Node classes should be constructed with FNodeData")
 			virtual TUniquePtr<INode> CreateNode(const FNodeInitData& InInitData) const override;
 
 			virtual const TSet<FMetasoundFrontendVersion>* GetImplementedInterfaces() const override;
@@ -109,20 +114,21 @@ namespace Metasound::Frontend
 			using TNodeRegistryEntryBase<NodeType>::TNodeRegistryEntryBase;
 		};
 
-#if UE_METASOUND_ENABLE_5_6_NODE_REGISTRATION_DEPRECATION_WARNINGS
+#if UE_METASOUND_DISABLE_5_6_NODE_REGISTRATION_DEPRECATION_WARNINGS
+
 		template<typename T>
-		UE_DEPRECATED(5.6, "Update the node's constructor to be of the form FMyNode(FNodeData InNodeData, TSharedRef<const FNodeClassMetadata> InClassMetadata)")
 		void TriggerDeprecatedNodeConstructorWarning() {}
 
 		template<typename T>
-		UE_DEPRECATED(5.6, "Update the node class to include 'static FNodeClassMetadata CreateNodeClassMetadata()`)")
 		void TriggerMissingCreateNodeClassMetadataWarning() {}
-#else
-		// These only trigger deprecation if the macro UE_METASOUND_ENABLED_5_6_NODE_REGISTRATION_DEPRECATION_WARNINGS is true
+#else		
+		// These only trigger deprecation if the macro UE_METASOUND_DISABLE_5_6_NODE_REGISTRATION_DEPRECATION_WARNINGS is false
 		template<typename T>
+		UE_DEPRECATED(5.6, "Update the node's constructor to be of the form FMyNode(FNodeData InNodeData, TSharedRef<const FNodeClassMetadata> InClassMetadata)")
 		void TriggerDeprecatedNodeConstructorWarning() {}
 
 		template<typename T>
+		UE_DEPRECATED(5.6, "Update the node class to include 'static FNodeClassMetadata CreateNodeClassMetadata()`)")
 		void TriggerMissingCreateNodeClassMetadataWarning() {}
 #endif
 
diff --git a/Engine/Plugins/Runtime/Metasound/Source/MetasoundGraphCore/Public/MetasoundDynamicOperatorTransactor.h b/Engine/Plugins/Runtime/Metasound/Source/MetasoundGraphCore/Public/MetasoundDynamicOperatorTransactor.h
--- a/Engine/Plugins/Runtime/Metasound/Source/MetasoundGraphCore/Public/MetasoundDynamicOperatorTransactor.h
+++ b/Engine/Plugins/Runtime/Metasound/Source/MetasoundGraphCore/Public/MetasoundDynamicOperatorTransactor.h
@@ -18,13 +18,6 @@
 #define METASOUND_DEBUG_DYNAMIC_TRANSACTOR !UE_BUILD_SHIPPING
 #endif
 
-// In UE 5.6, registered node are expected to support the constructor signature Constructor(FNodeData, TSharedRef<const FNodeClassMetadata>)
-// Because there are many existing nodes, it may take time to update them. For convenience, the deprecations related to this change are
-// configurable via a preprocessor macro so that the deprecation warnings do not drown out other compiler errors and warnings.
-#ifndef UE_METASOUND_ENABLE_5_6_NODE_REGISTRATION_DEPRECATION_WARNINGS
-#define UE_METASOUND_ENABLE_5_6_NODE_REGISTRATION_DEPRECATION_WARNINGS (0)
-#endif
-
 namespace Metasound
 {
 	class FOperatorSettings;
@@ -215,18 +208,16 @@ namespace Metasound
 			 */
 			void AddDataEdge(const FGuid& InFromNodeID, const FVertexName& InFromVertex, const FGuid& InToNodeID, const FVertexName& InToVertex);
 
-#if UE_METASOUND_ENABLE_5_6_NODE_REGISTRATION_DEPRECATION_WARNINGS
 			UE_DEPRECATED(5.6, "Literal nodes will no longer be supported. Use RemoveDataEdge with an FLiteral instead.")
-#endif
-			/** Remove the given data edge. */
 			void RemoveDataEdge(const FGuid& InFromNode, const FVertexName& InFromVertexName, const FGuid& InToNode, const FVertexName& InToVertexName, TUniquePtr<INode> InReplacementLiteralNode);
+
+			/** Remove the given data edge. */
 			void RemoveDataEdge(const FGuid& InFromNode, const FVertexName& InFromVertexName, const FGuid& InToNode, const FVertexName& InToVertexName, FLiteral InReplacementLiteral, const FReferenceCreationFunction InReferenceCreateFunc);
 
-#if UE_METASOUND_ENABLE_5_6_NODE_REGISTRATION_DEPRECATION_WARNINGS
 			UE_DEPRECATED(5.6, "Setting values using a literal node will no longer be supported. Set literal on the vertex interface instead.")
-#endif
-			/** Set the value on a unconnected node input vertex. */
 			void SetValue(const FGuid& InNodeID, const FVertexName& InVertexName, TUniquePtr<INode> InLiteralNode);
+
+			/** Set the value on a unconnected node input vertex. */
 			void SetValue(const FGuid& InNodeID, const FVertexName& InVertexName, const FLiteral& InLiteral, const FReferenceCreationFunction InReferenceCreateFunc);
 
 			/** Add an input data destination to describe how data provided 
