- Version MetaSound folder action refactor
- No longer require script to load all assets if asset tag data already reports asset is up-to-date
- Renamed "Version MetaSound Asset" instead of "Version MetaSound Document" where applicable
- Add recursive & non-recursive flavors
- Reserialize asset if document needs versioning OR asset tags need to be updated

#rnx

diff --git a/Engine/Plugins/Runtime/Metasound/Source/MetasoundEditor/Private/MetasoundEditorModule.cpp b/Engine/Plugins/Runtime/Metasound/Source/MetasoundEditor/Private/MetasoundEditorModule.cpp
--- a/Engine/Plugins/Runtime/Metasound/Source/MetasoundEditor/Private/MetasoundEditorModule.cpp
+++ b/Engine/Plugins/Runtime/Metasound/Source/MetasoundEditor/Private/MetasoundEditorModule.cpp
@@ -120,12 +120,10 @@ namespace Metasound::Editor
 				return;
 			}
 
-			Menu->AddMenuEntry("MetaSounds", FToolMenuEntry::InitMenuEntry(
-				"VersionMetaSounds",
-				LOCTEXT("ResaveMetaSoundAssetsMenuEntry", "Version MetaSounds"),
-				LOCTEXT("ResaveMetaSoundAssetsMenuEntryTooltip", "Load, version MetaSound asset document (if necessary), and resave all MetaSounds in the given folder context. Does not replace deprecated node classes with newer major versions."),
-				FSlateIcon(FAppStyle::GetAppStyleSetName(), "Icons.Save"),
-				FToolUIActionChoice(FExecuteAction::CreateLambda([FolderPaths = Context->GetSelectedPackagePaths()]()
+			bool bRecursePaths = false;
+			auto CreateVersioningExecutionLambda = [&Context](bool bRecursePaths) -> FExecuteAction
+			{
+				return FExecuteAction::CreateLambda([FolderPaths = Context->GetSelectedPackagePaths(), bRecursePaths]()
 				{
 					using namespace Frontend;
 
@@ -139,12 +137,12 @@ namespace Metasound::Editor
 						Notify->SetCompletionState(SNotificationItem::CS_Pending);
 					}
 
-					IMetaSoundAssetManager::FVersionDocumentResults Results;
+					IMetaSoundAssetManager::FVersionAssetResults Results;
 					{
 						FGCScopeGuard Guard;
-						UE::Tasks::FTask ProcessDocuments = UE::Tasks::Launch(UE_SOURCE_LOCATION, [&Results, FolderPaths]()
+						UE::Tasks::FTask ProcessDocuments = UE::Tasks::Launch(UE_SOURCE_LOCATION, [&Results, FolderPaths, bRecursePaths]()
 						{
-							Results = IMetaSoundAssetManager::GetChecked().VersionDocumentsInFolders(FolderPaths);
+							Results = IMetaSoundAssetManager::GetChecked().VersionAssetsInFolders(FolderPaths, bRecursePaths);
 						});
 
 						// TODO: Report progress and tick instead of immediate wait.
@@ -156,7 +154,7 @@ namespace Metasound::Editor
 					{
 						if (Notify.IsValid())
 						{
-							Notify->SetCompletionState(SNotificationItem::CS_Fail);
+							Notify->SetCompletionState(SNotificationItem::CS_None);
 						}
 						UE_LOG(LogMetasoundEditor, Display, TEXT("No MetaSound(s) Versioned: Folder/child folder(s) contain(s) no MetaSound asset(s)."));
 					}
@@ -164,7 +162,7 @@ namespace Metasound::Editor
 					{
 						if (Notify.IsValid())
 						{
-							Notify->SetCompletionState(SNotificationItem::CS_Fail);
+							Notify->SetCompletionState(SNotificationItem::CS_None);
 						}
 						UE_LOG(LogMetasoundEditor, Display, TEXT("No MetaSound(s) Versioned: Folder/child folder(s) contain(s) no MetaSound asset(s) requiring versioning."));
 					}
@@ -182,7 +180,28 @@ namespace Metasound::Editor
 					{
 						Notify->ExpireAndFadeout();
 					}
-				}))
+				});
+			};
+
+			const FText ResaveEntryWarning = LOCTEXT("ResaveMetaSoundAssetsMenuEntry_Warning", "Does not resave if replace deprecated node classes are found with newer major versions.");
+			const FText ResaveContext = LOCTEXT("ResaveMetaSoundAssetMenuEntry_Context", "Load, version MetaSound asset document and update asset tags if either are necessary.");
+
+			const FText ResaveContextFlat = LOCTEXT("ResaveMetaSoundAssetsMenuEntry_ContextFlat", "Resaves all MetaSounds in the given folder.");
+			Menu->AddMenuEntry("MetaSounds", FToolMenuEntry::InitMenuEntry(
+				"VersionMetaSoundsFlat",
+				LOCTEXT("ResaveMetaSoundAssetsMenuEntry", "Version MetaSounds"),
+				FText::Format(LOCTEXT("ResaveMetaSoundAssetsMenuEntryTooltip", "{0} {1} {2}"), ResaveContext, ResaveContextFlat, ResaveEntryWarning),
+				FSlateIcon(FAppStyle::GetAppStyleSetName(), "Icons.Save"),
+				FToolUIActionChoice(CreateVersioningExecutionLambda(false /* bRecursePaths */))
+			));
+
+			const FText ResaveContextRecursive = LOCTEXT("ResaveMetaSoundAssetsMenuEntry_RecurseContext", "Resaves all MetaSounds in the given folder and all sub-folders.");
+			Menu->AddMenuEntry("MetaSounds", FToolMenuEntry::InitMenuEntry(
+				"VersionMetaSoundsRecursively",
+				LOCTEXT("ResaveMetaSoundAssetsMenuEntry_Recursive", "Version MetaSounds (Recursive)"),
+				FText::Format(LOCTEXT("ResaveMetaSoundAssetsMenuEntryTooltip", "{0} {1} {2}"), ResaveContextRecursive, ResaveContextFlat, ResaveEntryWarning),
+				FSlateIcon(FAppStyle::GetAppStyleSetName(), "Icons.Save"),
+				FToolUIActionChoice(CreateVersioningExecutionLambda(true /* bRecursePaths */))
 			));
 		}
 	} // namespace ModulePrivate
diff --git a/Engine/Plugins/Runtime/Metasound/Source/MetasoundEngine/Private/MetasoundAssetSubsystem.cpp b/Engine/Plugins/Runtime/Metasound/Source/MetasoundEngine/Private/MetasoundAssetSubsystem.cpp
--- a/Engine/Plugins/Runtime/Metasound/Source/MetasoundEngine/Private/MetasoundAssetSubsystem.cpp
+++ b/Engine/Plugins/Runtime/Metasound/Source/MetasoundEngine/Private/MetasoundAssetSubsystem.cpp
@@ -220,7 +220,7 @@ namespace Metasound::Engine
 		return AssetSubsystemPrivate::AddClassInfo(TagDataMapCriticalSection, ClassInfoMap, FMetaSoundAssetClassInfo (DocInterface));
 	}
 
-	void  FMetaSoundAssetManager::AddOrLoadAndUpdateFromObjectAsync(const FAssetData& InAssetData, TUniqueFunction<void(FMetaSoundAssetKey, UObject&)>&& OnAssetRegistered)
+	void FMetaSoundAssetManager::AddOrLoadAndUpdateFromObjectAsync(const FAssetData& InAssetData, TUniqueFunction<void(FMetaSoundAssetKey, UObject&)>&& OnAssetRegistered)
 	{
 		using namespace Frontend;
 
@@ -1241,7 +1241,7 @@ namespace Metasound::Engine
 	}
 
 #if WITH_EDITORONLY_DATA
-	FMetaSoundAssetManager::FVersionDocumentResults FMetaSoundAssetManager::VersionDocumentsInFolders(const TArray<FString>& FolderPaths) const
+	FMetaSoundAssetManager::FVersionAssetResults FMetaSoundAssetManager::VersionAssetsInFolders(const TArray<FString>& FolderPaths, bool bRecursePaths) const
 	{
 		TArray<FTopLevelAssetPath> ClassNames;
 		IMetasoundUObjectRegistry::Get().IterateRegisteredUClasses([&ClassNames](UClass& InClass)
@@ -1251,45 +1251,59 @@ namespace Metasound::Engine
 
 		FARFilter Filter;
 		Filter.ClassPaths = ClassNames;
-		Filter.bRecursivePaths = true;
+		Filter.bRecursivePaths = bRecursePaths;
 		Filter.bRecursiveClasses = true;
 		Algo::Transform(FolderPaths, Filter.PackagePaths, [](const FString& PackagePath)
 		{
 			return FName(*PackagePath);
 		});
 
-		FVersionDocumentResults Results;
+		FVersionAssetResults Results;
 		FAssetRegistryModule& AssetRegistryModule = FModuleManager::LoadModuleChecked<FAssetRegistryModule>("AssetRegistry");
 		AssetRegistryModule.Get().EnumerateAssets(Filter, [&Results](const FAssetData& AssetData)
 		{
 			using namespace Frontend;
 
-			if (UObject* MetaSound = AssetData.GetAsset())
+			const FTopLevelAssetPath PackagePath(AssetData.PackageName, AssetData.AssetName);
+			const FMetaSoundAssetClassInfo ClassInfo(AssetData);
+			if (ClassInfo.bIsValid)
 			{
-				if (FMetasoundAssetBase* MetaSoundAsset = IMetasoundUObjectRegistry::Get().GetObjectAsAssetBase(MetaSound))
+				if (ClassInfo.DocInfo.DocumentVersion >= FMetasoundFrontendDocument::GetMaxVersion())
 				{
-					UMetaSoundBuilderBase& Builder = FDocumentBuilderRegistry::GetChecked().FindOrBeginBuilding(*MetaSound);
-					FMetaSoundFrontendDocumentBuilder& DocBuilder = Builder.GetBuilder();
-					const FString MetaSoundName = DocBuilder.GetDebugName();
-					const bool bVersionedOnLoad = MetaSoundAsset->GetVersionedOnLoad();
-					const FMetasoundFrontendDocumentMetadata& DocMetadata = DocBuilder.GetConstDocumentChecked().Metadata;
-					if (bVersionedOnLoad || MetaSoundAsset->VersionAsset(DocBuilder))
-					{
-						const FString DocVersion = DocMetadata.Version.Number.ToString();
-						constexpr bool bAlwaysMarkDirty = true;
-						MetaSound->Modify(bAlwaysMarkDirty);
-						MetaSoundAsset->SetVersionedOnLoad();
-						Results.VersionedPackages.Add(MetaSound->GetPackage());
-						UE_LOG(LogMetaSound, Display, TEXT("MetaSound asset '%s' document versioned to %s."), *MetaSoundName, *DocVersion);
-					}
-					else
-					{
-						const FString DocVersion = DocMetadata.Version.Number.ToString();
-						Results.UnversionedPackages.Add(MetaSound->GetPackage());
-						UE_LOG(LogMetaSound, Display, TEXT("MetaSound asset '%s' document not versioned (already version %s)."), *MetaSoundName, *DocVersion);
-					}
+					Results.UnversionedPackages.Add(PackagePath);
+					return true;
 				}
 			}
+
+			UObject* MetaSoundObject = AssetData.GetAsset();
+			FMetasoundAssetBase* MetaSound = IMetasoundUObjectRegistry::Get().GetObjectAsAssetBase(MetaSoundObject);
+			if (MetaSound)
+			{
+				FMetaSoundFrontendDocumentBuilder& DocBuilder = FDocumentBuilderRegistry::GetChecked().FindOrBeginBuilding(MetaSoundObject);
+				const bool bVersioned = MetaSound->GetVersionedOnLoad() || MetaSound->VersionAsset(DocBuilder);
+				const FMetasoundFrontendDocumentMetadata& DocMetadata = DocBuilder.GetConstDocumentChecked().Metadata;
+				const FString DocVersion = DocMetadata.Version.Number.ToString();
+				constexpr bool bAlwaysMarkDirty = true;
+				MetaSoundObject->Modify(bAlwaysMarkDirty);
+				MetaSound->SetVersionedOnLoad();
+
+				Results.VersionedPackages.Add(AssetData.GetPackage());
+
+				if (bVersioned)
+				{
+					UE_LOG(LogMetaSound, Display, TEXT("MetaSound '%s' document versioned to %s."), *AssetData.GetFullName(), *DocVersion);
+				}
+
+				if (!ClassInfo.bIsValid)
+				{
+					UE_LOG(LogMetaSound, Display, TEXT("MetaSound '%s' asset tags updated."), *AssetData.GetFullName());
+				}
+
+				return true;
+			}
+
+			Results.UnversionedPackages.Add(PackagePath);
+			UE_LOG(LogMetaSound, Error, TEXT("MetaSound asset '%s' failed to load: asset document not versioned/tags updated."), *AssetData.GetFullName());
 			return true;
 		});
 
diff --git a/Engine/Plugins/Runtime/Metasound/Source/MetasoundEngine/Public/MetasoundAssetSubsystem.h b/Engine/Plugins/Runtime/Metasound/Source/MetasoundEngine/Public/MetasoundAssetSubsystem.h
--- a/Engine/Plugins/Runtime/Metasound/Source/MetasoundEngine/Public/MetasoundAssetSubsystem.h
+++ b/Engine/Plugins/Runtime/Metasound/Source/MetasoundEngine/Public/MetasoundAssetSubsystem.h
@@ -45,7 +45,7 @@ PRAGMA_ENABLE_DEPRECATION_WARNINGS
 		using FNodeRegistryKey = Frontend::FNodeRegistryKey;
 
 #if WITH_EDITORONLY_DATA
-		using FVersionDocumentResults = Frontend::IMetaSoundAssetManager::FVersionDocumentResults;
+		using FVersionAssetResults = Frontend::IMetaSoundAssetManager::FVersionAssetResults;
 #endif // WITH_EDITORONLY_DATA
 
 		FMetaSoundAssetManager() = default;
@@ -143,7 +143,7 @@ PRAGMA_ENABLE_DEPRECATION_WARNINGS
 		virtual bool TryGetAssetIDFromClassName(const FMetasoundFrontendClassName& InClassName, FGuid& OutGuid) const override;
 		virtual bool TryLoadReferencedAssets(const FMetasoundAssetBase& InAssetBase, TArray<FMetasoundAssetBase*>& OutReferencedAssets) const override;
 #if WITH_EDITORONLY_DATA
-		virtual FVersionDocumentResults VersionDocumentsInFolders(const TArray<FString>& FolderPaths) const override;
+		virtual FVersionAssetResults VersionAssetsInFolders(const TArray<FString>& FolderPaths, bool bRecursePaths = true) const override;
 #endif // WITH_EDITORONLY_DATA
 
 		virtual void WaitUntilAsyncLoadReferencedAssetsComplete(FMetasoundAssetBase& InAssetBase) override;
diff --git a/Engine/Plugins/Runtime/Metasound/Source/MetasoundFrontend/Public/MetasoundAssetManager.h b/Engine/Plugins/Runtime/Metasound/Source/MetasoundFrontend/Public/MetasoundAssetManager.h
--- a/Engine/Plugins/Runtime/Metasound/Source/MetasoundFrontend/Public/MetasoundAssetManager.h
+++ b/Engine/Plugins/Runtime/Metasound/Source/MetasoundFrontend/Public/MetasoundAssetManager.h
@@ -123,16 +123,16 @@ namespace Metasound::Frontend
 	/** At runtime, contains a minimal set of information needed to further query additional
 	  * class data from the Node Class Registries. At edit-time, contains this plus additional
 	  * data useful for informing user of applicable asset classes within given edit contexts
-	  * (see 'FMetaSoundEditorGraphSchema' for usage).
+	  * (see 'FMetaSoundEditorGraphSchema').
 	  */
 	struct METASOUNDFRONTEND_API FMetaSoundAssetClassInfo : public FMetaSoundClassInfo
 	{
 		FMetaSoundAssetClassInfo() = default;
 		FMetaSoundAssetClassInfo(const IMetaSoundDocumentInterface& InDocInterface);
 
-		// Attempts to return tag data from the given AssetData, which may or may not be loaded.
-		// Invalid tag constructed if object isn't loaded and any fields fail to load/are not serialized
-		// (i.e. tags are out-of-date).
+		// Attempts to transform AssetTag data from the given AssetData to this class info (asset may or
+		// may not be loaded). bIsValid set to false if object isn't loaded and any fields fail to load/are
+		// not serialized (i.e. tags are out-of-date).
 		FMetaSoundAssetClassInfo(const FAssetData& InAssetData);
 
 		// Exports tag data to the given RegistryContext
@@ -213,10 +213,10 @@ namespace Metasound::Frontend
 		};
 
 #if WITH_EDITORONLY_DATA
-		struct FVersionDocumentResults
+		struct FVersionAssetResults
 		{
 			TArray<UPackage*> VersionedPackages;
-			TArray<UPackage*> UnversionedPackages;
+			TArray<FTopLevelAssetPath> UnversionedPackages;
 
 			// Returns whether or not documents were found and loaded/attempted to version
 			bool DocumentsFoundInPackages() const
@@ -340,10 +340,10 @@ PRAGMA_ENABLE_DEPRECATION_WARNINGS
 		virtual void RenameAsset(const FAssetData& InAssetData, const FString& InOldObjectPath) = 0;
 
 #if WITH_EDITORONLY_DATA
-		// Versions all MetaSound documents found recursively within the given folder paths. Populates the provided array with MetaSound
-		// packages whose document versions were out-of-date and consequently updated. Returns true if MetaSounds were found in the given
-		// directories, false if not.
-		virtual FVersionDocumentResults VersionDocumentsInFolders(const TArray<FString>& FolderPaths) const = 0;
+		// Versions all MetaSound asset tags & documents found within the given folder paths. Optionally, recurses the given paths.
+		// Populates the provided array with MetaSound packages whose document versions were out-of-date and consequently updated.
+		// Returns resulting versioned and unversioned paths.
+		virtual FVersionAssetResults VersionAssetsInFolders(const TArray<FString>& FolderPaths, bool bRecursePaths = true) const = 0;
 #endif // WITH_EDITORONLY_DATA
 
 		// Waits until all async load requests related to this asset are complete.
