Using metasound thread local debug info to capture asset scopes during build & execution of graphs/nodes
# preflight 67eac3d039205539ce8e25e4
# swarm https://p4-swarm.epicgames.net/reviews/41072068

diff --git a/Engine/Plugins/Runtime/Metasound/Source/MetasoundGenerator/Private/MetasoundGenerator.cpp b/Engine/Plugins/Runtime/Metasound/Source/MetasoundGenerator/Private/MetasoundGenerator.cpp
--- a/Engine/Plugins/Runtime/Metasound/Source/MetasoundGenerator/Private/MetasoundGenerator.cpp
+++ b/Engine/Plugins/Runtime/Metasound/Source/MetasoundGenerator/Private/MetasoundGenerator.cpp
@@ -12,6 +12,7 @@
 #include "MetasoundOperatorCache.h"
 #include "MetasoundParameterTransmitter.h"
 #include "MetasoundRouter.h"
+#include "MetasoundThreadLocalDebug.h"
 #include "MetasoundTrace.h"
 #include "MetasoundTrigger.h"
 #include "MetasoundVertexData.h"
@@ -658,13 +659,13 @@ PRAGMA_ENABLE_DEPRECATION_WARNINGS
 			// Create a scoped timed section for the bulk of the processing...
 			{
 				MetasoundGeneratorPrivate::FBlockRenderScope RuntimelBlockRenderScope(RenderTimer.Get());
+				
 				LLM_SCOPE_BYTAG(Audio_MetaSound_Generator_Execute);
-				UE_TRACE_METADATA_SCOPE_ASSET_FNAME(FName(MetasoundName), NAME_None, NAME_None);
 				
 				ApplyPendingUpdatesToInputs();
 
 				// Call metasound graph operator.
-				RootExecuter.Execute();
+ 				RootExecuter.Execute();
 				
 				// Check if generated finished during this execute call
 				if (*OnFinishedTriggerRef)
diff --git a/Engine/Plugins/Runtime/Metasound/Source/MetasoundGenerator/Private/MetasoundGeneratorBuilder.cpp b/Engine/Plugins/Runtime/Metasound/Source/MetasoundGenerator/Private/MetasoundGeneratorBuilder.cpp
--- a/Engine/Plugins/Runtime/Metasound/Source/MetasoundGenerator/Private/MetasoundGeneratorBuilder.cpp
+++ b/Engine/Plugins/Runtime/Metasound/Source/MetasoundGenerator/Private/MetasoundGeneratorBuilder.cpp
@@ -10,6 +10,7 @@
 #include "MetasoundOutputNode.h"
 #include "MetasoundParameterTransmitter.h"
 #include "MetasoundRouter.h"
+#include "MetasoundThreadLocalDebug.h"
 #include "MetasoundTrace.h"
 #include "MetasoundTrigger.h"
 #include "MetasoundVertexData.h"
@@ -275,7 +276,8 @@ namespace Metasound
 			LLM_SCOPE_BYTAG(Audio_MetaSound_Generator_Build);
 
 			// asset-level trace asset path
-			UE_TRACE_METADATA_SCOPE_ASSET_FNAME(
+			METASOUND_DEBUG_DECLARE_SCOPE;
+			METASOUND_DEBUG_SET_ACTIVE_ASSET_SCOPE(
 				  FName(InInitParams.MetaSoundName)
 				, InInitParams.ClassName
 				, InInitParams.AssetPath
@@ -297,7 +299,7 @@ namespace Metasound
 			FBuildDynamicGraphOperatorParams BuildParams 
 			{ 
 				{ 
-					*InInitParams.Graph, 
+					*InInitParams.Graph,
 					InOperatorSettings, 
 					OpAndInputs.Inputs, 
 					InInitParams.Environment,
@@ -340,7 +342,7 @@ namespace Metasound
 
 			if (const FAnyDataReference* FinishTriggerRef = VertexData.GetOutputs().FindDataReference(SourceOneShotInterface::Outputs::OnFinished))
 			{
-
+ 
 				FinishTrigger = FinishTriggerRef->GetDataReadReference<FTrigger>();
 			}
 
diff --git a/Engine/Plugins/Runtime/Metasound/Source/MetasoundGraphCore/Private/MetasoundOperatorBuilder.cpp b/Engine/Plugins/Runtime/Metasound/Source/MetasoundGraphCore/Private/MetasoundOperatorBuilder.cpp
--- a/Engine/Plugins/Runtime/Metasound/Source/MetasoundGraphCore/Private/MetasoundOperatorBuilder.cpp
+++ b/Engine/Plugins/Runtime/Metasound/Source/MetasoundGraphCore/Private/MetasoundOperatorBuilder.cpp
@@ -563,7 +563,6 @@ PRAGMA_ENABLE_DEPRECATION_WARNINGS
 		for (const INode* Node : InSortedNodes)
 		{
 			METASOUND_DEBUG_SET_ACTIVE_NODE_SCOPE(Node);
-
 			FOperatorID OperatorID = GetOperatorID(Node);
 			FGraphOperatorData::FOperatorInfo& OperatorInfo = OperatorMap.FindChecked(OperatorID);
 
@@ -583,7 +582,13 @@ PRAGMA_ENABLE_DEPRECATION_WARNINGS
 				{
 					METASOUND_TRACE_CPUPROFILER_EVENT_SCOPE_COND(Metasound::FOperatorBuilder::CreateOperators::CreateOperator, bVerboseTraceEvents);
 
-					OperatorInfo.Operator = MakeUnique<FProfilingOperator>(Factory->CreateOperator(CreateParams, InOutContext.Results), Node);
+					TUniquePtr<FProfilingOperator> ProfilingOperatorPtr = MakeUnique<FProfilingOperator>(Factory->CreateOperator(CreateParams, InOutContext.Results), Node);
+					check(ProfilingOperatorPtr);
+#if UE_METASOUND_DEBUG_ENABLED
+					ProfilingOperatorPtr->SetAssetMetaData(METASOUND_DEBUG_ACTIVE_ASSET_SCOPE);
+#endif // #if UE_METASOUND_DEBUG_ENABLED
+					
+					OperatorInfo.Operator = MoveTemp(ProfilingOperatorPtr);
 				}
 				else
 				{
diff --git a/Engine/Plugins/Runtime/Metasound/Source/MetasoundGraphCore/Private/MetasoundProfilingOperator.cpp b/Engine/Plugins/Runtime/Metasound/Source/MetasoundGraphCore/Private/MetasoundProfilingOperator.cpp
--- a/Engine/Plugins/Runtime/Metasound/Source/MetasoundGraphCore/Private/MetasoundProfilingOperator.cpp
+++ b/Engine/Plugins/Runtime/Metasound/Source/MetasoundGraphCore/Private/MetasoundProfilingOperator.cpp
@@ -5,6 +5,8 @@
 #include "MetasoundTrace.h"
 #include "HAL/IConsoleManager.h"
 
+#include "ProfilingDebugging/AssetMetadataTrace.h"
+
 DEFINE_LOG_CATEGORY_STATIC(LogMetasoundProfiling, Log, All);
 
 namespace Metasound
@@ -83,6 +85,13 @@ namespace Metasound
 		
 	}
 
+#if UE_METASOUND_DEBUG_ENABLED
+	void FProfilingOperator::SetAssetMetaData(const ThreadLocalDebug::FAssetMetaData& InAssetMetaData)
+	{
+		AssetMetaData = InAssetMetaData;
+	}
+#endif // #if UE_METASOUND_DEBUG_ENABLED
+
 	void FProfilingOperator::StaticReset(IOperator* InOperator, const IOperator::FResetParams& InParams)
 	{
 		FProfilingOperator* DerivedOperator = static_cast<FProfilingOperator*>(InOperator);
@@ -94,8 +103,14 @@ namespace Metasound
 	void FProfilingOperator::StaticExecute(IOperator* InOperator)
 	{
 		FProfilingOperator* DerivedOperator = static_cast<FProfilingOperator*>(InOperator);
-		check(nullptr != DerivedOperator && DerivedOperator->ExecuteFunction);
+		check(nullptr != DerivedOperator && DerivedOperator->ExecuteFunction); 
 		METASOUND_TRACE_CPUPROFILER_EVENT_SCOPE_FAST(DerivedOperator->InsightsExecuteEventSpecId, *DerivedOperator->InsightsExecuteEventName);
+
+#if UE_METASOUND_DEBUG_ENABLED
+		const ThreadLocalDebug::FAssetMetaData& AssetMetaData = DerivedOperator->AssetMetaData;
+		UE_TRACE_METADATA_SCOPE_ASSET_FNAME(AssetMetaData.ClassName, AssetMetaData.AssetName, AssetMetaData.AssetPath);
+#endif // #if UE_METASOUND_DEBUG_ENABLED
+		
 		DerivedOperator->ExecuteFunction(DerivedOperator->Operator.Get());
 	}
 
diff --git a/Engine/Plugins/Runtime/Metasound/Source/MetasoundGraphCore/Private/MetasoundProfilingOperator.h b/Engine/Plugins/Runtime/Metasound/Source/MetasoundGraphCore/Private/MetasoundProfilingOperator.h
--- a/Engine/Plugins/Runtime/Metasound/Source/MetasoundGraphCore/Private/MetasoundProfilingOperator.h
+++ b/Engine/Plugins/Runtime/Metasound/Source/MetasoundGraphCore/Private/MetasoundProfilingOperator.h
@@ -4,6 +4,7 @@
 
 #include "MetasoundNodeInterface.h"
 #include "MetasoundOperatorInterface.h"
+#include "MetasoundThreadLocalDebug.h"
 
 namespace Metasound
 {
@@ -19,14 +20,15 @@ namespace Metasound
 	class FProfilingOperator : public IOperator
 	{
 	public:
-		FProfilingOperator(TUniquePtr<IOperator>&& WrappedOperator, const INode* Node)
-			: Operator(MoveTemp(WrappedOperator))
-			, ResetFunction(Operator->GetResetFunction())
-			, ExecuteFunction(Operator->GetExecuteFunction())
-			, PostExecuteFunction(Operator->GetPostExecuteFunction())
-			, InsightsResetEventSpecId(0)
-			, InsightsExecuteEventSpecId(0)
-			, InsightsPostExecuteEventSpecId(0)
+		FProfilingOperator(
+			TUniquePtr<IOperator>&& WrappedOperator, const INode* Node)
+				: Operator(MoveTemp(WrappedOperator))
+				, ResetFunction(Operator->GetResetFunction())
+				, ExecuteFunction(Operator->GetExecuteFunction())
+				, PostExecuteFunction(Operator->GetPostExecuteFunction())
+				, InsightsResetEventSpecId(0)
+				, InsightsExecuteEventSpecId(0)
+				, InsightsPostExecuteEventSpecId(0)
 		{
 			check(Operator);
 			const FNodeClassMetadata& NodeMetadata = Node->GetMetadata();
@@ -45,6 +47,10 @@ namespace Metasound
 			InsightsPostExecuteEventName.Append(TEXT("_POSTEXECUTE"));
 		}
 
+#if UE_METASOUND_DEBUG_ENABLED
+		void SetAssetMetaData(const ThreadLocalDebug::FAssetMetaData& InAssetMetaData);
+#endif // #if UE_METASOUND_DEBUG_ENABLED
+
 		virtual ~FProfilingOperator() = default;
 
 		virtual void BindInputs(FInputVertexInterfaceData& InVertexData) override
@@ -87,15 +93,19 @@ namespace Metasound
 		static void StaticPostExecute(IOperator* InOperator);
 
 	private:
-		TUniquePtr<IOperator> Operator;
-		FResetFunction        ResetFunction;
-		FExecuteFunction      ExecuteFunction;
-		FPostExecuteFunction  PostExecuteFunction;
-		FString               InsightsResetEventName;
-		FString               InsightsExecuteEventName;
-		FString               InsightsPostExecuteEventName;
-		uint32                InsightsResetEventSpecId;
-		uint32                InsightsExecuteEventSpecId;
-		uint32                InsightsPostExecuteEventSpecId;
+#if UE_METASOUND_DEBUG_ENABLED
+		ThreadLocalDebug::FAssetMetaData AssetMetaData;
+#endif // #if UE_METASOUND_DEBUG_ENABLED
+		
+		TUniquePtr<IOperator>			 Operator;
+		FResetFunction        			 ResetFunction;
+		FExecuteFunction      			 ExecuteFunction;
+		FPostExecuteFunction  			 PostExecuteFunction;
+		FString               			 InsightsResetEventName;
+		FString               			 InsightsExecuteEventName;
+		FString               			 InsightsPostExecuteEventName;
+		uint32                			 InsightsResetEventSpecId;
+		uint32                			 InsightsExecuteEventSpecId;
+		uint32                			 InsightsPostExecuteEventSpecId;
 	};
 }
\ No newline at end of file
diff --git a/Engine/Plugins/Runtime/Metasound/Source/MetasoundGraphCore/Private/MetasoundThreadLocalDebug.cpp b/Engine/Plugins/Runtime/Metasound/Source/MetasoundGraphCore/Private/MetasoundThreadLocalDebug.cpp
--- a/Engine/Plugins/Runtime/Metasound/Source/MetasoundGraphCore/Private/MetasoundThreadLocalDebug.cpp
+++ b/Engine/Plugins/Runtime/Metasound/Source/MetasoundGraphCore/Private/MetasoundThreadLocalDebug.cpp
@@ -17,16 +17,24 @@ namespace Metasound
 			// Set the active node for the debug info.
 			void SetActiveNode(const INode* InNode);
 
+			// set the current asset scope
+			void SetCurrentAsset(FName Class, FName Name, FName Path);
+			const FAssetMetaData& GetCurrentAsset() const;
+
 			// Return the active node.
 			const INode* GetActiveNode() const;
 
 			// Returns the class name and version string for the active node in the current thread.
 			const TCHAR* GetActiveNodeClassNameAndVersion() const;
 
+			FAssetMetaData GetFullNodeAssetMetaData() const;
+
 		private:
 
 			const INode* ActiveNode = nullptr;
 			FString NodeClassNameAndVersion;
+
+			FAssetMetaData AssetMetaData;
 		};
 
 		namespace ThreadLocalDebugPrivate
@@ -40,7 +48,6 @@ namespace Metasound
 			if (ActiveNode)
 			{
 				const FNodeClassMetadata& Metadata = ActiveNode->GetMetadata();
-
 				NodeClassNameAndVersion = FString::Format(TEXT("{0} v{1}.{2}"), {Metadata.ClassName.ToString(), Metadata.MajorVersion, Metadata.MinorVersion});
 			}
 			else
@@ -49,6 +56,18 @@ namespace Metasound
 			}
 		}
 
+		void FDebugInfo::SetCurrentAsset(FName Class, FName Name, FName Path)
+		{
+			AssetMetaData.ClassName = Class;
+			AssetMetaData.AssetName = Name;
+			AssetMetaData.AssetPath = Path;
+		}
+
+		const FAssetMetaData& FDebugInfo::GetCurrentAsset() const
+		{
+			return AssetMetaData;
+		}
+
 		const INode* FDebugInfo::GetActiveNode() const
 		{
 			return ActiveNode;
@@ -63,11 +82,34 @@ namespace Metasound
 			return TEXT("");
 		}
 
+		FAssetMetaData FDebugInfo::GetFullNodeAssetMetaData() const
+		{
+			if (nullptr == ActiveNode) 
+			{
+				return {
+					.ClassName = AssetMetaData.ClassName,
+					.AssetName = AssetMetaData.AssetName,
+					.AssetPath = AssetMetaData.AssetName
+				};
+			}
+				
+			return {
+				.ClassName = AssetMetaData.ClassName,
+				.AssetName = AssetMetaData.AssetName,
+				.AssetPath = FName(AssetMetaData.AssetName.ToString() + '/' + ActiveNode->GetMetadata().ClassName.ToString())
+			};
+		}
+
 		const TCHAR* GetActiveNodeClassNameAndVersionOnThisThread()
 		{
 			return ThreadLocalDebugPrivate::DebugInfoOnThisThread.GetActiveNodeClassNameAndVersion();
 		}
 
+		void SetAssetScopeForThisThread(FName Class, FName Name, FName Path)
+		{
+			ThreadLocalDebugPrivate::DebugInfoOnThisThread.SetCurrentAsset(Class, Name, Path);
+		}
+
 		FDebugInfo* GetDebugInfoOnThisThread()
 		{
 			return &ThreadLocalDebugPrivate::DebugInfoOnThisThread;
@@ -80,6 +122,7 @@ namespace Metasound
 			{
 				PriorNode = DebugInfo->GetActiveNode();
 				DebugInfo->SetActiveNode(InNode);
+				UpdateAssetMetaData();
 			}
 		}
 
@@ -90,6 +133,14 @@ namespace Metasound
 				DebugInfo->SetActiveNode(PriorNode);
 			}
 		}
+
+		void FScopeDebugActiveNode::UpdateAssetMetaData()
+		{
+			if (DebugInfo)
+			{
+				AssetMetaData = DebugInfo->GetFullNodeAssetMetaData();
+			}
+		}
 	}
 }
 
diff --git a/Engine/Plugins/Runtime/Metasound/Source/MetasoundGraphCore/Private/MetasoundThreadLocalDebug.h b/Engine/Plugins/Runtime/Metasound/Source/MetasoundGraphCore/Private/MetasoundThreadLocalDebug.h
deleted file mode 100644
diff --git a/Engine/Plugins/Runtime/Metasound/Source/MetasoundGraphCore/Public/MetasoundThreadLocalDebug.h b/Engine/Plugins/Runtime/Metasound/Source/MetasoundGraphCore/Public/MetasoundThreadLocalDebug.h
new file mode 100644
--- /dev/null
+++ b/Engine/Plugins/Runtime/Metasound/Source/MetasoundGraphCore/Public/MetasoundThreadLocalDebug.h
@@ -0,0 +1,124 @@
+// Copyright Epic Games, Inc. All Rights Reserved.
+
+#pragma once
+
+#include "HAL/Platform.h"
+#include "Containers/UnrealString.h"
+#include "UObject/NameTypes.h"
+
+/** MetaSound Thread Local Debugging allows node info to be propagated down to 
+ * calling functions for debug purposes. It passes the information by storing thread_local
+ * debug data which can be accessed on other functions. thread_local access can
+ * have negative performance implications, and so this debug info is disabled by
+ * default when DO_CHECK and DO_ENSURE are false. 
+ *
+ * To set which node is active in the current thread, use the following macros.
+ *
+ * 	{
+ * 		// Declare a debug scope to minimize access to thread locals
+ * 		METASOUND_DEBUG_DECLARE_SCOPE
+ *
+ * 		for (const INode* Node : Nodes)
+ * 		{
+ * 			// Set the active node in this scope to be Node
+ * 			METASOUND_DEBUG_SET_ACTIVE_NODE_SCOPE(Node)
+ *
+ *
+ * 			// Any calls in this scope can get debug info about the "active node"
+ * 			// by using the METASOUND_DEBUG macros
+ * 			DoOtherThings().
+ * 		}
+ * 	}
+ *
+ *
+ * 	In another function in the same callstack, debug info can be accessed.
+ *
+ * 	{ 
+ * 		UE_LOG(LogMetaSound, Error, TEXT("Error with node %s"), METASOUND_DEBUG_ACTIVE_NODE_NAME);
+ * 	}
+ */
+
+#ifndef UE_METASOUND_DEBUG_ENABLED
+#define UE_METASOUND_DEBUG_ENABLED (DO_CHECK || DO_ENSURE)
+#endif
+
+#if UE_METASOUND_DEBUG_ENABLED
+
+namespace Metasound
+{
+	class INode;
+
+	// Debug util for setting stack scope debug variables
+	namespace ThreadLocalDebug
+	{
+		class FDebugInfo;
+
+		struct FAssetMetaData
+		{
+			FName ClassName;
+			FName AssetName;
+			FName AssetPath;
+		}; // struct FAssetMetaData
+
+		METASOUNDGRAPHCORE_API FDebugInfo* GetDebugInfoOnThisThread(); 
+		METASOUNDGRAPHCORE_API void SetAssetScopeForThisThread(FName Class, FName Name, FName Path);
+		
+		const TCHAR* GetActiveNodeClassNameAndVersionOnThisThread();
+
+
+		struct FScopeDebugActiveNode final
+		{
+			explicit FScopeDebugActiveNode(FDebugInfo* InDebugInfo, const INode* InNode);
+
+			~FScopeDebugActiveNode();
+
+			FScopeDebugActiveNode(const FScopeDebugActiveNode&) = delete;
+			FScopeDebugActiveNode(FScopeDebugActiveNode&&) = delete;
+			FScopeDebugActiveNode& operator=(const FScopeDebugActiveNode&) = delete;
+			FScopeDebugActiveNode& operator=(FScopeDebugActiveNode&&) = delete;
+
+			const FAssetMetaData& GetAssetMetaData() const { return AssetMetaData; };
+
+		private:
+			void UpdateAssetMetaData();
+			FAssetMetaData AssetMetaData;
+			const INode* PriorNode = nullptr;
+			FDebugInfo* DebugInfo = nullptr;
+		};
+	}
+}
+
+#endif // if UE_METASOUND_DEBUG_ENABLED
+
+#if UE_METASOUND_DEBUG_ENABLED
+
+#define METASOUND_DEBUG_SCOPE_NAME __MetaSoundDebugScope
+#define METASOUND_DEBUG_DECLARE_SCOPE ::Metasound::ThreadLocalDebug::FDebugInfo* METASOUND_DEBUG_SCOPE_NAME = ::Metasound::ThreadLocalDebug::GetDebugInfoOnThisThread();
+
+#define METASOUND_DEBUG_SET_ACTIVE_NODE_SCOPE(InNodePtr)\
+	::Metasound::ThreadLocalDebug::FScopeDebugActiveNode __MetaSoundDebugScope_ActiveNode(METASOUND_DEBUG_SCOPE_NAME, InNodePtr);\
+	UE_TRACE_METADATA_SCOPE_ASSET_FNAME(\
+		__MetaSoundDebugScope_ActiveNode.GetAssetMetaData().ClassName,\
+		__MetaSoundDebugScope_ActiveNode.GetAssetMetaData().AssetName,\
+		__MetaSoundDebugScope_ActiveNode.GetAssetMetaData().AssetPath);\
+
+#define METASOUND_DEBUG_ACTIVE_ASSET_SCOPE (__MetaSoundDebugScope_ActiveNode.GetAssetMetaData())
+
+#define METASOUND_DEBUG_ACTIVE_NODE_NAME (::Metasound::ThreadLocalDebug::GetActiveNodeClassNameAndVersionOnThisThread())
+
+#define METASOUND_DEBUG_SET_ACTIVE_ASSET_SCOPE(AssetName, AssetClass, AssetPath)\
+	UE_TRACE_METADATA_SCOPE_ASSET_FNAME(AssetName, AssetClass, AssetPath);\
+	::Metasound::ThreadLocalDebug::SetAssetScopeForThisThread(AssetName, AssetClass, AssetPath);\
+
+#else // if UE_METASOUND_DEBUG_ENABLED
+
+#define METASOUND_DEBUG_SCOPE_NAME 
+#define METASOUND_DEBUG_DECLARE_SCOPE 
+#define METASOUND_DEBUG_SET_ACTIVE_NODE_SCOPE(InNodePtr)
+#define METASOUND_DEBUG_ACTIVE_ASSET_SCOPE
+#define METASOUND_DEBUG_ACTIVE_NODE_NAME TEXT("[Metasound debug disable. Use UE_METASOUND_DEBUG_ENABLED to enable]")
+#define METASOUND_DEBUG_SET_ACTIVE_ASSET_SCOPE(AssetName, AssetClass, AssetPath)
+
+#endif // if UE_METASOUND_DEBUG_ENABLED
+
+ 
\ No newline at end of file
