- MetaSound Editor synchronization perf fixes Round 1
- Minor builder --> controller clean-up
- Fix input/output copy pasta causing connection results to not show in actions menu
#rnx

diff --git a/Engine/Plugins/Runtime/Metasound/Source/MetasoundEditor/Private/MetasoundEditor.cpp b/Engine/Plugins/Runtime/Metasound/Source/MetasoundEditor/Private/MetasoundEditor.cpp
--- a/Engine/Plugins/Runtime/Metasound/Source/MetasoundEditor/Private/MetasoundEditor.cpp
+++ b/Engine/Plugins/Runtime/Metasound/Source/MetasoundEditor/Private/MetasoundEditor.cpp
@@ -4600,7 +4600,8 @@ PRAGMA_ENABLE_DEPRECATION_WARNINGS
 			FGraphBuilder::BindEditorGraph(DocBuilder, &Graph);
 			check(Graph);
 
-			const bool bSynchronizedGraph = FGraphBuilder::SynchronizeGraph(DocBuilder, *Graph, !bRefreshGraph);
+			// Requires editable builder as transient registry context may need to be recached on the document.
+			const bool bSynchronizedGraph = FGraphBuilder::SynchronizeGraph(Builder->GetBuilder(), *Graph, !bRefreshGraph);
 			bRefreshGraph = false;
 
 			FMetasoundAssetBase* MetaSoundAsset = IMetasoundUObjectRegistry::Get().GetObjectAsAssetBase(&MetaSound);
diff --git a/Engine/Plugins/Runtime/Metasound/Source/MetasoundEditor/Private/MetasoundEditorGraph.cpp b/Engine/Plugins/Runtime/Metasound/Source/MetasoundEditor/Private/MetasoundEditorGraph.cpp
--- a/Engine/Plugins/Runtime/Metasound/Source/MetasoundEditor/Private/MetasoundEditorGraph.cpp
+++ b/Engine/Plugins/Runtime/Metasound/Source/MetasoundEditor/Private/MetasoundEditorGraph.cpp
@@ -477,10 +477,8 @@ bool UMetasoundEditorGraphVertex::NameContainsInterfaceNamespace(FMetasoundFront
 
 bool UMetasoundEditorGraphVertex::CanRename() const
 {
-	Metasound::Frontend::FConstDocumentHandle DocumentHandle = GetOwningGraph()->GetDocumentHandle();
-	const FMetasoundFrontendGraphClass& RootGraphClass = DocumentHandle->GetRootGraphClass();
-
-	return !RootGraphClass.PresetOptions.bIsPreset && !IsInterfaceMember();
+	FMetaSoundFrontendDocumentBuilder& Builder = GetFrontendBuilderChecked();
+	return !Builder.IsPreset() && !IsInterfaceMember();
 }
 
 bool UMetasoundEditorGraphVertex::CanRename(const FText& InNewText, FText& OutError) const
@@ -506,10 +504,8 @@ bool UMetasoundEditorGraphVertex::CanRename(const FText& InNewText, FText& OutEr
 		return false;
 	}
 
-	Metasound::Frontend::FConstDocumentHandle DocumentHandle = GetOwningGraph()->GetDocumentHandle();
-	const FMetasoundFrontendGraphClass& RootGraphClass = DocumentHandle->GetRootGraphClass();
-
-	if (RootGraphClass.PresetOptions.bIsPreset)
+	const FMetaSoundFrontendDocumentBuilder& Builder = GetFrontendBuilderChecked();
+	if (Builder.IsPreset())
 	{
 		OutError = FText::Format(LOCTEXT("GraphVertexRenameInvalid_Preset", "{0} is a vertex in a preset graph and cannot be renamed."), InNewText);
 		return false;
@@ -521,13 +517,11 @@ bool UMetasoundEditorGraphVertex::CanRename(const FText& InNewText, FText& OutEr
 	Audio::FParameterPath::SplitName(NewName, Namespace, ParameterName);
 
 	bool bIsNameValid = true;
-	FConstNodeHandle NodeHandle = GetConstNodeHandle();
-	FConstGraphHandle GraphHandle = NodeHandle->GetOwningGraph();
-	GraphHandle->IterateConstNodes([&](FConstNodeHandle NodeToCompare)
+	Builder.IterateNodesByClassType([&](const FMetasoundFrontendClass&, const FMetasoundFrontendNode& NodeToCompare)
 	{
-		if (NodeID != NodeToCompare->GetID())
+		if (NodeID != NodeToCompare.GetID())
 		{
-			FName OtherName = NodeToCompare->GetNodeName();
+			FName OtherName = NodeToCompare.Name;
 			if (NewName == OtherName)
 			{
 				bIsNameValid = false;
diff --git a/Engine/Plugins/Runtime/Metasound/Source/MetasoundEditor/Private/MetasoundEditorGraphBuilder.cpp b/Engine/Plugins/Runtime/Metasound/Source/MetasoundEditor/Private/MetasoundEditorGraphBuilder.cpp
--- a/Engine/Plugins/Runtime/Metasound/Source/MetasoundEditor/Private/MetasoundEditorGraphBuilder.cpp
+++ b/Engine/Plugins/Runtime/Metasound/Source/MetasoundEditor/Private/MetasoundEditorGraphBuilder.cpp
@@ -999,11 +999,13 @@ PRAGMA_ENABLE_DEPRECATION_WARNINGS
 				}
 				case EMetasoundFrontendClassType::Input:
 				{
-					return &FoundNode->Interface.Inputs.Last();
+					ensureMsgf(InPin->Direction == EGPD_Output, TEXT("Querying for hidden input node output vertex, which should never be represented on an editor graph."));
+					return &FoundNode->Interface.Outputs.Last();
 				}
 				case EMetasoundFrontendClassType::Output:
 				{
-					return &FoundNode->Interface.Outputs.Last();
+					ensureMsgf(InPin->Direction == EGPD_Input, TEXT("Querying for hidden output node input vertex, which should never be represented on an editor graph."));
+					return &FoundNode->Interface.Inputs.Last();
 				}
 
 				default:
diff --git a/Engine/Plugins/Runtime/Metasound/Source/MetasoundEditor/Private/MetasoundEditorGraphNode.cpp b/Engine/Plugins/Runtime/Metasound/Source/MetasoundEditor/Private/MetasoundEditorGraphNode.cpp
--- a/Engine/Plugins/Runtime/Metasound/Source/MetasoundEditor/Private/MetasoundEditorGraphNode.cpp
+++ b/Engine/Plugins/Runtime/Metasound/Source/MetasoundEditor/Private/MetasoundEditorGraphNode.cpp
@@ -114,33 +114,33 @@ bool UMetasoundEditorGraphNode::SyncLocationFromFrontendNode(bool bUpdateEditorN
 {
 	using namespace Metasound::Frontend;
 
-	FConstNodeHandle NodeHandle = GetConstNodeHandle();
-	const FMetasoundFrontendNodeStyle& Style = NodeHandle->GetNodeStyle();
-
-	if (!Style.Display.Locations.IsEmpty())
+	if (const FMetasoundFrontendNode* Node = GetFrontendNode())
 	{
-		if (bUpdateEditorNodeID)
+		if (!Node->Style.Display.Locations.IsEmpty())
 		{
-			if (ensureMsgf(Style.Display.Locations.Num() == 1, TEXT("Style location map should only contain one location (legacy support for multiple location values should be migrated by this point)")))
+			if (bUpdateEditorNodeID)
 			{
-				for (const TPair<FGuid, FVector2D>& IDLocationPair : Style.Display.Locations)
+				if (ensureMsgf(Node->Style.Display.Locations.Num() == 1, TEXT("Style location map should only contain one location (legacy support for multiple location values should be migrated by this point)")))
 				{
-					if (ensureMsgf(IDLocationPair.Key.IsValid(), TEXT("Location cannot be updated for ed node with invalid guid")))
+					for (const TPair<FGuid, FVector2D>& IDLocationPair : Node->Style.Display.Locations)
 					{
-						NodeGuid = IDLocationPair.Key;
-						NodePosX = UE::LWC::FloatToIntCastChecked<int32>(IDLocationPair.Value.X);
-						NodePosY = UE::LWC::FloatToIntCastChecked<int32>(IDLocationPair.Value.Y);
-						return true;
+						if (ensureMsgf(IDLocationPair.Key.IsValid(), TEXT("Location cannot be updated for ed node with invalid guid")))
+						{
+							NodeGuid = IDLocationPair.Key;
+							NodePosX = UE::LWC::FloatToIntCastChecked<int32>(IDLocationPair.Value.X);
+							NodePosY = UE::LWC::FloatToIntCastChecked<int32>(IDLocationPair.Value.Y);
+							return true;
+						}
 					}
 				}
 			}
-		}
-		if (const FVector2D* Location = Style.Display.Locations.Find(NodeGuid))
-		{
-			NodePosX = UE::LWC::FloatToIntCastChecked<int32>(Location->X);
-			NodePosY = UE::LWC::FloatToIntCastChecked<int32>(Location->Y);
+			if (const FVector2D* Location = Node->Style.Display.Locations.Find(NodeGuid))
+			{
+				NodePosX = UE::LWC::FloatToIntCastChecked<int32>(Location->X);
+				NodePosY = UE::LWC::FloatToIntCastChecked<int32>(Location->Y);
 
-			return true;
+				return true;
+			}
 		}
 	}
 
diff --git a/Engine/Plugins/Runtime/Metasound/Source/MetasoundEditor/Private/MetasoundEditorGraphSchema.cpp b/Engine/Plugins/Runtime/Metasound/Source/MetasoundEditor/Private/MetasoundEditorGraphSchema.cpp
--- a/Engine/Plugins/Runtime/Metasound/Source/MetasoundEditor/Private/MetasoundEditorGraphSchema.cpp
+++ b/Engine/Plugins/Runtime/Metasound/Source/MetasoundEditor/Private/MetasoundEditorGraphSchema.cpp
@@ -2135,16 +2135,16 @@ bool UMetasoundEditorGraphSchema::ShouldHidePinDefaultValue(UEdGraphPin* Pin) co
 		return true;
 	}
 
-	if (Pin->Direction == EGPD_Input)
-	{
-		FConstInputHandle InputHandle = FGraphBuilder::GetConstInputHandleFromPin(Pin);
-		return !InputHandle->GetOwningNode()->GetClassStyle().Display.bShowLiterals;
-	}
+	UMetasoundEditorGraphNode* EdNode = CastChecked<UMetasoundEditorGraphNode>(Pin->GetOwningNode());
+	const FMetaSoundFrontendDocumentBuilder& Builder = EdNode->GetBuilderChecked().GetBuilder();
+	FMetasoundFrontendVertexHandle InputVertexHandle = FGraphBuilder::GetPinVertexHandle(Builder, Pin);
 
-	if (Pin->Direction == EGPD_Output)
+	if (const FMetasoundFrontendNode* Node = Builder.FindNode(InputVertexHandle.NodeID))
 	{
-		FConstOutputHandle OutputHandle = FGraphBuilder::GetConstOutputHandleFromPin(Pin);
-		return !OutputHandle->GetOwningNode()->GetClassStyle().Display.bShowLiterals;
+		if (const FMetasoundFrontendClass* Class = Builder.FindDependency(Node->ClassID))
+		{
+			return Class->Style.Display.bShowLiterals;
+		}
 	}
 
 	// TODO: Determine if should be hidden from doc data
diff --git a/Engine/Plugins/Runtime/Metasound/Source/MetasoundFrontend/Private/MetasoundFrontendGraphController.cpp b/Engine/Plugins/Runtime/Metasound/Source/MetasoundFrontend/Private/MetasoundFrontendGraphController.cpp
--- a/Engine/Plugins/Runtime/Metasound/Source/MetasoundFrontend/Private/MetasoundFrontendGraphController.cpp
+++ b/Engine/Plugins/Runtime/Metasound/Source/MetasoundFrontend/Private/MetasoundFrontendGraphController.cpp
@@ -242,45 +242,41 @@ namespace Metasound
 			return Invalid::GetInvalidInterfaceStyle();
 		}
 
-		void FGraphController::SetGraphStyle(const FMetasoundFrontendGraphStyle& InStyle)
+		void FGraphController::SetGraphStyle(FMetasoundFrontendGraphStyle Style)
 		{
 			if (FMetasoundFrontendGraphClass* GraphClass = GraphClassPtr.Get())
 			{
-				FindBuildGraphChecked(*GraphClass).Style = InStyle;
+				FindBuildGraphChecked(*GraphClass).Style = MoveTemp(Style);
 			}
 		}
 
-		void FGraphController::SetInputStyle(const FMetasoundFrontendInterfaceStyle& InStyle)
+		void FGraphController::SetInputStyle(FMetasoundFrontendInterfaceStyle Style)
 		{
 			if (FMetasoundFrontendGraphClass* GraphClass = GraphClassPtr.Get())
 			{
-				FMetasoundFrontendInterfaceStyle Style = InStyle;
-
 				TArray<FMetasoundFrontendClassInput>& Inputs = GraphClass->Interface.Inputs;
 				Style.DefaultSortOrder.SetNumZeroed(Inputs.Num());
 
 				for (int32 i = 0; i < Inputs.Num(); ++i)
 				{
-					Inputs[i].Metadata.SortOrderIndex = InStyle.DefaultSortOrder[i];
+					Inputs[i].Metadata.SortOrderIndex = Style.DefaultSortOrder[i];
 				}
-				GraphClass->Interface.SetInputStyle(Style);
+				GraphClass->Interface.SetInputStyle(MoveTemp(Style));
 			}
 		}
 
-		void FGraphController::SetOutputStyle(const FMetasoundFrontendInterfaceStyle& InStyle)
+		void FGraphController::SetOutputStyle(FMetasoundFrontendInterfaceStyle Style)
 		{
 			if (FMetasoundFrontendGraphClass* GraphClass = GraphClassPtr.Get())
 			{
-				FMetasoundFrontendInterfaceStyle Style = InStyle;
-
 				TArray<FMetasoundFrontendClassOutput>& Outputs = GraphClass->Interface.Outputs;
 				Style.DefaultSortOrder.SetNumZeroed(Outputs.Num());
 
 				for (int32 i = 0; i < Outputs.Num(); ++i)
 				{
-					Outputs[i].Metadata.SortOrderIndex = InStyle.DefaultSortOrder[i];
+					Outputs[i].Metadata.SortOrderIndex = Style.DefaultSortOrder[i];
 				}
-				GraphClass->Interface.SetOutputStyle(Style);
+				GraphClass->Interface.SetOutputStyle(MoveTemp(Style));
 			}
 		}
 #endif // WITH_EDITOR
@@ -1424,7 +1420,7 @@ namespace Metasound
 					}
 					Style.DefaultSortOrder.Add(Output.Metadata.SortOrderIndex);
 				}
-				GraphClass->Interface.SetOutputStyle(Style);
+				GraphClass->Interface.SetOutputStyle(MoveTemp(Style));
 			}
 		}
 #endif // WITH_EDITOR
diff --git a/Engine/Plugins/Runtime/Metasound/Source/MetasoundFrontend/Private/MetasoundFrontendGraphController.h b/Engine/Plugins/Runtime/Metasound/Source/MetasoundFrontend/Private/MetasoundFrontendGraphController.h
--- a/Engine/Plugins/Runtime/Metasound/Source/MetasoundFrontend/Private/MetasoundFrontendGraphController.h
+++ b/Engine/Plugins/Runtime/Metasound/Source/MetasoundFrontend/Private/MetasoundFrontendGraphController.h
@@ -110,9 +110,9 @@ namespace Metasound
 			virtual const FMetasoundFrontendGraphStyle& GetGraphStyle() const override;
 			virtual const FMetasoundFrontendInterfaceStyle& GetInputStyle() const override;
 			virtual const FMetasoundFrontendInterfaceStyle& GetOutputStyle() const override;
-			virtual void SetGraphStyle(const FMetasoundFrontendGraphStyle& InStyle) override;
-			virtual void SetInputStyle(const FMetasoundFrontendInterfaceStyle& InStyle) override;
-			virtual void SetOutputStyle(const FMetasoundFrontendInterfaceStyle& InStyle) override;
+			virtual void SetGraphStyle(FMetasoundFrontendGraphStyle Style) override;
+			virtual void SetInputStyle(FMetasoundFrontendInterfaceStyle Style) override;
+			virtual void SetOutputStyle(FMetasoundFrontendInterfaceStyle Style) override;
 #endif // WITH_EDITOR
 
 			bool ContainsInputVertex(const FVertexName& InName, const FName& InTypeName) const override;
diff --git a/Engine/Plugins/Runtime/Metasound/Source/MetasoundFrontend/Private/MetasoundFrontendInvalidController.h b/Engine/Plugins/Runtime/Metasound/Source/MetasoundFrontend/Private/MetasoundFrontendInvalidController.h
--- a/Engine/Plugins/Runtime/Metasound/Source/MetasoundFrontend/Private/MetasoundFrontendInvalidController.h
+++ b/Engine/Plugins/Runtime/Metasound/Source/MetasoundFrontend/Private/MetasoundFrontendInvalidController.h
@@ -338,11 +338,11 @@ namespace Metasound
 
 #if WITH_EDITOR
 			virtual const FMetasoundFrontendGraphStyle& GetGraphStyle() const override { return Invalid::GetInvalidGraphStyle(); }
-			virtual void SetGraphStyle(const FMetasoundFrontendGraphStyle& InStyle) override { }
+			virtual void SetGraphStyle(FMetasoundFrontendGraphStyle InStyle) override { }
 			virtual const FMetasoundFrontendInterfaceStyle& GetInputStyle() const override { return Invalid::GetInvalidInterfaceStyle(); }
 			virtual const FMetasoundFrontendInterfaceStyle& GetOutputStyle() const override { return Invalid::GetInvalidInterfaceStyle(); }
-			virtual void SetInputStyle(const FMetasoundFrontendInterfaceStyle& InStyle) override { }
-			virtual void SetOutputStyle(const FMetasoundFrontendInterfaceStyle& InStyle) override { }
+			virtual void SetInputStyle(FMetasoundFrontendInterfaceStyle InStyle) override { }
+			virtual void SetOutputStyle(FMetasoundFrontendInterfaceStyle InStyle) override { }
 #endif // WITH_EDITOR
 
 			virtual void ClearGraph() override { };
diff --git a/Engine/Plugins/Runtime/Metasound/Source/MetasoundFrontend/Private/MetasoundFrontendTransform.cpp b/Engine/Plugins/Runtime/Metasound/Source/MetasoundFrontend/Private/MetasoundFrontendTransform.cpp
--- a/Engine/Plugins/Runtime/Metasound/Source/MetasoundFrontend/Private/MetasoundFrontendTransform.cpp
+++ b/Engine/Plugins/Runtime/Metasound/Source/MetasoundFrontend/Private/MetasoundFrontendTransform.cpp
@@ -827,7 +827,7 @@ namespace Metasound
 			// Sort before adding nodes to graph layout & copy to preset (must be done after all
 			// inputs/outputs are added but before setting locations to propagate effectively)
 			FMetasoundFrontendInterfaceStyle Style = ReferencedGraphHandle->GetInputStyle();
-			InPresetGraphHandle->SetInputStyle(Style);
+			InPresetGraphHandle->SetInputStyle(MoveTemp(Style));
 
 			Style.SortDefaults(NodeHandles, DocumentTransform::GetNodeDisplayNameProjection());
 
@@ -872,7 +872,7 @@ namespace Metasound
 			// Sort before adding nodes to graph layout & copy to preset (must be done after all
 			// inputs/outputs are added but before setting locations to propagate effectively)
 			FMetasoundFrontendInterfaceStyle Style = ReferencedGraphHandle->GetOutputStyle();
-			InPresetGraphHandle->SetOutputStyle(Style);
+			InPresetGraphHandle->SetOutputStyle(MoveTemp(Style));
 
 			Style.SortDefaults(NodeHandles, DocumentTransform::GetNodeDisplayNameProjection());
 
diff --git a/Engine/Plugins/Runtime/Metasound/Source/MetasoundFrontend/Public/MetasoundFrontendController.h b/Engine/Plugins/Runtime/Metasound/Source/MetasoundFrontend/Public/MetasoundFrontendController.h
--- a/Engine/Plugins/Runtime/Metasound/Source/MetasoundFrontend/Public/MetasoundFrontendController.h
+++ b/Engine/Plugins/Runtime/Metasound/Source/MetasoundFrontend/Public/MetasoundFrontendController.h
@@ -640,13 +640,13 @@ namespace Metasound
 			virtual const FMetasoundFrontendInterfaceStyle& GetOutputStyle() const = 0;
 
 			// Sets graph style.
-			virtual void SetGraphStyle(const FMetasoundFrontendGraphStyle& InStyle) = 0;
+			virtual void SetGraphStyle(FMetasoundFrontendGraphStyle Style) = 0;
 
 			// Sets the input style for the graph.
-			virtual void SetInputStyle(const FMetasoundFrontendInterfaceStyle& InStyle) = 0;
+			virtual void SetInputStyle(FMetasoundFrontendInterfaceStyle Style) = 0;
 
 			// Sets the output style for the graph.
-			virtual void SetOutputStyle(const FMetasoundFrontendInterfaceStyle& InStyle) = 0;
+			virtual void SetOutputStyle(FMetasoundFrontendInterfaceStyle Style) = 0;
 
 			/** Return the display name of the graph. */
 			virtual FText GetDisplayName() const = 0;
diff --git a/Engine/Plugins/Runtime/Metasound/Source/MetasoundFrontend/Public/MetasoundFrontendDocument.h b/Engine/Plugins/Runtime/Metasound/Source/MetasoundFrontend/Public/MetasoundFrontendDocument.h
--- a/Engine/Plugins/Runtime/Metasound/Source/MetasoundFrontend/Public/MetasoundFrontendDocument.h
+++ b/Engine/Plugins/Runtime/Metasound/Source/MetasoundFrontend/Public/MetasoundFrontendDocument.h
@@ -1064,7 +1064,7 @@ struct METASOUNDFRONTEND_API FMetasoundFrontendClassEnvironmentVariable
 
 // Style info of an interface.
 USTRUCT()
-struct FMetasoundFrontendInterfaceStyle
+struct METASOUNDFRONTEND_API FMetasoundFrontendInterfaceStyle
 {
 	GENERATED_BODY()
 
@@ -1162,10 +1162,18 @@ public:
 		return InputStyle;
 	}
 
-	void SetInputStyle(const FMetasoundFrontendInterfaceStyle& InInputStyle)
+	void SetInputStyle(FMetasoundFrontendInterfaceStyle InInputStyle)
 	{
-		InputStyle = InInputStyle;
-		ChangeID = FGuid::NewGuid();
+		// A bit of a hack to only update the ChangeID if something in the sort
+		// order has changed to avoid invalidating node widgets and editor graph
+		// re-synchronization. This can cause major perf regression on graph edits.
+		// Currently, RequiredMembers do not change as interfaces are registered
+		// once so no need to check them.
+		if (InInputStyle.DefaultSortOrder != InputStyle.DefaultSortOrder)
+		{
+			ChangeID = FGuid::NewGuid();
+		}
+		InputStyle = MoveTemp(InInputStyle);
 	}
 
 	const FMetasoundFrontendInterfaceStyle& GetOutputStyle() const
@@ -1173,10 +1181,18 @@ public:
 		return OutputStyle;
 	}
 
-	void SetOutputStyle(const FMetasoundFrontendInterfaceStyle& InOutputStyle)
+	void SetOutputStyle(FMetasoundFrontendInterfaceStyle InOutputStyle)
 	{
-		OutputStyle = InOutputStyle;
-		ChangeID = FGuid::NewGuid();
+		// A bit of a hack to only update the ChangeID if something in the sort
+		// order has changed to avoid invalidating node widgets and editor graph
+		// re-synchronization. This can cause major perf regression on graph edits.
+		// Currently, RequiredMembers do not change as interfaces are registered
+		// once so no need to check them.
+		if (InOutputStyle.DefaultSortOrder != OutputStyle.DefaultSortOrder)
+		{
+			ChangeID = FGuid::NewGuid();
+		}
+		OutputStyle = MoveTemp(InOutputStyle);
 	}
 
 	void AddRequiredInputToStyle(const FName& InInputName, const FText& InRequiredText)
