- Fix-up bad cherrypick of local CL that caused MetaSound literals to disappear
- Deprecate GetAccessPtr
- Missing minor input/output schema fix-up from last CL
#rnx

diff --git a/Engine/Plugins/Runtime/Metasound/Source/MetasoundEditor/Private/MetasoundEditor.cpp b/Engine/Plugins/Runtime/Metasound/Source/MetasoundEditor/Private/MetasoundEditor.cpp
--- a/Engine/Plugins/Runtime/Metasound/Source/MetasoundEditor/Private/MetasoundEditor.cpp
+++ b/Engine/Plugins/Runtime/Metasound/Source/MetasoundEditor/Private/MetasoundEditor.cpp
@@ -582,10 +582,8 @@ namespace Metasound
 					{
 						if (GetHoveredGraph() == OwningGraph)
 						{
-							FConstDocumentHandle DocumentHandle= OwningGraph->GetDocumentHandle();
-							const FMetasoundFrontendGraphClass& RootGraphClass = DocumentHandle->GetRootGraphClass();
-							const bool bIsPreset = RootGraphClass.PresetOptions.bIsPreset;
-
+							const UMetaSoundBuilderBase& Builder = OwningGraph->GetBuilderChecked();
+							const bool bIsPreset = Builder.IsPreset();
 							if (bIsPreset)
 							{
 								Message = FText::Format(LOCTEXT("DropTargetFailIsPreset", "'{0}': Graph is Preset"), GraphMember->GetDisplayName());
diff --git a/Engine/Plugins/Runtime/Metasound/Source/MetasoundEditor/Private/MetasoundEditorGraphSchema.cpp b/Engine/Plugins/Runtime/Metasound/Source/MetasoundEditor/Private/MetasoundEditorGraphSchema.cpp
--- a/Engine/Plugins/Runtime/Metasound/Source/MetasoundEditor/Private/MetasoundEditorGraphSchema.cpp
+++ b/Engine/Plugins/Runtime/Metasound/Source/MetasoundEditor/Private/MetasoundEditorGraphSchema.cpp
@@ -2143,7 +2143,7 @@ bool UMetasoundEditorGraphSchema::ShouldHidePinDefaultValue(UEdGraphPin* Pin) co
 	{
 		if (const FMetasoundFrontendClass* Class = Builder.FindDependency(Node->ClassID))
 		{
-			return Class->Style.Display.bShowLiterals;
+			return !Class->Style.Display.bShowLiterals;
 		}
 	}
 
@@ -2563,7 +2563,7 @@ void UMetasoundEditorGraphSchema::GetConversionActions(FGraphActionMenuBuilder&
 
 		if (AccessFilters.HasInputFilters())
 		{
-			if (!Result->CanConnectInputOfTypeAndAccess(AccessFilters.InputTypeName, AccessFilters.InputAccessType))
+			if (!Result->CanConnectOutputOfTypeAndAccess(AccessFilters.InputTypeName, AccessFilters.InputAccessType))
 			{
 				return;
 			}
@@ -2571,7 +2571,7 @@ void UMetasoundEditorGraphSchema::GetConversionActions(FGraphActionMenuBuilder&
 
 		if (AccessFilters.HasOutputFilters())
 		{
-			if (!Result->CanConnectOutputOfTypeAndAccess(AccessFilters.OutputTypeName, AccessFilters.OutputAccessType))
+			if (!Result->CanConnectInputOfTypeAndAccess(AccessFilters.OutputTypeName, AccessFilters.OutputAccessType))
 			{
 				return;
 			}
diff --git a/Engine/Plugins/Runtime/Metasound/Source/MetasoundEngine/Public/Metasound.h b/Engine/Plugins/Runtime/Metasound/Source/MetasoundEngine/Public/Metasound.h
--- a/Engine/Plugins/Runtime/Metasound/Source/MetasoundEngine/Public/Metasound.h
+++ b/Engine/Plugins/Runtime/Metasound/Source/MetasoundEngine/Public/Metasound.h
@@ -192,6 +192,9 @@ protected:
 	virtual void SetReferencedAssets(TSet<Metasound::Frontend::IMetaSoundAssetManager::FAssetRef>&& InAssetRefs) override;
 #endif // #if WITH_EDITOR
 
+	UE_DEPRECATED(5.6, "AccessPtrs are actively being deprecated. Writable access outside of the builder API "
+		"is particularly problematic as in so accessing, the builder's caches are reset which can cause major "
+		"editor performance regressions.")
 	Metasound::Frontend::FDocumentAccessPtr GetDocumentAccessPtr() override;
 	Metasound::Frontend::FConstDocumentAccessPtr GetDocumentConstAccessPtr() const override;
 
diff --git a/Engine/Plugins/Runtime/Metasound/Source/MetasoundEngine/Public/MetasoundSource.h b/Engine/Plugins/Runtime/Metasound/Source/MetasoundEngine/Public/MetasoundSource.h
--- a/Engine/Plugins/Runtime/Metasound/Source/MetasoundEngine/Public/MetasoundSource.h
+++ b/Engine/Plugins/Runtime/Metasound/Source/MetasoundEngine/Public/MetasoundSource.h
@@ -316,6 +316,9 @@ public:
 	virtual const UClass& GetBuilderUClass() const final override;
 
 protected:
+	UE_DEPRECATED(5.6, "AccessPtrs are actively being deprecated. Writable access outside of the builder API "
+		"is particularly problematic as in so accessing, the builder's caches are reset which can cause major "
+		"editor performance regressions.")
 	Metasound::Frontend::FDocumentAccessPtr GetDocumentAccessPtr() override;
 	Metasound::Frontend::FConstDocumentAccessPtr GetDocumentConstAccessPtr() const override;
 
diff --git a/Engine/Plugins/Runtime/Metasound/Source/MetasoundFrontend/Private/MetasoundAssetBase.cpp b/Engine/Plugins/Runtime/Metasound/Source/MetasoundFrontend/Private/MetasoundAssetBase.cpp
--- a/Engine/Plugins/Runtime/Metasound/Source/MetasoundFrontend/Private/MetasoundAssetBase.cpp
+++ b/Engine/Plugins/Runtime/Metasound/Source/MetasoundFrontend/Private/MetasoundAssetBase.cpp
@@ -305,6 +305,7 @@ void FMetasoundAssetBase::UpdateAndRegisterForExecution(Metasound::Frontend::FMe
 	}
 
 #if WITH_EDITOR
+	FMetaSoundFrontendDocumentBuilder* DocBuilder = nullptr;
 	if (InRegistrationOptions.bRebuildReferencedAssetClasses)
 	{
 		RebuildReferencedAssetClasses();
@@ -333,8 +334,8 @@ void FMetasoundAssetBase::UpdateAndRegisterForExecution(Metasound::Frontend::FMe
 		// Only attempt asset versioning if owner is asset (dependency versioning on runtime MetaSound instances isn't supported nor necessary).
 		if (Owner->IsAsset())
 		{
-			FMetaSoundFrontendDocumentBuilder& DocBuilder = IDocumentBuilderRegistry::GetChecked().FindOrBeginBuilding(Owner);
-			VersionDependencies(DocBuilder, InRegistrationOptions.bAutoUpdateLogWarningOnDroppedConnection);
+			DocBuilder = &IDocumentBuilderRegistry::GetChecked().FindOrBeginBuilding(Owner);
+			VersionDependencies(*DocBuilder, InRegistrationOptions.bAutoUpdateLogWarningOnDroppedConnection);
 		}
 #else // !WITH_EDITORONLY_DATA
 		constexpr bool bDidUpdate = false;
@@ -521,7 +522,10 @@ bool FMetasoundAssetBase::IsInterfaceDeclared(const FMetasoundFrontendVersion& I
 
 void FMetasoundAssetBase::SetDocument(FMetasoundFrontendDocument InDocument, bool bMarkDirty)
 {
+PRAGMA_DISABLE_DEPRECATION_WARNINGS
 	FMetasoundFrontendDocument* Document = GetDocumentAccessPtr().Get();
+PRAGMA_ENABLE_DEPRECATION_WARNINGS
+
 	*Document = MoveTemp(InDocument);
 	if (bMarkDirty)
 	{
@@ -576,9 +580,15 @@ bool FMetasoundAssetBase::VersionAsset(FMetaSoundFrontendDocumentBuilder& Builde
 #if WITH_EDITOR
 void FMetasoundAssetBase::CacheRegistryMetadata()
 {
+	METASOUND_TRACE_CPUPROFILER_EVENT_SCOPE(MetaSoundAssetBase::CacheRegistryMetadata);
+
 	using namespace Metasound::Frontend;
 
+
+PRAGMA_DISABLE_DEPRECATION_WARNINGS
 	FMetasoundFrontendDocument* Document = GetDocumentAccessPtr().Get();
+PRAGMA_ENABLE_DEPRECATION_WARNINGS
+
 	if (!ensure(Document))
 	{
 		return;
@@ -807,7 +817,9 @@ bool FMetasoundAssetBase::MarkMetasoundDocumentDirty() const
 
 Metasound::Frontend::FDocumentHandle FMetasoundAssetBase::GetDocumentHandle()
 {
+PRAGMA_DISABLE_DEPRECATION_WARNINGS
 	return Metasound::Frontend::IDocumentController::CreateDocumentHandle(GetDocumentAccessPtr());
+PRAGMA_ENABLE_DEPRECATION_WARNINGS
 }
 
 Metasound::Frontend::FConstDocumentHandle FMetasoundAssetBase::GetDocumentHandle() const
@@ -829,7 +841,10 @@ bool FMetasoundAssetBase::ImportFromJSON(const FString& InJSON)
 {
 	METASOUND_TRACE_CPUPROFILER_EVENT_SCOPE(MetaSoundAssetBase::ImportFromJSON);
 
+PRAGMA_DISABLE_DEPRECATION_WARNINGS
 	FMetasoundFrontendDocument* Document = GetDocumentAccessPtr().Get();
+PRAGMA_ENABLE_DEPRECATION_WARNINGS
+
 	if (ensure(nullptr != Document))
 	{
 		bool bSuccess = Metasound::Frontend::ImportJSONToMetasound(InJSON, *Document);
@@ -848,7 +863,10 @@ bool FMetasoundAssetBase::ImportFromJSONAsset(const FString& InAbsolutePath)
 {
 	METASOUND_TRACE_CPUPROFILER_EVENT_SCOPE(MetaSoundAssetBase::ImportFromJSONAsset);
 
+PRAGMA_DISABLE_DEPRECATION_WARNINGS
 	Metasound::Frontend::FDocumentAccessPtr DocumentPtr = GetDocumentAccessPtr();
+PRAGMA_ENABLE_DEPRECATION_WARNINGS
+
 	if (FMetasoundFrontendDocument* Document = DocumentPtr.Get())
 	{
 		bool bSuccess = Metasound::Frontend::ImportJSONAssetToMetasound(InAbsolutePath, *Document);
@@ -873,7 +891,10 @@ const FMetasoundFrontendDocument& FMetasoundAssetBase::GetConstDocumentChecked()
 
 FMetasoundFrontendDocument& FMetasoundAssetBase::GetDocumentChecked()
 {
+PRAGMA_DISABLE_DEPRECATION_WARNINGS
 	FMetasoundFrontendDocument* Document = GetDocumentAccessPtr().Get();
+PRAGMA_ENABLE_DEPRECATION_WARNINGS
+
 	check(nullptr != Document);
 	return *Document;
 }
diff --git a/Engine/Plugins/Runtime/Metasound/Source/MetasoundFrontend/Public/MetasoundAssetBase.h b/Engine/Plugins/Runtime/Metasound/Source/MetasoundFrontend/Public/MetasoundAssetBase.h
--- a/Engine/Plugins/Runtime/Metasound/Source/MetasoundFrontend/Public/MetasoundAssetBase.h
+++ b/Engine/Plugins/Runtime/Metasound/Source/MetasoundFrontend/Public/MetasoundAssetBase.h
@@ -229,7 +229,9 @@ PRAGMA_ENABLE_DEPRECATION_WARNINGS
 	FText GetDisplayName(FString&& InTypeName) const;
 #endif // WITH_EDITORONLY_DATA
 
-	// Returns an access pointer to the document.
+	UE_DEPRECATED(5.6, "AccessPtrs are actively being deprecated. Writable access outside of the builder API "
+		" is particularly problematic as in so accessing, the builder's caches are reset which can cause major "
+		"editor performance regressions.")
 	virtual Metasound::Frontend::FDocumentAccessPtr GetDocumentAccessPtr() = 0;
 
 	// Returns an access pointer to the document.
