Fix Xcode16.3 compilation issues

diff --git a/Engine/Plugins/Runtime/Metasound/Source/MetasoundGraphCore/Public/MetasoundDataFactory.h b/Engine/Plugins/Runtime/Metasound/Source/MetasoundGraphCore/Public/MetasoundDataFactory.h
--- a/Engine/Plugins/Runtime/Metasound/Source/MetasoundGraphCore/Public/MetasoundDataFactory.h
+++ b/Engine/Plugins/Runtime/Metasound/Source/MetasoundGraphCore/Public/MetasoundDataFactory.h
@@ -1079,7 +1079,7 @@ namespace Metasound
 							// to pass ownership of the proxy to the TUniquePtr<>.
 							TUniquePtr<Audio::IProxyData> DeprecatedProxyDataPtr(InLiteral.Value.Get<TSharedPtr<Audio::IProxyData>>().Get());
 
-							auto Result = FFactoryForDeprecatedPtr::template CreateExplicitArgs(InSettings, (const TUniquePtr<Audio::IProxyData>&)DeprecatedProxyDataPtr);
+							auto Result = FFactoryForDeprecatedPtr::CreateExplicitArgs(InSettings, (const TUniquePtr<Audio::IProxyData>&)DeprecatedProxyDataPtr);
 
 							(void)DeprecatedProxyDataPtr.Release();
 							return Result;
@@ -1133,7 +1133,7 @@ namespace Metasound
 								DeprecatedArray.Emplace(ProxyDataPtr.Get());
 							}
 
-							auto Result = FFactoryForDeprecatedPtr::template CreateExplicitArgs(InSettings, (const TArray<TUniquePtr<Audio::IProxyData>>&)DeprecatedArray);
+							auto Result = FFactoryForDeprecatedPtr::CreateExplicitArgs(InSettings, (const TArray<TUniquePtr<Audio::IProxyData>>&)DeprecatedArray);
 
 							for (TUniquePtr<Audio::IProxyData>& DeprecatedPtr : DeprecatedArray)
 							{
diff --git a/Engine/Source/Editor/ContentBrowser/Public/ContentBrowserMenuContexts.h b/Engine/Source/Editor/ContentBrowser/Public/ContentBrowserMenuContexts.h
--- a/Engine/Source/Editor/ContentBrowser/Public/ContentBrowserMenuContexts.h
+++ b/Engine/Source/Editor/ContentBrowser/Public/ContentBrowserMenuContexts.h
@@ -241,7 +241,7 @@ public:
 	template<typename MenuOrSectionType>
 	static int32 GetNumAssetsSelected(const MenuOrSectionType& MenuOrSection)
 	{
-		if (const UContentBrowserAssetContextMenuContext* Context = UContentBrowserAssetContextMenuContext::template FindContextWithAssets(MenuOrSection))
+		if (const UContentBrowserAssetContextMenuContext* Context = UContentBrowserAssetContextMenuContext::FindContextWithAssets(MenuOrSection))
 		{
 			return Context->SelectedAssets.Num();
 		}
@@ -251,7 +251,7 @@ public:
 	template<typename ExpectedType, typename MenuOrSectionType>
 	static ExpectedType* LoadSingleSelectedAsset(const MenuOrSectionType& MenuOrSection)
 	{
-		if (const UContentBrowserAssetContextMenuContext* Context = UContentBrowserAssetContextMenuContext::template FindContextWithAssets(MenuOrSection))
+		if (const UContentBrowserAssetContextMenuContext* Context = UContentBrowserAssetContextMenuContext::FindContextWithAssets(MenuOrSection))
 		{
 			if (ensure(Context->SelectedAssets.Num() == 1))
 			{
diff --git a/Engine/Source/Programs/UnrealBuildTool/ToolChain/AppleToolChain.cs b/Engine/Source/Programs/UnrealBuildTool/ToolChain/AppleToolChain.cs
--- a/Engine/Source/Programs/UnrealBuildTool/ToolChain/AppleToolChain.cs
+++ b/Engine/Source/Programs/UnrealBuildTool/ToolChain/AppleToolChain.cs
@@ -492,6 +492,21 @@ namespace UnrealBuildTool
 			Arguments.Add("-Wno-unknown-warning-option");
 			Arguments.Add("-Wno-range-loop-analysis");
 			Arguments.Add("-Wno-single-bit-bitfield-constant-conversion");
+
+			// Various new warning from Xcode16.3. Disable for now.
+			VersionNumber AppleClangVersion = ((AppleToolChainInfo)Info).AppleClangVersion;
+			if (AppleClangVersion >= new VersionNumber(17))
+			{
+				// "implicit conversion loses integer precision:" in many files
+				Arguments.Add("-Wno-shorten-64-to-32");
+				
+				// clang17 is detecting some uses of cxx-extensions,
+				// such as: "error: variable length arrays in C++ are a Clang extension" 
+				Arguments.Add("-Wno-vla-extension");
+				
+				// From Engine/Plugins/Runtime/Metasound/Source/MetasoundFrontend/Public/MetasoundPrimitives/Time/AudioBuffers.h
+				Arguments.Add("-Wno-extra-qualification");
+			}
 		}
 
 		/// <inheritdoc/>
