Fix for crash when renaming inputs

diff --git a/Engine/Plugins/Runtime/Metasound/Source/MetasoundEditor/Private/MetasoundEditorGraph.cpp b/Engine/Plugins/Runtime/Metasound/Source/MetasoundEditor/Private/MetasoundEditorGraph.cpp
--- a/Engine/Plugins/Runtime/Metasound/Source/MetasoundEditor/Private/MetasoundEditorGraph.cpp
+++ b/Engine/Plugins/Runtime/Metasound/Source/MetasoundEditor/Private/MetasoundEditorGraph.cpp
@@ -298,6 +298,12 @@ FName UMetasoundEditorGraphVertex::GetMemberName() const
 }
 
 void UMetasoundEditorGraphVertex::SetMemberName(const FName& InNewName, bool bPostTransaction)
+{
+	constexpr bool bPropagateToPinNames = true;
+	SetMemberNameInternal(InNewName, bPropagateToPinNames, bPostTransaction);
+}
+
+void UMetasoundEditorGraphVertex::SetMemberNameInternal(const FName& InNewName, bool bPropagateToPinNames, bool bPostTransaction)
 {
 	using namespace Metasound;
 	using namespace Metasound::Editor;
@@ -329,16 +335,19 @@ void UMetasoundEditorGraphVertex::SetMemberName(const FName& InNewName, bool bPo
 	RenameFrontendMemberInternal(DocBuilder, OldName, InNewName);
 	Breadcrumb.MemberName = InNewName;
 
-	const TArray<UMetasoundEditorGraphMemberNode*> Nodes = GetNodes();
-	for (UMetasoundEditorGraphMemberNode* Node : Nodes)
+	if (bPropagateToPinNames)
 	{
-		const TArray<UEdGraphPin*>& Pins = Node->GetAllPins();
-		ensure(Pins.Num() == 1);
-
-		for (UEdGraphPin* Pin : Pins)
+		const TArray<UMetasoundEditorGraphMemberNode*> Nodes = GetNodes();
+		for (UMetasoundEditorGraphMemberNode* Node : Nodes)
 		{
-			Pin->Modify();
-			Pin->PinName = InNewName;
+			const TArray<UEdGraphPin*>& Pins = Node->GetAllPins();
+			ensure(Pins.Num() == 1);
+
+			for (UEdGraphPin* Pin : Pins)
+			{
+				Pin->Modify();
+				Pin->PinName = InNewName;
+			}
 		}
 	}
 
@@ -1052,7 +1061,9 @@ void UMetasoundEditorGraphInput::SetMemberName(const FName& InNewName, bool bPos
 		}
 	}
 
-	Super::SetMemberName(InNewName, bPostTransaction);
+	// Input nodes are templates, which don't have specialized pin names and get their display names via custom function.
+	constexpr bool bPropagateToPinNames = false;
+	SetMemberNameInternal(InNewName, bPropagateToPinNames, bPostTransaction);
 }
 
 bool UMetasoundEditorGraphInput::RenameFrontendMemberInternal(FMetaSoundFrontendDocumentBuilder& Builder, FName OldName, FName InNewName) const
diff --git a/Engine/Plugins/Runtime/Metasound/Source/MetasoundEditor/Public/MetasoundEditorGraph.h b/Engine/Plugins/Runtime/Metasound/Source/MetasoundEditor/Public/MetasoundEditorGraph.h
--- a/Engine/Plugins/Runtime/Metasound/Source/MetasoundEditor/Public/MetasoundEditorGraph.h
+++ b/Engine/Plugins/Runtime/Metasound/Source/MetasoundEditor/Public/MetasoundEditorGraph.h
@@ -364,6 +364,7 @@ public:
 	friend class Metasound::Editor::FEditor;
 
 protected:
+	void SetMemberNameInternal(const FName& InNewName, bool bPropagateToPinNames, bool bPostTransaction);
 
 	UPROPERTY()
 	FMetasoundEditorGraphVertexBreadcrumb Breadcrumb;
