Update change ID when renaming MetaSound interface members to avoid missed auto-update request

#rnx

diff --git a/Engine/Plugins/Runtime/Metasound/Source/MetasoundFrontend/Private/MetasoundFrontendDocumentBuilder.cpp b/Engine/Plugins/Runtime/Metasound/Source/MetasoundFrontend/Private/MetasoundFrontendDocumentBuilder.cpp
--- a/Engine/Plugins/Runtime/Metasound/Source/MetasoundFrontend/Private/MetasoundFrontendDocumentBuilder.cpp
+++ b/Engine/Plugins/Runtime/Metasound/Source/MetasoundFrontend/Private/MetasoundFrontendDocumentBuilder.cpp
@@ -4881,6 +4881,7 @@ bool FMetaSoundFrontendDocumentBuilder::SetGraphInputName(FName InputName, FName
 	DocumentDelegates->InterfaceDelegates.OnInputNameChanged.Broadcast(InputName, NewName);
 
 #if WITH_EDITORONLY_DATA
+	RootGraph.Interface.UpdateChangeID();
 	Document.Metadata.ModifyContext.AddMemberIDModified(GraphInput.NodeID);
 #endif // WITH_EDITORONLY_DATA
 
@@ -5102,6 +5103,7 @@ bool FMetaSoundFrontendDocumentBuilder::SetGraphOutputName(FName OutputName, FNa
 	DocumentDelegates->InterfaceDelegates.OnOutputNameChanged.Broadcast(OutputName, NewName);
 	
 #if WITH_EDITORONLY_DATA
+	GraphClass.Interface.UpdateChangeID();
 	Document.Metadata.ModifyContext.AddMemberIDModified(GraphOutput.NodeID);
 #endif // WITH_EDITORONLY_DATA
 
