- Isolated MetaSound versioning script action clean-up
- Fix version metasound folder action to account for already-opened MetaSounds
- Report if any failed properly
- Fail early if asset registry isn't done scanning
- Remove temp task-based execution code that wasn't necessary (and throwing checks in certain cases as you cannot enumerate in-memory only assets on non GT)

#rnx

diff --git a/Engine/Plugins/Runtime/Metasound/Source/MetasoundEditor/Private/MetasoundEditorModule.cpp b/Engine/Plugins/Runtime/Metasound/Source/MetasoundEditor/Private/MetasoundEditorModule.cpp
--- a/Engine/Plugins/Runtime/Metasound/Source/MetasoundEditor/Private/MetasoundEditorModule.cpp
+++ b/Engine/Plugins/Runtime/Metasound/Source/MetasoundEditor/Private/MetasoundEditorModule.cpp
@@ -137,19 +137,7 @@ namespace Metasound::Editor
 						Notify->SetCompletionState(SNotificationItem::CS_Pending);
 					}
 
-					IMetaSoundAssetManager::FVersionAssetResults Results;
-					{
-						FGCScopeGuard Guard;
-						UE::Tasks::FTask ProcessDocuments = UE::Tasks::Launch(UE_SOURCE_LOCATION, [&Results, FolderPaths, bRecursePaths]()
-						{
-							Results = IMetaSoundAssetManager::GetChecked().VersionAssetsInFolders(FolderPaths, bRecursePaths);
-						});
-
-						// TODO: Report progress and tick instead of immediate wait.
-						// Make this a more generic utility for future folder/object utility calls.
-						ProcessDocuments.Wait();
-					}
-
+					const IMetaSoundAssetManager::FVersionAssetResults Results = IMetaSoundAssetManager::GetChecked().VersionAssetsInFolders(FolderPaths, bRecursePaths);
 					if (!Results.DocumentsFoundInPackages())
 					{
 						if (Notify.IsValid())
@@ -158,7 +146,7 @@ namespace Metasound::Editor
 						}
 						UE_LOG(LogMetasoundEditor, Display, TEXT("No MetaSound(s) Versioned: Folder/child folder(s) contain(s) no MetaSound asset(s)."));
 					}
-					else if (Results.VersionedPackages.IsEmpty())
+					else if (Results.PackagesToReserialize.IsEmpty())
 					{
 						if (Notify.IsValid())
 						{
@@ -166,14 +154,23 @@ namespace Metasound::Editor
 						}
 						UE_LOG(LogMetasoundEditor, Display, TEXT("No MetaSound(s) Versioned: Folder/child folder(s) contain(s) no MetaSound asset(s) requiring versioning."));
 					}
+					else if (!Results.FailedPackages.IsEmpty())
+					{
+						if (Notify.IsValid())
+						{
+							Notify->SetCompletionState(SNotificationItem::CS_Fail);
+						}
+						UE_LOG(LogMetasoundEditor, Display, TEXT("%i MetaSound(s) Reserialized but '%i' failed: See output log for details."), Results.PackagesToReserialize.Num(), Results.FailedPackages.Num());
+						FEditorFileUtils::PromptToCheckoutPackages(false, Results.PackagesToReserialize);
+					}
 					else
 					{
 						if (Notify.IsValid())
 						{
 							Notify->SetCompletionState(SNotificationItem::CS_Success);
 						}
-						UE_LOG(LogMetasoundEditor, Display, TEXT("%i MetaSound(s) Successfully Versioned."), Results.VersionedPackages.Num());
-						FEditorFileUtils::PromptToCheckoutPackages(false, Results.VersionedPackages);
+						UE_LOG(LogMetasoundEditor, Display, TEXT("%i MetaSound(s) Successfully Reserialized."), Results.PackagesToReserialize.Num());
+						FEditorFileUtils::PromptToCheckoutPackages(false, Results.PackagesToReserialize);
 					}
 					
 					if (Notify.IsValid())
diff --git a/Engine/Plugins/Runtime/Metasound/Source/MetasoundEngine/Private/MetasoundAssetSubsystem.cpp b/Engine/Plugins/Runtime/Metasound/Source/MetasoundEngine/Private/MetasoundAssetSubsystem.cpp
--- a/Engine/Plugins/Runtime/Metasound/Source/MetasoundEngine/Private/MetasoundAssetSubsystem.cpp
+++ b/Engine/Plugins/Runtime/Metasound/Source/MetasoundEngine/Private/MetasoundAssetSubsystem.cpp
@@ -143,7 +143,9 @@ namespace Metasound::Engine
 		using namespace Frontend;
 
 		{
-			const FMetasoundFrontendDocument& Document = InAssetBase.GetConstDocumentChecked();
+			TScriptInterface<const IMetaSoundDocumentInterface> DocInterface = InAssetBase.GetOwningAsset();
+			check(DocInterface.GetObject());
+			const FMetasoundFrontendDocument& Document = DocInterface->GetConstDocument();
 			const FMetaSoundAssetKey AssetKey(Document.RootGraph.Metadata);
 			if (!ContainsKey(AssetKey))
 			{
@@ -178,7 +180,9 @@ namespace Metasound::Engine
 		{
 			if (Asset)
 			{
-				const FMetasoundFrontendDocument& RefDocument = Asset->GetConstDocumentChecked();
+				TScriptInterface<const IMetaSoundDocumentInterface> DocInterface = Asset->GetOwningAsset();
+				check(DocInterface.GetObject());
+				const FMetasoundFrontendDocument& RefDocument = DocInterface->GetConstDocument();
 				const FMetaSoundAssetKey ClassKey(RefDocument.RootGraph);
 				if (!ContainsKey(ClassKey))
 				{
@@ -636,7 +640,10 @@ namespace Metasound::Engine
 		using FAssetRef = IMetaSoundAssetManager::FAssetRef;
 
 		TSet<FAssetRef> OutAssetRefs;
-		const FMetasoundFrontendDocument& Document = InAssetBase.GetConstDocumentChecked();
+
+		TScriptInterface<const IMetaSoundDocumentInterface> DocInterface = InAssetBase.GetOwningAsset();
+		check(DocInterface.GetObject());
+		const FMetasoundFrontendDocument& Document = DocInterface->GetConstDocument();
 		for (const FMetasoundFrontendClass& Class : Document.Dependencies)
 		{
 			if (Class.Metadata.GetType() != EMetasoundFrontendClassType::External)
@@ -685,7 +692,10 @@ namespace Metasound::Engine
 	bool FMetaSoundAssetManager::GetReferencedPresetHierarchy(FMetasoundAssetBase& InAsset, TArray<FMetasoundAssetBase*>& OutReferencedAssets) const
 	{
 		OutReferencedAssets.Reset();
-		if (!InAsset.GetConstDocumentChecked().RootGraph.PresetOptions.bIsPreset)
+
+		TScriptInterface<const IMetaSoundDocumentInterface> DocInterface = InAsset.GetOwningAsset();
+		check(DocInterface.GetObject());
+		if (!DocInterface->GetConstDocument().RootGraph.PresetOptions.bIsPreset)
 		{
 			return false;
 		}
@@ -703,7 +713,8 @@ namespace Metasound::Engine
 	{
 		OutReferencedAssets.Add(&InAsset);
 
-		const bool bIsPreset = InAsset.GetConstDocumentChecked().RootGraph.PresetOptions.bIsPreset;
+		TScriptInterface<const IMetaSoundDocumentInterface> DocInterface = InAsset.GetOwningAsset();
+		const bool bIsPreset = DocInterface->GetConstDocument().RootGraph.PresetOptions.bIsPreset;
 
 		// Base case (found the first non preset referenced asset)
 		if (!bIsPreset)
@@ -831,7 +842,8 @@ namespace Metasound::Engine
 				// If loaded, this is the "freshest" reference list, so use that.
 				if (const FMetasoundAssetBase* LoadedAsset = FindAsset(ReferencedKey))
 				{
-					const FMetasoundFrontendDocument& Document = LoadedAsset->GetConstDocumentChecked();
+					TScriptInterface<const IMetaSoundDocumentInterface> LoadedAssetDocInterface = LoadedAsset->GetOwningAsset();
+					const FMetasoundFrontendDocument& Document = LoadedAssetDocInterface->GetConstDocument();
 					for (const FMetasoundFrontendClass& Class : Document.Dependencies)
 					{
 						if (Class.Metadata.GetType() == EMetasoundFrontendClassType::External)
@@ -1260,17 +1272,28 @@ namespace Metasound::Engine
 
 		FVersionAssetResults Results;
 		FAssetRegistryModule& AssetRegistryModule = FModuleManager::LoadModuleChecked<FAssetRegistryModule>("AssetRegistry");
+
+		if (AssetRegistryModule.Get().IsGathering())
+		{
+			UE_LOG(LogMetaSound, Error, TEXT("Asset Registry is still scanning, wait to version MetaSound assets in given folder(s) until initial scan is complete."));
+			return Results;
+		}
+
 		AssetRegistryModule.Get().EnumerateAssets(Filter, [&Results](const FAssetData& AssetData)
 		{
 			using namespace Frontend;
 
 			const FTopLevelAssetPath PackagePath(AssetData.PackageName, AssetData.AssetName);
 			const FMetaSoundAssetClassInfo ClassInfo(AssetData);
-			if (ClassInfo.bIsValid)
+
+			// Loaded assets will have likely versioned already, so data is already likely correct.
+			const bool bWasLoaded = AssetData.IsAssetLoaded();
+			if (!bWasLoaded && ClassInfo.bIsValid)
 			{
 				if (ClassInfo.DocInfo.DocumentVersion >= FMetasoundFrontendDocument::GetMaxVersion())
 				{
-					Results.UnversionedPackages.Add(PackagePath);
+					Results.PackagesUpToDate.Add(PackagePath);
+					UE_LOG(LogMetaSound, Display, TEXT("MetaSound '%s' already versioned & contains valid asset tags. Skipping reserialization."), *AssetData.GetFullName());
 					return true;
 				}
 			}
@@ -1279,30 +1302,43 @@ namespace Metasound::Engine
 			FMetasoundAssetBase* MetaSound = IMetasoundUObjectRegistry::Get().GetObjectAsAssetBase(MetaSoundObject);
 			if (MetaSound)
 			{
-				FMetaSoundFrontendDocumentBuilder& DocBuilder = FDocumentBuilderRegistry::GetChecked().FindOrBeginBuilding(MetaSoundObject);
-				const bool bVersioned = MetaSound->GetVersionedOnLoad() || MetaSound->VersionAsset(DocBuilder);
-				const FMetasoundFrontendDocumentMetadata& DocMetadata = DocBuilder.GetConstDocumentChecked().Metadata;
-				const FString DocVersion = DocMetadata.Version.Number.ToString();
+				bool bVersioned = MetaSound->GetVersionedOnLoad();
+				TScriptInterface<IMetaSoundDocumentInterface> DocInterface(MetaSoundObject);
+				const FMetasoundFrontendDocumentMetadata& DocMetadata = DocInterface->GetConstDocument().Metadata;
+				FString DocVersion;
+				if (bVersioned)
+				{
+					DocVersion = DocMetadata.Version.Number.ToString();
+				}
+				else
+				{
+					FMetaSoundFrontendDocumentBuilder& DocBuilder = FDocumentBuilderRegistry::GetChecked().FindOrBeginBuilding(MetaSoundObject);
+					bVersioned = MetaSound->VersionAsset(DocBuilder);
+					DocVersion = DocMetadata.Version.Number.ToString();
+					if (bVersioned)
+					{
+						MetaSound->SetVersionedOnLoad();
+					}
+				}
+
 				constexpr bool bAlwaysMarkDirty = true;
 				MetaSoundObject->Modify(bAlwaysMarkDirty);
-				MetaSound->SetVersionedOnLoad();
-
-				Results.VersionedPackages.Add(AssetData.GetPackage());
 
 				if (bVersioned)
 				{
-					UE_LOG(LogMetaSound, Display, TEXT("MetaSound '%s' document versioned to %s."), *AssetData.GetFullName(), *DocVersion);
+					UE_LOG(LogMetaSound, Display, TEXT("MetaSound '%s' document versioned to %s"), *AssetData.GetFullName(), *DocVersion);
+					Results.PackagesToReserialize.Add(AssetData.GetPackage());
+					return true;
 				}
-
-				if (!ClassInfo.bIsValid)
+				else
 				{
-					UE_LOG(LogMetaSound, Display, TEXT("MetaSound '%s' asset tags updated."), *AssetData.GetFullName());
+					Results.PackagesToReserialize.Add(AssetData.GetPackage());
+					UE_LOG(LogMetaSound, Display, TEXT("MetaSound '%s' already opened but not versioned. Reserializing to ensure tags updated."), *AssetData.GetFullName());
+					return true;
 				}
-
-				return true;
 			}
 
-			Results.UnversionedPackages.Add(PackagePath);
+			Results.FailedPackages.Add(PackagePath);
 			UE_LOG(LogMetaSound, Error, TEXT("MetaSound asset '%s' failed to load: asset document not versioned/tags updated."), *AssetData.GetFullName());
 			return true;
 		});
diff --git a/Engine/Plugins/Runtime/Metasound/Source/MetasoundFrontend/Public/MetasoundAssetManager.h b/Engine/Plugins/Runtime/Metasound/Source/MetasoundFrontend/Public/MetasoundAssetManager.h
--- a/Engine/Plugins/Runtime/Metasound/Source/MetasoundFrontend/Public/MetasoundAssetManager.h
+++ b/Engine/Plugins/Runtime/Metasound/Source/MetasoundFrontend/Public/MetasoundAssetManager.h
@@ -215,13 +215,14 @@ namespace Metasound::Frontend
 #if WITH_EDITORONLY_DATA
 		struct FVersionAssetResults
 		{
-			TArray<UPackage*> VersionedPackages;
-			TArray<FTopLevelAssetPath> UnversionedPackages;
+			TArray<FTopLevelAssetPath> FailedPackages;
+			TArray<UPackage*> PackagesToReserialize;
+			TArray<FTopLevelAssetPath> PackagesUpToDate;
 
 			// Returns whether or not documents were found and loaded/attempted to version
 			bool DocumentsFoundInPackages() const
 			{
-				return !VersionedPackages.IsEmpty() || !UnversionedPackages.IsEmpty();
+				return !PackagesToReserialize.IsEmpty() || !PackagesUpToDate.IsEmpty() || !FailedPackages.IsEmpty();
 			}
 		};
 
diff --git a/Engine/Plugins/Runtime/Metasound/Source/MetasoundFrontend/Public/MetasoundFrontendQuery.h b/Engine/Plugins/Runtime/Metasound/Source/MetasoundFrontend/Public/MetasoundFrontendQuery.h
--- a/Engine/Plugins/Runtime/Metasound/Source/MetasoundFrontend/Public/MetasoundFrontendQuery.h
+++ b/Engine/Plugins/Runtime/Metasound/Source/MetasoundFrontend/Public/MetasoundFrontendQuery.h
@@ -186,7 +186,12 @@ namespace Metasound
 			FMetaSoundClassInterfaceInfo InterfaceInfo;
 #endif // WITH_EDITORONLY_DATA
 
-			// If true, class query data was successfully loaded from manager, false if not.
+			// If true, class info is valid and accurately reflects
+			// that from what was provided on construction. If false,
+			// info was constructed from AssetData that failed to
+			// provide all expected tags or tag values (indicating tag
+			// data is out-of-date). Default constructed info is considered
+			// invalid as well.
 			uint8 bIsValid : 1;
 		};
 	} // namespace Frontend
