Minor clean-up ahead of deprecating redundant asset base document getter/mark dirty calls
#rnx

diff --git a/Engine/Plugins/Runtime/Metasound/Source/MetasoundEditor/Private/MetasoundEditor.cpp b/Engine/Plugins/Runtime/Metasound/Source/MetasoundEditor/Private/MetasoundEditor.cpp
--- a/Engine/Plugins/Runtime/Metasound/Source/MetasoundEditor/Private/MetasoundEditor.cpp
+++ b/Engine/Plugins/Runtime/Metasound/Source/MetasoundEditor/Private/MetasoundEditor.cpp
@@ -1440,9 +1440,8 @@ namespace Metasound
 					IconName += TEXT(".MetasoundPatch");
 				}
 
-				const FMetasoundAssetBase* MetaSoundAsset = IMetasoundUObjectRegistry::Get().GetObjectAsAssetBase(MetaSound);
-				check(MetaSoundAsset);
-				if (MetaSoundAsset->GetConstDocumentChecked().RootGraph.PresetOptions.bIsPreset)
+				TScriptInterface<const IMetaSoundDocumentInterface> DocInterface = MetaSound;
+				if (DocInterface->GetConstDocument().RootGraph.PresetOptions.bIsPreset)
 				{
 					IconName += TEXT(".Preset");
 				}
@@ -1578,13 +1577,17 @@ namespace Metasound
 			if (MetasoundGraphEditor.IsValid())
 			{
 				UMetasoundEditorGraph& MetaSoundGraph = GetMetaSoundGraphChecked();
-				if (FMetasoundAssetBase* MetaSoundAsset = IMetasoundUObjectRegistry::Get().GetObjectAsAssetBase(GetMetasoundObject()))
+				UObject* MetaSoundObject = GetMetasoundObject();
+				if (FMetasoundAssetBase* MetaSoundAsset = IMetasoundUObjectRegistry::Get().GetObjectAsAssetBase(MetaSoundObject))
 				{
 					if (MetaSoundAsset->GetVersionedOnLoad())
 					{
 						MetaSoundAsset->ClearVersionedOnLoad();
 
-						const FString VersionString = MetaSoundAsset->GetConstDocumentChecked().Metadata.Version.Number.ToString();
+						TScriptInterface<IMetaSoundDocumentInterface> DocInterface = MetaSoundAsset->GetOwningAsset();
+						check(DocInterface.GetObject());
+
+						const FString VersionString = DocInterface->GetConstDocument().Metadata.Version.Number.ToString();
 						FText Msg = FText::Format(LOCTEXT("MetaSoundDocumentVersioned", "Document versioned to '{0}' on load."), FText::FromString(VersionString));
 						FNotificationInfo Info(Msg);
 						Info.bFireAndForget = true;
@@ -1592,8 +1595,7 @@ namespace Metasound
 						Info.ExpireDuration = 5.0f;
 
 						MetasoundGraphEditor->AddNotification(Info, false /* bSuccess */);
-
-						MetaSoundAsset->MarkMetasoundDocumentDirty();
+						MetaSoundObject->MarkPackageDirty();
 					}
 				}
 			}
@@ -1793,7 +1795,8 @@ namespace Metasound
 			{
 				if (ensure(ReferencedAsset))
 				{
-					const FMetasoundFrontendDocument& ReferencedDocument = ReferencedAsset->GetConstDocumentChecked();
+					TScriptInterface<const IMetaSoundDocumentInterface> RefDocInterface = ReferencedAsset->GetOwningAsset();
+					const FMetasoundFrontendDocument& ReferencedDocument = RefDocInterface->GetConstDocument();
 					const FMetasoundFrontendGraphClass& ReferencedGraphClass = ReferencedDocument.RootGraph;
 					return FNodeRegistryKey(ReferencedGraphClass);
 				}
@@ -5000,7 +5003,7 @@ PRAGMA_ENABLE_DEPRECATION_WARNINGS
 				return MakeUnique<FGraphConnectionManager>();
 			}
 
-			const FGuid ResolvedGraphPageID = FDocumentBuilderRegistry::GetChecked().ResolveTargetPageID(Source->GetConstDocumentChecked().RootGraph);
+			const FGuid ResolvedGraphPageID = FDocumentBuilderRegistry::GetChecked().ResolveTargetPageID(Source->GetConstDocument().RootGraph);
 			if (ResolvedGraphPageID != Builder->GetConstBuilder().GetBuildPageID())
 			{
 				return MakeUnique<FGraphConnectionManager>();
@@ -5149,8 +5152,8 @@ PRAGMA_ENABLE_DEPRECATION_WARNINGS
 						
 						// Mark dirty (needed for when this is being called in response to a change from Builder BP API for example)
 						UObject* MetaSound = ParentPtr->GetMetasoundObject();
-						const FMetasoundAssetBase* MetaSoundAsset = IMetasoundUObjectRegistry::Get().GetObjectAsAssetBase(MetaSound);
-						MetaSoundAsset->MarkMetasoundDocumentDirty();
+						check(MetaSound);
+						MetaSound->MarkPackageDirty();
 					}
 				}
 			}
diff --git a/Engine/Plugins/Runtime/Metasound/Source/MetasoundEngine/Private/MetasoundOperatorCacheSubsystem.cpp b/Engine/Plugins/Runtime/Metasound/Source/MetasoundEngine/Private/MetasoundOperatorCacheSubsystem.cpp
--- a/Engine/Plugins/Runtime/Metasound/Source/MetasoundEngine/Private/MetasoundOperatorCacheSubsystem.cpp
+++ b/Engine/Plugins/Runtime/Metasound/Source/MetasoundEngine/Private/MetasoundOperatorCacheSubsystem.cpp
@@ -136,7 +136,7 @@ void UMetaSoundCacheSubsystem::PrecacheMetaSoundInternal(UMetaSoundSource* InMet
 	const UMetaSoundSource&	NoninflatableSource = InMetaSound->FindFirstNoninflatableSource(InitParams->Environment, [](const UMetaSoundSource&){});
 
 	FGuid AssetID;
-	const FMetasoundFrontendClassName& ClassName = InMetaSound->GetConstDocumentChecked().RootGraph.Metadata.GetClassName();
+	const FMetasoundFrontendClassName& ClassName = InMetaSound->GetConstDocument().RootGraph.Metadata.GetClassName();
 	if (!IMetaSoundAssetManager::GetChecked().TryGetAssetIDFromClassName(ClassName, AssetID))
 	{
 		UE_LOG(LogMetaSound, Warning, TEXT("Failed to retrieve MetaSoundClassName when precaching operator for MetaSound '%s'"), *InMetaSound->GetPathName());
@@ -189,7 +189,7 @@ void UMetaSoundCacheSubsystem::RemoveCachedOperatorsForMetaSound(UMetaSoundSourc
 	if (TSharedPtr<FOperatorPool> OperatorPool = Module->GetOperatorPool())
 	{
 		FGuid AssetID;
-		const FMetasoundFrontendClassName& ClassName = InMetaSound->GetConstDocumentChecked().RootGraph.Metadata.GetClassName();
+		const FMetasoundFrontendClassName& ClassName = InMetaSound->GetConstDocument().RootGraph.Metadata.GetClassName();
 		if (IMetaSoundAssetManager::GetChecked().TryGetAssetIDFromClassName(ClassName, AssetID))
 		{
 			OperatorPool->RemoveOperatorsWithAssetClassID(AssetID);
diff --git a/Engine/Plugins/Runtime/Metasound/Source/MetasoundEngine/Private/MetasoundSource.cpp b/Engine/Plugins/Runtime/Metasound/Source/MetasoundEngine/Private/MetasoundSource.cpp
--- a/Engine/Plugins/Runtime/Metasound/Source/MetasoundEngine/Private/MetasoundSource.cpp
+++ b/Engine/Plugins/Runtime/Metasound/Source/MetasoundEngine/Private/MetasoundSource.cpp
@@ -530,7 +530,7 @@ void UMetaSoundSource::PostEditChangeOutputFormat()
 		{
 			EditorGraph->RegisterGraphWithFrontend();
 		}
-		MarkMetasoundDocumentDirty();
+		MarkPackageDirty();
 	}
 }
 
@@ -984,7 +984,7 @@ Metasound::Frontend::FConstDocumentAccessPtr UMetaSoundSource::GetDocumentConstA
 bool UMetaSoundSource::ImplementsParameterInterface(Audio::FParameterInterfacePtr InInterface) const
 {
 	const FMetasoundFrontendVersion Version { InInterface->GetName(), { InInterface->GetVersion().Major, InInterface->GetVersion().Minor } };
-	return GetConstDocumentChecked().Interfaces.Contains(Version);
+	return GetConstDocument().Interfaces.Contains(Version);
 }
 
 ISoundGeneratorPtr UMetaSoundSource::CreateSoundGenerator(const FSoundGeneratorInitParams& InParams, TArray<FAudioParameter>&& InDefaultParameters)
@@ -1298,7 +1298,7 @@ void UMetaSoundSource::InitParametersInternal(const Metasound::TSortedVertexName
 
 bool UMetaSoundSource::IsParameterValid(const FAudioParameter& InParameter) const
 {
-	const TArray<FMetasoundFrontendClassInput>& Inputs = GetConstDocumentChecked().RootGraph.Interface.Inputs;
+	const TArray<FMetasoundFrontendClassInput>& Inputs = GetConstDocument().RootGraph.Interface.Inputs;
 	const FMetasoundFrontendVertex* Vertex = Algo::FindByPredicate(Inputs, [&InParameter] (const FMetasoundFrontendClassInput& Input)
 	{
 		return Input.Name == InParameter.ParamName;
diff --git a/Engine/Plugins/Runtime/Metasound/Source/MetasoundEngine/Public/MetasoundBuilderBase.h b/Engine/Plugins/Runtime/Metasound/Source/MetasoundEngine/Public/MetasoundBuilderBase.h
--- a/Engine/Plugins/Runtime/Metasound/Source/MetasoundEngine/Public/MetasoundBuilderBase.h
+++ b/Engine/Plugins/Runtime/Metasound/Source/MetasoundEngine/Public/MetasoundBuilderBase.h
@@ -574,7 +574,7 @@ protected:
 			{
 				// If MetaSound already exists, preserve the class name to avoid
 				// nametable bloat & preserve potentially existing references.
-				const FMetasoundFrontendDocument& ExistingDoc = CastChecked<const UClassType>(MetaSound)->GetConstDocumentChecked();
+				const FMetasoundFrontendDocument& ExistingDoc = CastChecked<const UClassType>(MetaSound)->GetConstDocument();
 				if (!BuilderOptions.bForceUniqueClassName)
 				{
 					DocClassName = &ExistingDoc.RootGraph.Metadata.GetClassName();
diff --git a/Engine/Plugins/Runtime/Metasound/Source/MetasoundEngineTest/Private/EngineTestMetasoundSource.cpp b/Engine/Plugins/Runtime/Metasound/Source/MetasoundEngineTest/Private/EngineTestMetasoundSource.cpp
--- a/Engine/Plugins/Runtime/Metasound/Source/MetasoundEngineTest/Private/EngineTestMetasoundSource.cpp
+++ b/Engine/Plugins/Runtime/Metasound/Source/MetasoundEngineTest/Private/EngineTestMetasoundSource.cpp
@@ -317,8 +317,8 @@ bool FMetaSoundSourceBuilderAuditionLatentCommand::Update()
 		{
 			if (const USoundBase* InitSound = AudioComponent->GetSound())
 			{
-				const UMetaSoundSource* InitMetaSound = CastChecked<UMetaSoundSource>(InitSound);
-				return InitMetaSound->GetConstDocumentChecked().RootGraph.Metadata.GetClassName().GetFullName();
+				const TScriptInterface<const IMetaSoundDocumentInterface> InitMetaSound = InitSound;
+				return InitMetaSound->GetConstDocument().RootGraph.Metadata.GetClassName().GetFullName();
 			}
 
 			return { };
diff --git a/Engine/Plugins/Runtime/Metasound/Source/MetasoundFrontend/Private/MetasoundAssetBase.cpp b/Engine/Plugins/Runtime/Metasound/Source/MetasoundFrontend/Private/MetasoundAssetBase.cpp
--- a/Engine/Plugins/Runtime/Metasound/Source/MetasoundFrontend/Private/MetasoundAssetBase.cpp
+++ b/Engine/Plugins/Runtime/Metasound/Source/MetasoundFrontend/Private/MetasoundAssetBase.cpp
@@ -517,7 +517,9 @@ void FMetasoundAssetBase::UnregisterGraphWithFrontend()
 
 bool FMetasoundAssetBase::IsInterfaceDeclared(const FMetasoundFrontendVersion& InVersion) const
 {
-	return GetConstDocumentChecked().Interfaces.Contains(InVersion);
+	TScriptInterface<const IMetaSoundDocumentInterface> DocInterface = GetOwningAsset();
+	check(DocInterface.GetObject());
+	return DocInterface->GetConstDocument().Interfaces.Contains(InVersion);
 }
 
 void FMetasoundAssetBase::SetDocument(FMetasoundFrontendDocument InDocument, bool bMarkDirty)
@@ -529,7 +531,9 @@ PRAGMA_ENABLE_DEPRECATION_WARNINGS
 	*Document = MoveTemp(InDocument);
 	if (bMarkDirty)
 	{
-		MarkMetasoundDocumentDirty();
+		UObject* OwningAsset = GetOwningAsset();
+		check(OwningAsset);
+		OwningAsset->MarkPackageDirty();
 	}
 }
 
@@ -703,17 +707,23 @@ FMetasoundFrontendDocumentModifyContext& FMetasoundAssetBase::GetModifyContext()
 {
 	// ModifyContext is now mutable to avoid mutations to it requiring access through
 	// the deprecated Document controller causing the builder cache to get wiped unnecessarily.
-	return GetConstDocumentChecked().Metadata.ModifyContext;
+	TScriptInterface<const IMetaSoundDocumentInterface> DocInterface = GetOwningAsset();
+	check(DocInterface.GetObject());
+	return DocInterface->GetConstDocument().Metadata.ModifyContext;
 }
 
 const FMetasoundFrontendDocumentModifyContext& FMetasoundAssetBase::GetConstModifyContext() const
 {
-	return GetConstDocumentChecked().Metadata.ModifyContext;
+	TScriptInterface<const IMetaSoundDocumentInterface> DocInterface = GetOwningAsset();
+	check(DocInterface.GetObject());
+	return DocInterface->GetConstDocument().Metadata.ModifyContext;
 }
 
 const FMetasoundFrontendDocumentModifyContext& FMetasoundAssetBase::GetModifyContext() const
 {
-	return GetConstDocumentChecked().Metadata.ModifyContext;
+	TScriptInterface<const IMetaSoundDocumentInterface> DocInterface = GetOwningAsset();
+	check(DocInterface.GetObject());
+	return DocInterface->GetConstDocument().Metadata.ModifyContext;
 }
 #endif // WITH_EDITOR
 
@@ -851,7 +861,9 @@ PRAGMA_ENABLE_DEPRECATION_WARNINGS
 
 		if (bSuccess)
 		{
-			ensure(MarkMetasoundDocumentDirty());
+			UObject* OwningAsset = GetOwningAsset();
+			check(OwningAsset);
+			ensure(OwningAsset->MarkPackageDirty());
 		}
 
 		return bSuccess;
@@ -873,7 +885,9 @@ PRAGMA_ENABLE_DEPRECATION_WARNINGS
 
 		if (bSuccess)
 		{
-			ensure(MarkMetasoundDocumentDirty());
+			UObject* OwningAsset = GetOwningAsset();
+			check(OwningAsset);
+			ensure(OwningAsset->MarkPackageDirty());
 		}
 
 		return bSuccess;
@@ -901,7 +915,9 @@ PRAGMA_ENABLE_DEPRECATION_WARNINGS
 
 const FMetasoundFrontendDocument& FMetasoundAssetBase::GetDocumentChecked() const
 {
-	return GetConstDocumentChecked();
+	TScriptInterface<const IMetaSoundDocumentInterface> DocInterface = GetOwningAsset();
+	check(DocInterface.GetObject());
+	return DocInterface->GetConstDocument();
 }
 
 const Metasound::Frontend::FGraphRegistryKey& FMetasoundAssetBase::GetGraphRegistryKey() const
@@ -909,11 +925,6 @@ const Metasound::Frontend::FGraphRegistryKey& FMetasoundAssetBase::GetGraphRegis
 	return GraphRegistryKey;
 }
 
-const Metasound::Frontend::FNodeClassRegistryKey& FMetasoundAssetBase::GetRegistryKey() const
-{
-	return GraphRegistryKey.NodeKey;
-}
-
 FString FMetasoundAssetBase::GetOwningAssetName() const
 {
 	if (const UObject* OwningAsset = GetOwningAsset())
diff --git a/Engine/Plugins/Runtime/Metasound/Source/MetasoundFrontend/Public/MetasoundAssetBase.h b/Engine/Plugins/Runtime/Metasound/Source/MetasoundFrontend/Public/MetasoundAssetBase.h
--- a/Engine/Plugins/Runtime/Metasound/Source/MetasoundFrontend/Public/MetasoundAssetBase.h
+++ b/Engine/Plugins/Runtime/Metasound/Source/MetasoundFrontend/Public/MetasoundAssetBase.h
@@ -44,10 +44,12 @@ namespace Metasound
 	} // namespace Frontend
 } // namespace Metasound
 
-/** FMetasoundAssetBase is intended to be a mix-in subclass for UObjects which utilize
- * Metasound assets.  It provides consistent access to FMetasoundFrontendDocuments, control
- * over the FMetasoundFrontendClassInterface of the FMetasoundFrontendDocument.  It also enables the UObject
- * to be utilized by a host of other engine tools built to support MetaSounds.
+/** FMetasoundAssetBase is a mix-in subclass for UObjects which utilize MetaSound objects. As MetaSounds
+ * can now be generated dynamically via the Builder API, the name does not accurately reflect this classes
+ * current implementation.  While it currently has some support for asset, editor graph & document accessors (actively
+ * being deprecated), its primary use is to support all access to runtime-specific, MetaSound features and associated
+ * data, such as proxy generation and runtime node class registration.  For forward support of the MetaSound document
+ * model and supported accessors, see 'IMetaSoundDocumentInterface'.
  */
 class METASOUNDFRONTEND_API FMetasoundAssetBase : public IAudioProxyDataFactory
 {
@@ -154,14 +156,11 @@ public:
 	// of moving all mutable document calls to the Frontend/Subsystem Document Builder API.
 	FMetasoundFrontendDocument& GetDocumentChecked();
 
-	UE_DEPRECATED(5.5, "Use GetConstDocumentChecked instead.")
+	UE_DEPRECATED(5.5, "Use GetConstDocument from casting Owning Asset to IMetaSoundDocumentInterface (See 'GetOwningAsset') instead.")
 	const FMetasoundFrontendDocument& GetDocumentChecked() const;
 
 	const Metasound::Frontend::FGraphRegistryKey& GetGraphRegistryKey() const;
 
-	UE_DEPRECATED(5.4, "Use GetGraphRegistryKey instead.")
-	const Metasound::Frontend::FNodeRegistryKey& GetRegistryKey() const;
-
 #if WITH_EDITORONLY_DATA
 	bool VersionAsset(FMetaSoundFrontendDocumentBuilder& Builder);
 #endif // WITH_EDITORONLY_DATA
