[MetaSound Output Watch] - Add Unwatch function to MetaSound Output Watch Subsystem

[RN] Added Unwatch functions to MetaSound Output Subsystem to Unbind Output Delegates

diff --git a/Engine/Plugins/Runtime/Metasound/Source/MetasoundEngine/Private/MetasoundOutputSubsystem.cpp b/Engine/Plugins/Runtime/Metasound/Source/MetasoundEngine/Private/MetasoundOutputSubsystem.cpp
--- a/Engine/Plugins/Runtime/Metasound/Source/MetasoundEngine/Private/MetasoundOutputSubsystem.cpp
+++ b/Engine/Plugins/Runtime/Metasound/Source/MetasoundEngine/Private/MetasoundOutputSubsystem.cpp
@@ -48,6 +48,44 @@ bool UMetaSoundOutputSubsystem::WatchOutput(
 	return Handle->WatchOutput(OutputName, OnOutputValueChanged, AnalyzerName, AnalyzerOutputName);
 }
 
+bool UMetaSoundOutputSubsystem::UnwatchOutput(
+	UAudioComponent* AudioComponent,
+	const FName OutputName,
+	const FOnMetasoundOutputValueChanged& OnOutputValueChanged,
+	const FName AnalyzerName,
+	const FName AnalyzerOutputName)
+{
+	METASOUND_TRACE_CPUPROFILER_EVENT_SCOPE(UMetasoundOutputSubsystem::UnwatchOutput_Dynamic);
+
+	const TSharedPtr<Metasound::FMetasoundGeneratorHandle> Handle = GetOrCreateGeneratorHandle(AudioComponent);
+
+	if (!HandleIsValid(Handle))
+	{
+		return false;
+	}
+
+	return Handle->UnwatchOutput(OutputName, OnOutputValueChanged, AnalyzerName, AnalyzerOutputName);
+}
+
+bool UMetaSoundOutputSubsystem::UnwatchOutput(
+	UAudioComponent* AudioComponent,
+	const FName OutputName,
+	const FOnMetasoundOutputValueChangedNative& OnOutputValueChanged,
+	const FName AnalyzerName,
+	const FName AnalyzerOutputName)
+{
+	METASOUND_TRACE_CPUPROFILER_EVENT_SCOPE(UMetasoundOutputSubsystem::UnwatchOutput_Native);
+
+	const TSharedPtr<Metasound::FMetasoundGeneratorHandle> Handle = GetOrCreateGeneratorHandle(AudioComponent);
+
+	if (!HandleIsValid(Handle))
+	{
+		return false;
+	}
+
+	return Handle->UnwatchOutput(OutputName, OnOutputValueChanged, AnalyzerName, AnalyzerOutputName);
+}
+
 TSharedPtr<Metasound::FMetasoundGeneratorHandle> UMetaSoundOutputSubsystem::GetOrCreateGeneratorHandle(UAudioComponent* AudioComponent)
 {
 	METASOUND_TRACE_CPUPROFILER_EVENT_SCOPE(UMetasoundOutputSubsystem::GetOrCreateGeneratorHandle);
diff --git a/Engine/Plugins/Runtime/Metasound/Source/MetasoundEngine/Public/MetasoundOutputSubsystem.h b/Engine/Plugins/Runtime/Metasound/Source/MetasoundEngine/Public/MetasoundOutputSubsystem.h
--- a/Engine/Plugins/Runtime/Metasound/Source/MetasoundEngine/Public/MetasoundOutputSubsystem.h
+++ b/Engine/Plugins/Runtime/Metasound/Source/MetasoundEngine/Public/MetasoundOutputSubsystem.h
@@ -45,6 +45,21 @@ public:
 		FName AnalyzerName = NAME_None,
 		FName AnalyzerOutputName = NAME_None);
 
+	UFUNCTION(BlueprintCallable, Category = "MetaSoundOutput", meta = (AdvancedDisplay = "3"))
+	bool UnwatchOutput(
+		UAudioComponent* AudioComponent,
+		FName OutputName,
+		const FOnMetasoundOutputValueChanged& OnOutputValueChanged,
+		FName AnalyzerName = NAME_None,
+		FName AnalyzerOutputName = NAME_None);
+
+	bool UnwatchOutput(
+		UAudioComponent* AudioComponent,
+		FName OutputName,
+		const FOnMetasoundOutputValueChangedNative& OnOutputValueChanged,
+		FName AnalyzerName = NAME_None,
+		FName AnalyzerOutputName = NAME_None);
+
 private:
 	TSharedPtr<Metasound::FMetasoundGeneratorHandle> GetOrCreateGeneratorHandle(UAudioComponent* AudioComponent);
 	void CleanUpInvalidGeneratorHandles();
