MetaSound Mid-Side Encode Fix

[RN] Fixed bug on MetaSound Mid-Side Encode/Decode nodes preventing live update of input values

diff --git a/Engine/Plugins/Runtime/Metasound/Source/MetasoundStandardNodes/Private/MetasoundMidSideNodes.cpp b/Engine/Plugins/Runtime/Metasound/Source/MetasoundStandardNodes/Private/MetasoundMidSideNodes.cpp
--- a/Engine/Plugins/Runtime/Metasound/Source/MetasoundStandardNodes/Private/MetasoundMidSideNodes.cpp
+++ b/Engine/Plugins/Runtime/Metasound/Source/MetasoundStandardNodes/Private/MetasoundMidSideNodes.cpp
@@ -50,6 +50,7 @@ namespace Metasound
 			, PrevMidScale(0.0f)
 			, PrevSideScale(0.0f)
 			, PrevSpreadAmount(*SpreadAmount)
+			, bPrevEqualPower(*bEqualPower)
 		{
 			Reset(InParams);
 		}
@@ -140,11 +141,13 @@ namespace Metasound
 		{
 			/* Update internal variables, if necessary */
 			const float ClampedSpread = FMath::Clamp(*SpreadAmount, 0.0f, 1.0f);
-			bool bNeedsUpdate = FMath::IsNearlyEqual(ClampedSpread, PrevSpreadAmount);
+			const float bNeedsUpdate = !FMath::IsNearlyEqual(ClampedSpread, PrevSpreadAmount) || bPrevEqualPower != *bEqualPower;
+
 			if (bNeedsUpdate)
 			{
 				UpdateScale(ClampedSpread);
 				PrevSpreadAmount = ClampedSpread;
+				bPrevEqualPower = *bEqualPower;
 			}
 
 			/* Generate Mid-Side Output */
@@ -232,6 +235,9 @@ namespace Metasound
 
 		// Previous Spread variable; need to update MidScale and SideScale if this changes between Execute() calls
 		float PrevSpreadAmount;
+
+		// Previous EqualPower bool value; need to update MidScale and SideScale if this changes between Execute() calls
+		bool bPrevEqualPower;
 	};
 
 	// Node Class
@@ -275,6 +281,7 @@ namespace Metasound
 			, PrevMidScale(0.0f)
 			, PrevSideScale(0.0f)
 			, PrevSpreadAmount(*SpreadAmount)
+			, bPrevEqualPower(*bEqualPower)
 		{
 			MidInBuffer.AddUninitialized(InParams.OperatorSettings.GetNumFramesPerBlock());
 			FMemory::Memset(MidInBuffer.GetData(), 0, sizeof(float) * MidInBuffer.Num());
@@ -371,11 +378,13 @@ namespace Metasound
 		{
 			/* Update internal variables, if necessary */
 			const float ClampedSpread = FMath::Clamp(*SpreadAmount, 0.0f, 1.0f);
-			
-			if (!FMath::IsNearlyEqual(ClampedSpread, PrevSpreadAmount))
+			const float bNeedsUpdate = !FMath::IsNearlyEqual(ClampedSpread, PrevSpreadAmount) || bPrevEqualPower != *bEqualPower;
+
+			if (bNeedsUpdate)
 			{
 				UpdateScale(ClampedSpread);
 				PrevSpreadAmount = ClampedSpread;
+				bPrevEqualPower = *bEqualPower;
 			}
 
 			/* Generate Mid-Side Output */
@@ -462,6 +471,9 @@ namespace Metasound
 		// Previous Spread variable; need to update MidScale and SideScale if this changes between Execute() calls
 		float PrevSpreadAmount;
 
+		// Previous EqualPower bool value; need to update MidScale and SideScale if this changes between Execute() calls
+		bool bPrevEqualPower;
+
 		// Reusable buffers to save performance during Execute().
 		Audio::FAlignedFloatBuffer MidInBuffer;
 		Audio::FAlignedFloatBuffer SideInBuffer;
