- Fix for toggling AccessType causing open references to display error due to not auto-updating
- Fix for toggling AccessType input not calling conform nodes if moving back to value access type
- Safeguard against auto-update potentially being not getting called if data type is swapped
- Minor controller --> builder API clean-up

diff --git a/Engine/Plugins/Runtime/Metasound/Source/MetasoundEditor/Private/MetasoundEditorGraph.cpp b/Engine/Plugins/Runtime/Metasound/Source/MetasoundEditor/Private/MetasoundEditorGraph.cpp
--- a/Engine/Plugins/Runtime/Metasound/Source/MetasoundEditor/Private/MetasoundEditorGraph.cpp
+++ b/Engine/Plugins/Runtime/Metasound/Source/MetasoundEditor/Private/MetasoundEditorGraph.cpp
@@ -1032,18 +1032,23 @@ void UMetasoundEditorGraphInput::SetVertexAccessType(EMetasoundFrontendVertexAcc
 
 	if (InNewAccessType != GetVertexAccessType())
 	{
-		if (UMetasoundEditorGraph* Graph = GetOwningGraph())
+		if (UMetasoundEditorGraph* Graph = GetOwningGraph(); ensure(Graph))
 		{
 			const FScopedTransaction Transaction(FText::Format(GraphPrivate::SetMemberAccessTypeTransactionLabelFormat, GetGraphMemberLabel(), GetDisplayName()), bPostTransaction);
 			Graph->GetMetasoundChecked().Modify();
 			Graph->Modify();
 			Modify();
 
-			EMetaSoundBuilderResult Result = EMetaSoundBuilderResult::Failed;
-			Graph->GetBuilderChecked().SetGraphInputAccessType(GetMemberName(), InNewAccessType, Result);
-			ensure(Result == EMetaSoundBuilderResult::Succeeded);
+			FMetaSoundFrontendDocumentBuilder& DocBuilder = GetFrontendBuilderChecked();
+			DocBuilder.SetGraphInputAccessType(GetMemberName(), InNewAccessType);
 
-			ClassName = GetConstNodeHandle()->GetClassMetadata().GetClassName();
+			if (const FMetasoundFrontendNode* Node = GetFrontendNode(); ensure(Node))
+			{
+				if (const FMetasoundFrontendClass* Class = DocBuilder.FindDependency(Node->ClassID); ensure(Class))
+				{
+					ClassName = Class->Metadata.GetClassName();
+				}
+			}
 			Graph->RegisterGraphWithFrontend();
 		}
 	}
@@ -1405,19 +1410,23 @@ void UMetasoundEditorGraphOutput::SetVertexAccessType(EMetasoundFrontendVertexAc
 
 	if (InNewAccessType != GetVertexAccessType())
 	{
-		UMetasoundEditorGraph* Graph = GetOwningGraph();
-		if (ensure(Graph))
+		if (UMetasoundEditorGraph* Graph = GetOwningGraph(); ensure(Graph))
 		{
 			const FScopedTransaction Transaction(FText::Format(GraphPrivate::SetMemberAccessTypeTransactionLabelFormat, GetGraphMemberLabel(), GetDisplayName()), bPostTransaction);
 			Graph->GetMetasoundChecked().Modify();
 			Graph->Modify();
 			Modify();
 
-			EMetaSoundBuilderResult Result = EMetaSoundBuilderResult::Failed;
-			Graph->GetBuilderChecked().SetGraphOutputAccessType(GetMemberName(), InNewAccessType, Result);
-			ensure(Result == EMetaSoundBuilderResult::Succeeded);
+			FMetaSoundFrontendDocumentBuilder& DocBuilder = GetFrontendBuilderChecked();
+			DocBuilder.SetGraphOutputAccessType(GetMemberName(), InNewAccessType);
 
-			ClassName = GetConstNodeHandle()->GetClassMetadata().GetClassName();
+			if (const FMetasoundFrontendNode* Node = GetFrontendNode(); ensure(Node))
+			{
+				if (const FMetasoundFrontendClass* Class = DocBuilder.FindDependency(Node->ClassID); ensure(Class))
+				{
+					ClassName = Class->Metadata.GetClassName();
+				}
+			}
 			Graph->RegisterGraphWithFrontend();
 		}
 	}
diff --git a/Engine/Plugins/Runtime/Metasound/Source/MetasoundFrontend/Private/MetasoundFrontendDocumentBuilder.cpp b/Engine/Plugins/Runtime/Metasound/Source/MetasoundFrontend/Private/MetasoundFrontendDocumentBuilder.cpp
--- a/Engine/Plugins/Runtime/Metasound/Source/MetasoundFrontend/Private/MetasoundFrontendDocumentBuilder.cpp
+++ b/Engine/Plugins/Runtime/Metasound/Source/MetasoundFrontend/Private/MetasoundFrontendDocumentBuilder.cpp
@@ -4578,23 +4578,35 @@ bool FMetaSoundFrontendDocumentBuilder::SetGraphInputAccessType(FName InputName,
 	if (GraphInput.AccessType != AccessType)
 	{
 		GraphInput.AccessType = AccessType;
-		if (AccessType == EMetasoundFrontendVertexAccessType::Reference)
+
+		RootGraph.IterateGraphPages([this, &Document, &GraphInput, &AccessType](FMetasoundFrontendGraph& Graph)
 		{
-			RootGraph.IterateGraphPages([this, &GraphInput, &AccessType](FMetasoundFrontendGraph& Graph)
+			const IDocumentGraphNodeCache& NodeCache = DocumentCache->GetNodeCache(Graph.PageID);
+			if (const int32* NodeIndex = NodeCache.FindNodeIndex(GraphInput.NodeID))
 			{
-				const IDocumentGraphNodeCache& NodeCache = DocumentCache->GetNodeCache(Graph.PageID);
-				if (const int32* NodeIndex = NodeCache.FindNodeIndex(GraphInput.NodeID))
+				FMetasoundFrontendNode& Node = Graph.Nodes[*NodeIndex];
+				const FMetasoundFrontendVertex& NodeOutput = Node.Interface.Outputs.Last();
+				IterateNodesConnectedWithVertex({ GraphInput.NodeID, NodeOutput.VertexID }, [this, &Document, &Graph, &AccessType](const FMetasoundFrontendEdge& Edge, FMetasoundFrontendNode& ConnectedNode)
 				{
-					FMetasoundFrontendNode& Node = Graph.Nodes[*NodeIndex];
-					const FMetasoundFrontendVertex& NodeOutput = Node.Interface.Outputs.Last();
-					IterateNodesConnectedWithVertex({ GraphInput.NodeID, NodeOutput.VertexID }, [this, &Graph, &AccessType](const FMetasoundFrontendEdge& Edge, FMetasoundFrontendNode& ConnectedNode)
+					if (const FMetasoundFrontendClass* ConnectedNodeClass = FindDependency(ConnectedNode.ClassID))
 					{
-						if (const FMetasoundFrontendClass* ConnectedNodeClass = FindDependency(ConnectedNode.ClassID))
+						// If connected to an input template node, disconnect the template node from other nodes as the data type is
+						// about to be mismatched.  Otherwise, direct connection to other nodes (i.e. at runtime when template
+						// nodes aren't injected) forcefully remove to avoid data type mismatch.
+						if (ConnectedNodeClass->Metadata.GetClassName() == FInputNodeTemplate::ClassName)
 						{
-							// If connected to an input template node, disconnect the template node from other nodes as the data type is
-							// about to be mismatched.  Otherwise, direct connection to other nodes (i.e. at runtime when template
-							// nodes aren't injected) forcefully remove to avoid data type mismatch.
-							if (ConnectedNodeClass->Metadata.GetClassName() == FInputNodeTemplate::ClassName)
+#if WITH_EDITORONLY_DATA
+							// Even if not disconnecting, need to bump the template node so that listeners
+							// such as the editor know to re-evaluate internal node access type resolution
+							// (in the case of the base MetaSound editor, this in turn redraws the pins
+							// accordingly as it caches the AccessType for better performance)
+							if (Graph.PageID == GetBuildPageID())
+							{
+								Document.Metadata.ModifyContext.AddNodeIDModified(ConnectedNode.GetID());
+							}
+#endif // WITH_EDITORONLY_DATA
+
+							if (AccessType == EMetasoundFrontendVertexAccessType::Reference)
 							{
 								const FMetasoundFrontendVertex& ConnectedNodeOutput = ConnectedNode.Interface.Outputs.Last();
 								IterateNodesConnectedWithVertex({ Edge.ToNodeID, ConnectedNodeOutput.VertexID }, [this, &Graph, &AccessType](const FMetasoundFrontendEdge& TempEdge, FMetasoundFrontendNode&)
@@ -4606,29 +4618,30 @@ bool FMetaSoundFrontendDocumentBuilder::SetGraphInputAccessType(FName InputName,
 									}
 								}, Graph.PageID);
 							}
-							else
+						}
+						else if (AccessType == EMetasoundFrontendVertexAccessType::Reference)
+						{
+							const EMetasoundFrontendVertexAccessType ConnectedAccessType = GetNodeInputAccessType(Edge.ToNodeID, Edge.ToVertexID, &Graph.PageID);
+							if (!FMetasoundFrontendClassVertex::CanConnectVertexAccessTypes(AccessType, ConnectedAccessType))
 							{
-								const EMetasoundFrontendVertexAccessType ConnectedAccessType = GetNodeInputAccessType(Edge.ToNodeID, Edge.ToVertexID, &Graph.PageID);
-								if (!FMetasoundFrontendClassVertex::CanConnectVertexAccessTypes(AccessType, ConnectedAccessType))
-								{
-									RemoveEdgeToNodeInput(Edge.ToNodeID, Edge.ToVertexID, &Graph.PageID);
-								}
+								RemoveEdgeToNodeInput(Edge.ToNodeID, Edge.ToVertexID, &Graph.PageID);
 							}
 						}
-					}, Graph.PageID);
-				}
-			});
-
-			const bool bNodeConformed = ConformGraphInputNodeToClass(GraphInput);
-			if (!bNodeConformed)
-			{
-				return false;
+					}
+				}, Graph.PageID);
 			}
+		});
+
+		const bool bNodeConformed = ConformGraphInputNodeToClass(GraphInput);
+		if (!bNodeConformed)
+		{
+			return false;
+		}
 
 #if WITH_EDITORONLY_DATA
-			Document.Metadata.ModifyContext.AddMemberIDModified(GraphInput.NodeID);
+		RootGraph.Interface.UpdateChangeID();
+		Document.Metadata.ModifyContext.AddMemberIDModified(GraphInput.NodeID);
 #endif // WITH_EDITORONLY_DATA
-		}
 	}
 
 	return true;
@@ -4698,6 +4711,7 @@ bool FMetaSoundFrontendDocumentBuilder::SetGraphInputDataType(FName InputName, F
 
 #if WITH_EDITORONLY_DATA
 			ClearMemberMetadata(GraphInput.NodeID);
+			RootGraph.Interface.UpdateChangeID();
 			Document.Metadata.ModifyContext.AddMemberIDModified(GraphInput.NodeID);
 			Document.Metadata.ModifyContext.AddNodeIDModified(GraphInput.NodeID);
 #endif // WITH_EDITORONLY_DATA
@@ -4962,6 +4976,7 @@ bool FMetaSoundFrontendDocumentBuilder::SetGraphOutputAccessType(FName OutputNam
 		}
 
 #if WITH_EDITORONLY_DATA
+		RootGraph.Interface.UpdateChangeID();
 		Document.Metadata.ModifyContext.AddMemberIDModified(GraphOutput.NodeID);
 #endif // WITH_EDITORONLY_DATA
 	}
@@ -5018,6 +5033,7 @@ bool FMetaSoundFrontendDocumentBuilder::SetGraphOutputDataType(FName OutputName,
 #endif // WITH_EDITOR
 
 #if WITH_EDITORONLY_DATA
+		RootGraph.Interface.UpdateChangeID();
 		ClearMemberMetadata(GraphOutput.NodeID);
 		Document.Metadata.ModifyContext.AddMemberIDModified(GraphOutput.NodeID);
 #endif // WITH_EDITORONLY_DATA
