Set MetasoundEngineTest module to UncookedOnly
#12261
#rnx

diff --git a/Engine/Plugins/Runtime/Metasound/Metasound.uplugin b/Engine/Plugins/Runtime/Metasound/Metasound.uplugin
--- a/Engine/Plugins/Runtime/Metasound/Metasound.uplugin
+++ b/Engine/Plugins/Runtime/Metasound/Metasound.uplugin
@@ -42,7 +42,7 @@
 		},
 		{
 			"Name": "MetasoundEngineTest",
-			"Type": "Runtime",
+			"Type": "UncookedOnly",
 			"LoadingPhase": "PreDefault"
 		},
 		{
