Reduce small, spammy re-allocs when building metasound graph.

#swarm https://p4-swarm.epicgames.net/reviews/41373570
#rnx

diff --git a/Engine/Plugins/Runtime/Metasound/Source/MetasoundGraphCore/Private/MetasoundGraphOperator.cpp b/Engine/Plugins/Runtime/Metasound/Source/MetasoundGraphCore/Private/MetasoundGraphOperator.cpp
--- a/Engine/Plugins/Runtime/Metasound/Source/MetasoundGraphCore/Private/MetasoundGraphOperator.cpp
+++ b/Engine/Plugins/Runtime/Metasound/Source/MetasoundGraphCore/Private/MetasoundGraphOperator.cpp
@@ -83,6 +83,11 @@ namespace Metasound
 	{
 		using namespace DirectedGraphAlgo;
 
+		const int32 NumNodes = InOperatorState->NodeOrder.Num();
+		ExecuteStack.Reserve(NumNodes);
+		PostExecuteStack.Reserve(NumNodes);
+		ResetStack.Reserve(NumNodes);
+
 		// Append operators in order.
 		for (const INode* Node: InOperatorState->NodeOrder)
 		{
@@ -99,6 +104,10 @@ namespace Metasound
 			}
 		}
 
+		ExecuteStack.Shrink();
+		PostExecuteStack.Shrink();
+		ResetStack.Shrink();
+
 		// Copy over vertex data
 		SetVertexInterfaceData(MoveTemp(InOperatorState->VertexData));
 	}
@@ -136,7 +145,7 @@ namespace Metasound
 
 	void FGraphOperator::SetVertexInterfaceData(FVertexInterfaceData&& InVertexData)
 	{
-		VertexData = InVertexData;
+		VertexData = MoveTemp(InVertexData);
 	}
 
 	void FGraphOperator::BindInputs(FInputVertexInterfaceData& InInputVertexData)
