Allow frontend metasound vertices to be modified by adding optional class interface.

#rnx

diff --git a/Engine/Plugins/Runtime/Metasound/Source/MetasoundEditor/Private/MetasoundEditor.cpp b/Engine/Plugins/Runtime/Metasound/Source/MetasoundEditor/Private/MetasoundEditor.cpp
--- a/Engine/Plugins/Runtime/Metasound/Source/MetasoundEditor/Private/MetasoundEditor.cpp
+++ b/Engine/Plugins/Runtime/Metasound/Source/MetasoundEditor/Private/MetasoundEditor.cpp
@@ -279,7 +279,7 @@ namespace Metasound
 								return Vertex.NodeID == MemberID;
 							});
 							const FMetasoundFrontendDocument& Document = DocBuilder.GetConstDocumentChecked();
-							const FMetasoundFrontendClassInterface& Interface = Document.RootGraph.Interface;
+							const FMetasoundFrontendClassInterface& Interface = Document.RootGraph.GetDefaultInterface();
 							if (Member->IsA<UMetasoundEditorGraphInput>())
 							{
 								return Interface.Inputs.IndexOfByPredicate(FindVertexWithID);
@@ -4156,7 +4156,7 @@ PRAGMA_ENABLE_DEPRECATION_WARNINGS
 				OutAllActions.AddAction(NewFuncAction);
 			};
 
-			for (const FMetasoundFrontendClassInput& Input : DocBuilder.GetConstDocumentChecked().RootGraph.Interface.Inputs)
+			for (const FMetasoundFrontendClassInput& Input : DocBuilder.GetConstDocumentChecked().RootGraph.GetDefaultInterface().Inputs)
 			{
 				if (const FMetasoundFrontendNode* Node = DocBuilder.FindGraphInputNode(Input.Name))
 				{
@@ -4178,7 +4178,7 @@ PRAGMA_ENABLE_DEPRECATION_WARNINGS
 			}
 
 			const FMetasoundFrontendGraphClass& RootGraph = DocBuilder.GetConstDocumentChecked().RootGraph;
-			for (const FMetasoundFrontendClassOutput& Output : RootGraph.Interface.Outputs)
+			for (const FMetasoundFrontendClassOutput& Output : RootGraph.GetDefaultInterface().Outputs)
 			{
 				if (const FMetasoundFrontendNode* Node = DocBuilder.FindGraphOutputNode(Output.Name))
 				{
@@ -5128,7 +5128,7 @@ PRAGMA_ENABLE_DEPRECATION_WARNINGS
 				const FMetasoundFrontendGraphClass& GraphClass = ConstBuilder.GetConstDocumentChecked().RootGraph;
 				const FMetasoundFrontendGraph& Graph = GraphClass.GetConstDefaultGraph();
 
-				const FMetasoundFrontendClassInput& GraphInput = GraphClass.Interface.Inputs[Index];
+				const FMetasoundFrontendClassInput& GraphInput = GraphClass.GetDefaultInterface().Inputs[Index];
 				const FName InputName = GraphInput.Name;
 				UMetasoundEditorGraph& MetaSoundGraph = ParentPtr->GetMetaSoundGraphChecked();
 				UMetasoundEditorGraphInput* EditorGraphInput = MetaSoundGraph.FindInput(InputName);
diff --git a/Engine/Plugins/Runtime/Metasound/Source/MetasoundEditor/Private/MetasoundEditorBuilderListener.cpp b/Engine/Plugins/Runtime/Metasound/Source/MetasoundEditor/Private/MetasoundEditorBuilderListener.cpp
--- a/Engine/Plugins/Runtime/Metasound/Source/MetasoundEditor/Private/MetasoundEditorBuilderListener.cpp
+++ b/Engine/Plugins/Runtime/Metasound/Source/MetasoundEditor/Private/MetasoundEditorBuilderListener.cpp
@@ -19,7 +19,7 @@ void UMetaSoundEditorBuilderListener::OnGraphInputAdded(int32 Index)
 	{
 		const FMetaSoundFrontendDocumentBuilder& ConstBuilder = Builder->GetConstBuilder();
 		const FMetasoundFrontendGraphClass& GraphClass = ConstBuilder.GetConstDocumentChecked().RootGraph;
-		const FMetasoundFrontendClassInput& GraphInput = GraphClass.Interface.Inputs[Index];
+		const FMetasoundFrontendClassInput& GraphInput = GraphClass.GetDefaultInterface().Inputs[Index];
 
 		OnGraphInputAddedDelegate.Broadcast(GraphInput.Name, GraphInput.TypeName);
 	}
@@ -36,7 +36,7 @@ void UMetaSoundEditorBuilderListener::OnGraphInputDefaultChanged(int32 Index)
 		const FMetasoundFrontendGraph& Graph = GraphClass.GetConstDefaultGraph();
 		FName PageName = Metasound::Frontend::DefaultPageName;
 
-		const FMetasoundFrontendClassInput& GraphInput = GraphClass.Interface.Inputs[Index];
+		const FMetasoundFrontendClassInput& GraphInput = GraphClass.GetDefaultInterface().Inputs[Index];
 		const FMetasoundFrontendLiteral* DefaultLiteral = GraphInput.FindConstDefault(Metasound::Frontend::DefaultPageID);
 		if (ensure(DefaultLiteral))
 		{
@@ -51,7 +51,7 @@ void UMetaSoundEditorBuilderListener::OnRemovingGraphInput(int32 Index)
 	{
 		const FMetaSoundFrontendDocumentBuilder& ConstBuilder = Builder->GetConstBuilder();
 		const FMetasoundFrontendGraphClass& GraphClass = ConstBuilder.GetConstDocumentChecked().RootGraph;
-		const FMetasoundFrontendClassOutput& GraphInput = GraphClass.Interface.Inputs[Index];
+		const FMetasoundFrontendClassOutput& GraphInput = GraphClass.GetDefaultInterface().Inputs[Index];
 
 		OnRemovingGraphInputDelegate.Broadcast(GraphInput.Name, GraphInput.TypeName);
 	}
@@ -63,7 +63,7 @@ void UMetaSoundEditorBuilderListener::OnGraphOutputAdded(int32 Index)
 	{
 		const FMetaSoundFrontendDocumentBuilder& ConstBuilder = Builder->GetConstBuilder();
 		const FMetasoundFrontendGraphClass& GraphClass = ConstBuilder.GetConstDocumentChecked().RootGraph;
-		const FMetasoundFrontendClassOutput& GraphOutput = GraphClass.Interface.Outputs[Index];
+		const FMetasoundFrontendClassOutput& GraphOutput = GraphClass.GetDefaultInterface().Outputs[Index];
 
 		OnGraphOutputAddedDelegate.Broadcast(GraphOutput.Name, GraphOutput.TypeName);
 	}
@@ -75,7 +75,7 @@ void UMetaSoundEditorBuilderListener::OnRemovingGraphOutput(int32 Index)
 	{
 		const FMetaSoundFrontendDocumentBuilder& ConstBuilder = Builder->GetConstBuilder();
 		const FMetasoundFrontendGraphClass& GraphClass = ConstBuilder.GetConstDocumentChecked().RootGraph;
-		const FMetasoundFrontendClassOutput& GraphOutput = GraphClass.Interface.Outputs[Index];
+		const FMetasoundFrontendClassOutput& GraphOutput = GraphClass.GetDefaultInterface().Outputs[Index];
 
 		OnRemovingGraphOutputDelegate.Broadcast(GraphOutput.Name, GraphOutput.TypeName);
 	}
@@ -88,7 +88,7 @@ void UMetaSoundEditorBuilderListener::OnGraphInputDataTypeChanged(int32 Index)
 	{
 		const FMetaSoundFrontendDocumentBuilder& ConstBuilder = Builder->GetConstBuilder();
 		const FMetasoundFrontendGraphClass& GraphClass = ConstBuilder.GetConstDocumentChecked().RootGraph;
-		const FMetasoundFrontendClassOutput& GraphInput = GraphClass.Interface.Inputs[Index];
+		const FMetasoundFrontendClassOutput& GraphInput = GraphClass.GetDefaultInterface().Inputs[Index];
 
 		OnGraphInputDataTypeChangedDelegate.Broadcast(GraphInput.Name, GraphInput.TypeName);
 	}
@@ -100,7 +100,7 @@ void UMetaSoundEditorBuilderListener::OnGraphOutputDataTypeChanged(int32 Index)
 	{
 		const FMetaSoundFrontendDocumentBuilder& ConstBuilder = Builder->GetConstBuilder();
 		const FMetasoundFrontendGraphClass& GraphClass = ConstBuilder.GetConstDocumentChecked().RootGraph;
-		const FMetasoundFrontendClassOutput& GraphOutput = GraphClass.Interface.Outputs[Index];
+		const FMetasoundFrontendClassOutput& GraphOutput = GraphClass.GetDefaultInterface().Outputs[Index];
 
 		OnGraphOutputDataTypeChangedDelegate.Broadcast(GraphOutput.Name, GraphOutput.TypeName);
 	}
diff --git a/Engine/Plugins/Runtime/Metasound/Source/MetasoundEditor/Private/MetasoundEditorGraphBuilder.cpp b/Engine/Plugins/Runtime/Metasound/Source/MetasoundEditor/Private/MetasoundEditorGraphBuilder.cpp
--- a/Engine/Plugins/Runtime/Metasound/Source/MetasoundEditor/Private/MetasoundEditorGraphBuilder.cpp
+++ b/Engine/Plugins/Runtime/Metasound/Source/MetasoundEditor/Private/MetasoundEditorGraphBuilder.cpp
@@ -2595,7 +2595,7 @@ PRAGMA_ENABLE_DEPRECATION_WARNINGS
 			{
 				TSet<FGuid> EdMemberIDs;
 				Algo::TransformIf(OutGraph.Inputs, EdMemberIDs, IsValidMember, GetMemberID);
-				for (const FMetasoundFrontendClassInput& FrontendInput : RootGraph.Interface.Inputs)
+				for (const FMetasoundFrontendClassInput& FrontendInput : RootGraph.GetDefaultInterface().Inputs)
 				{
 					if (EdMemberIDs.Remove(FrontendInput.NodeID) == 0)
 					{
@@ -2626,7 +2626,7 @@ PRAGMA_ENABLE_DEPRECATION_WARNINGS
 			{
 				TSet<FGuid> EdMemberIDs;
 				Algo::TransformIf(OutGraph.Outputs, EdMemberIDs, IsValidMember, GetMemberID);
-				for (const FMetasoundFrontendClassOutput& FrontendOutput : RootGraph.Interface.Outputs)
+				for (const FMetasoundFrontendClassOutput& FrontendOutput : RootGraph.GetDefaultInterface().Outputs)
 				{
 					if (EdMemberIDs.Remove(FrontendOutput.NodeID) == 0)
 					{
diff --git a/Engine/Plugins/Runtime/Metasound/Source/MetasoundEditor/Private/MetasoundEditorGraphSchema.cpp b/Engine/Plugins/Runtime/Metasound/Source/MetasoundEditor/Private/MetasoundEditorGraphSchema.cpp
--- a/Engine/Plugins/Runtime/Metasound/Source/MetasoundEditor/Private/MetasoundEditorGraphSchema.cpp
+++ b/Engine/Plugins/Runtime/Metasound/Source/MetasoundEditor/Private/MetasoundEditorGraphSchema.cpp
@@ -582,7 +582,7 @@ PRAGMA_ENABLE_DEPRECATION_WARNINGS
 
 			virtual bool CanConnectInputOfTypeAndAccess(FName InputTypeName, EMetasoundFrontendVertexAccessType InputAccessType) const override
 			{
-				const bool bCanConnect = Algo::AnyOf(Class.Interface.Outputs, [&InputTypeName, &InputAccessType](const FMetasoundFrontendClassOutput& Output)
+				const bool bCanConnect = Algo::AnyOf(Class.GetDefaultInterface().Outputs, [&InputTypeName, &InputAccessType](const FMetasoundFrontendClassOutput& Output)
 				{
 					return InputTypeName == Output.TypeName && FMetasoundFrontendClassVertex::CanConnectVertexAccessTypes(Output.AccessType, InputAccessType);
 				});
@@ -591,7 +591,7 @@ PRAGMA_ENABLE_DEPRECATION_WARNINGS
 
 			virtual bool CanConnectOutputOfTypeAndAccess(FName OutputTypeName, EMetasoundFrontendVertexAccessType OutputAccessType) const override
 			{
-				const bool bCanConnect = Algo::AnyOf(Class.Interface.Inputs, [&OutputTypeName, &OutputAccessType](const FMetasoundFrontendClassInput& Input)
+				const bool bCanConnect = Algo::AnyOf(Class.GetDefaultInterface().Inputs, [&OutputTypeName, &OutputAccessType](const FMetasoundFrontendClassInput& Input)
 				{
 					return OutputTypeName == Input.TypeName && FMetasoundFrontendClassVertex::CanConnectVertexAccessTypes(OutputAccessType, Input.AccessType);
 				});
@@ -2201,15 +2201,16 @@ FText UMetasoundEditorGraphSchema::GetPinDisplayName(const UEdGraphPin* Pin) con
 			auto PinMatchesClassVertex = [&Pin](const FMetasoundFrontendClassVertex& OtherVertex) { return OtherVertex.Name == Pin->GetFName(); };
 			const FMetasoundFrontendVertex* Vertex = nullptr;
 			const FMetasoundFrontendClassVertex* ClassVertex = nullptr;
+			const FMetasoundFrontendClassInterface& ClassInterface = Class->GetInterfaceForNode(*Node);
 			if (Pin->Direction == EGPD_Input)
 			{
 				Vertex = Builder.FindNodeInput(EdNode->GetNodeID(), Pin->GetFName());
-				ClassVertex = Class->Interface.Inputs.FindByPredicate(PinMatchesClassVertex);
+				ClassVertex = ClassInterface.Inputs.FindByPredicate(PinMatchesClassVertex);
 			}
 			else
 			{
 				Vertex = Builder.FindNodeOutput(EdNode->GetNodeID(), Pin->GetFName());
-				ClassVertex = Class->Interface.Outputs.FindByPredicate(PinMatchesClassVertex);
+				ClassVertex = ClassInterface.Outputs.FindByPredicate(PinMatchesClassVertex);
 			}
 
 			if (Vertex && ClassVertex)
diff --git a/Engine/Plugins/Runtime/Metasound/Source/MetasoundEngine/Private/MetasoundBuilderSubsystem.cpp b/Engine/Plugins/Runtime/Metasound/Source/MetasoundEngine/Private/MetasoundBuilderSubsystem.cpp
--- a/Engine/Plugins/Runtime/Metasound/Source/MetasoundEngine/Private/MetasoundBuilderSubsystem.cpp
+++ b/Engine/Plugins/Runtime/Metasound/Source/MetasoundEngine/Private/MetasoundBuilderSubsystem.cpp
@@ -445,7 +445,7 @@ void UMetaSoundSourceBuilder::OnInputAdded(int32 InputIndex)
 
 		const FMetasoundFrontendDocument& Doc = Builder.GetConstDocumentChecked();
 		const FMetasoundFrontendGraphClass& GraphClass = Doc.RootGraph;
-		const FMetasoundFrontendClassInput& NewInput = GraphClass.Interface.Inputs[InputIndex];
+		const FMetasoundFrontendClassInput& NewInput = GraphClass.GetDefaultInterface().Inputs[InputIndex];
 
 		constexpr bool bCreateUObjectProxies = true;
 		UMetaSoundSource& Source = GetMetaSoundSource();
@@ -484,6 +484,13 @@ void UMetaSoundSourceBuilder::OnNodeAdded(int32 NodeIndex) const
 		const FMetasoundFrontendClass* NodeClass = Builder.FindDependency(AddedNode.ClassID);
 		checkf(NodeClass, TEXT("Node successfully added to graph but document is missing associated dependency"));
 
+		// The logic for creating the FVertexInterface already lives in the FMetasoundFrontendGraphBuilder. 
+		// That logic should be exposed in an API that can be called here. In order
+		// to account for class overrides from node configurations, the vertex interface
+		// will need to be built from the `FMetasoundFrontendClass` & `FMetasoundFrontendNode`
+		// instead of being retrieved from the registry. 
+		// JIRA UE-235230
+
 		const FNodeRegistryKey& ClassKey = FNodeRegistryKey(NodeClass->Metadata);
 		FMetasoundFrontendRegistryContainer& NodeRegistry = *FMetasoundFrontendRegistryContainer::Get();
 		IDataTypeRegistry& DataTypeRegistry = IDataTypeRegistry::Get();
@@ -528,7 +535,8 @@ void UMetaSoundSourceBuilder::OnNodeAdded(int32 NodeIndex) const
 				else
 				{
 					// As a last resort, get default literal on node class 
-					DefaultLiteral = &(NodeClass->Interface.Inputs.Last().FindConstDefaultChecked(Frontend::DefaultPageID));
+					// JIRA UE-235230 - This should use the applicable class interface instead of the default class interface
+					DefaultLiteral = &(NodeClass->GetDefaultInterface().Inputs.Last().FindConstDefaultChecked(Frontend::DefaultPageID));
 				}
 
 				TSharedPtr<const FNodeClassMetadata> InputClassMetadata = DataTypeRegistry.GetInputClassMetadata(DataTypeName);
@@ -631,7 +639,7 @@ void UMetaSoundSourceBuilder::OnOutputAdded(int32 OutputIndex) const
 
 		const FMetasoundFrontendDocument& Doc = Builder.GetConstDocumentChecked();
 		const FMetasoundFrontendGraphClass& GraphClass = Doc.RootGraph;
-		const FMetasoundFrontendClassOutput& NewOutput = GraphClass.Interface.Outputs[OutputIndex];
+		const FMetasoundFrontendClassOutput& NewOutput = GraphClass.GetDefaultInterface().Outputs[OutputIndex];
 
 		Transactor.AddOutputDataSource(NewOutput.NodeID, NewOutput.Name);
 		return true;
@@ -705,7 +713,7 @@ void UMetaSoundSourceBuilder::OnRemovingInput(int32 InputIndex)
 
 		const FMetasoundFrontendDocument& Doc = Builder.GetConstDocumentChecked();
 		const FMetasoundFrontendGraphClass& GraphClass = Doc.RootGraph;
-		const FMetasoundFrontendClassInput& InputBeingRemoved = GraphClass.Interface.Inputs[InputIndex];
+		const FMetasoundFrontendClassInput& InputBeingRemoved = GraphClass.GetDefaultInterface().Inputs[InputIndex];
 
 		UMetaSoundSource& Source = GetMetaSoundSource();
 		Source.RuntimeInputData.InputMap.Remove(InputBeingRemoved.Name);
@@ -787,7 +795,7 @@ void UMetaSoundSourceBuilder::OnRemovingOutput(int32 OutputIndex) const
 
 		const FMetasoundFrontendDocument& Doc = Builder.GetConstDocumentChecked();
 		const FMetasoundFrontendGraphClass& GraphClass = Doc.RootGraph;
-		const FMetasoundFrontendClassOutput& OutputBeingRemoved = GraphClass.Interface.Outputs[OutputIndex];
+		const FMetasoundFrontendClassOutput& OutputBeingRemoved = GraphClass.GetDefaultInterface().Outputs[OutputIndex];
 
 		Transactor.RemoveOutputDataSource(OutputBeingRemoved.Name);
 		return true;
diff --git a/Engine/Plugins/Runtime/Metasound/Source/MetasoundEngine/Private/MetasoundDocumentBuilderRegistry.cpp b/Engine/Plugins/Runtime/Metasound/Source/MetasoundEngine/Private/MetasoundDocumentBuilderRegistry.cpp
--- a/Engine/Plugins/Runtime/Metasound/Source/MetasoundEngine/Private/MetasoundDocumentBuilderRegistry.cpp
+++ b/Engine/Plugins/Runtime/Metasound/Source/MetasoundEngine/Private/MetasoundDocumentBuilderRegistry.cpp
@@ -175,7 +175,7 @@ namespace Metasound::Engine
 		}
 
 		{ // Strip default input values
-			for (const FMetasoundFrontendClassInput& GraphInput : Document.RootGraph.Interface.Inputs)
+			for (const FMetasoundFrontendClassInput& GraphInput : Document.RootGraph.GetDefaultInterface().Inputs)
 			{
 				ResolvePageIDs.Reset();
 				auto AddPageID = [&ResolvePageIDs](const FGuid& PageID, const FMetasoundFrontendLiteral&) { ResolvePageIDs.Add(PageID); };
@@ -604,4 +604,4 @@ namespace Metasound::Engine
 	{
 		EventLogVerbosity.FindOrAdd(Event) = Verbosity;
 	}
-} // namespace Metasound::Engine
\ No newline at end of file
+} // namespace Metasound::Engine
diff --git a/Engine/Plugins/Runtime/Metasound/Source/MetasoundEngine/Private/MetasoundEngineAsset.h b/Engine/Plugins/Runtime/Metasound/Source/MetasoundEngine/Private/MetasoundEngineAsset.h
--- a/Engine/Plugins/Runtime/Metasound/Source/MetasoundEngine/Private/MetasoundEngineAsset.h
+++ b/Engine/Plugins/Runtime/Metasound/Source/MetasoundEngine/Private/MetasoundEngineAsset.h
@@ -187,7 +187,7 @@ namespace Metasound::Engine
 				ErrorIfMissing(Graph.PageID, LOCTEXT("GraphPageDescriptor", "graph"));
 			}
 
-			for (const FMetasoundFrontendClassInput& ClassInput : Document.RootGraph.Interface.Inputs)
+			for (const FMetasoundFrontendClassInput& ClassInput : Document.RootGraph.GetDefaultInterface().Inputs)
 			{
 				ClassInput.IterateDefaults([&](const FGuid& PageID, const FMetasoundFrontendLiteral&)
 				{
diff --git a/Engine/Plugins/Runtime/Metasound/Source/MetasoundEngine/Private/MetasoundGeneratorHandle.cpp b/Engine/Plugins/Runtime/Metasound/Source/MetasoundEngine/Private/MetasoundGeneratorHandle.cpp
--- a/Engine/Plugins/Runtime/Metasound/Source/MetasoundEngine/Private/MetasoundGeneratorHandle.cpp
+++ b/Engine/Plugins/Runtime/Metasound/Source/MetasoundEngine/Private/MetasoundGeneratorHandle.cpp
@@ -473,7 +473,7 @@ namespace Metasound
 
 			// Find the node id and type name
 			const FMetasoundFrontendClassOutput* OutputPtr =
-				Source->GetConstDocument().RootGraph.Interface.Outputs.FindByPredicate(
+				Source->GetConstDocument().RootGraph.GetDefaultInterface().Outputs.FindByPredicate(
 					[&AnalyzerAddress](const FMetasoundFrontendClassOutput& Output)
 					{
 						return Output.Name == AnalyzerAddress.OutputName;
diff --git a/Engine/Plugins/Runtime/Metasound/Source/MetasoundEngine/Private/MetasoundSource.cpp b/Engine/Plugins/Runtime/Metasound/Source/MetasoundEngine/Private/MetasoundSource.cpp
--- a/Engine/Plugins/Runtime/Metasound/Source/MetasoundEngine/Private/MetasoundSource.cpp
+++ b/Engine/Plugins/Runtime/Metasound/Source/MetasoundEngine/Private/MetasoundSource.cpp
@@ -1298,7 +1298,7 @@ void UMetaSoundSource::InitParametersInternal(const Metasound::TSortedVertexName
 
 bool UMetaSoundSource::IsParameterValid(const FAudioParameter& InParameter) const
 {
-	const TArray<FMetasoundFrontendClassInput>& Inputs = GetConstDocument().RootGraph.Interface.Inputs;
+	const TArray<FMetasoundFrontendClassInput>& Inputs = GetConstDocument().RootGraph.GetDefaultInterface().Inputs;
 	const FMetasoundFrontendVertex* Vertex = Algo::FindByPredicate(Inputs, [&InParameter] (const FMetasoundFrontendClassInput& Input)
 	{
 		return Input.Name == InParameter.ParamName;
@@ -1815,7 +1815,7 @@ Metasound::TSortedVertexNameMap<UMetaSoundSource::FRuntimeInput> UMetaSoundSourc
 
 	// Cache all inputs which are not private inputs.
 	TSortedVertexNameMap<FRuntimeInput> PublicInputs;
-	for (const FMetasoundFrontendClassInput& Input : Doc.RootGraph.Interface.Inputs)
+	for (const FMetasoundFrontendClassInput& Input : Doc.RootGraph.GetDefaultInterface().Inputs)
 	{
 		if (!PrivateInputs.Contains(Input.Name))
 		{
diff --git a/Engine/Plugins/Runtime/Metasound/Source/MetasoundEngine/Public/MetasoundBuilderBase.cpp b/Engine/Plugins/Runtime/Metasound/Source/MetasoundEngine/Public/MetasoundBuilderBase.cpp
--- a/Engine/Plugins/Runtime/Metasound/Source/MetasoundEngine/Public/MetasoundBuilderBase.cpp
+++ b/Engine/Plugins/Runtime/Metasound/Source/MetasoundEngine/Public/MetasoundBuilderBase.cpp
@@ -894,7 +894,7 @@ TArray<FName> UMetaSoundBuilderBase::GetGraphInputNames(EMetaSoundBuilderResult&
 {
 	TArray<FName> GraphInputNames;
 	const FMetasoundFrontendGraphClass& RootGraph = Builder.GetConstDocumentChecked().RootGraph;
-	for (const FMetasoundFrontendClassInput& Input : RootGraph.Interface.Inputs)
+	for (const FMetasoundFrontendClassInput& Input : RootGraph.GetDefaultInterface().Inputs)
 	{
 		GraphInputNames.Emplace(Input.Name);
 	}
@@ -905,7 +905,7 @@ TArray<FName> UMetaSoundBuilderBase::GetGraphOutputNames(EMetaSoundBuilderResult
 {
 	TArray<FName> GraphOutputNames;
 	const FMetasoundFrontendGraphClass& RootGraph = Builder.GetConstDocumentChecked().RootGraph;
-	for (const FMetasoundFrontendClassOutput& Output : RootGraph.Interface.Outputs)
+	for (const FMetasoundFrontendClassOutput& Output : RootGraph.GetDefaultInterface().Outputs)
 	{
 		GraphOutputNames.Emplace(Output.Name);
 	}
diff --git a/Engine/Plugins/Runtime/Metasound/Source/MetasoundEngineTest/Private/EngineTestMetaSoundAutomatedNodeTest.cpp b/Engine/Plugins/Runtime/Metasound/Source/MetasoundEngineTest/Private/EngineTestMetaSoundAutomatedNodeTest.cpp
--- a/Engine/Plugins/Runtime/Metasound/Source/MetasoundEngineTest/Private/EngineTestMetaSoundAutomatedNodeTest.cpp
+++ b/Engine/Plugins/Runtime/Metasound/Source/MetasoundEngineTest/Private/EngineTestMetaSoundAutomatedNodeTest.cpp
@@ -698,7 +698,7 @@ namespace Metasound::EngineTest{
 				// Vertex names must be set for input nodes
 				InputNodeInterface.GetInputInterface().At(0).VertexName = TestVertexName;
 				InputNodeInterface.GetOutputInterface().At(0).VertexName = TestVertexName;
-				InputNodeInterface.GetInputInterface().At(0).SetDefaultLiteral(NodeClass.Interface.Inputs[0].FindConstDefaultChecked(Frontend::DefaultPageID).ToLiteral(DataTypeName));
+				InputNodeInterface.GetInputInterface().At(0).SetDefaultLiteral(NodeClass.GetDefaultInterface().Inputs[0].FindConstDefaultChecked(Frontend::DefaultPageID).ToLiteral(DataTypeName));
 
 				Node = DataTypeRegistry.CreateInputNode(DataTypeName, FNodeData{ TestNodeName, TestNodeID, InputNodeInterface });
 			}
diff --git a/Engine/Plugins/Runtime/Metasound/Source/MetasoundEngineTest/Private/EngineTestMetasoundSource.cpp b/Engine/Plugins/Runtime/Metasound/Source/MetasoundEngineTest/Private/EngineTestMetasoundSource.cpp
--- a/Engine/Plugins/Runtime/Metasound/Source/MetasoundEngineTest/Private/EngineTestMetasoundSource.cpp
+++ b/Engine/Plugins/Runtime/Metasound/Source/MetasoundEngineTest/Private/EngineTestMetasoundSource.cpp
@@ -959,8 +959,8 @@ bool FAudioMetaSoundFindAssetClassInfo::RunTest(const FString& Parameters)
 				AddErrorIfFalse(ClassInterfaceInfo.DefinedInterfaces.Num() == 0, TEXT("Cool! We support declared interfaces now! Update the test!"));
 
 				AddErrorIfFalse(ClassInterfaceInfo.InheritedInterfaces.Num() == Doc.Interfaces.Num(), TEXT("MetaSound asset query tags did not return that asset implements number of interfaces as expected"));
-				AddErrorIfFalse(ClassInterfaceInfo.Inputs.Num() == GraphClass.Interface.Inputs.Num(), TEXT("MetaSound asset query tags did not return that asset implements number of inputs expected"));
-				AddErrorIfFalse(ClassInterfaceInfo.Outputs.Num() == GraphClass.Interface.Outputs.Num(), TEXT("MetaSound asset query tags did not return that asset implements number of outputs expected"));
+				AddErrorIfFalse(ClassInterfaceInfo.Inputs.Num() == GraphClass.GetDefaultInterface().Inputs.Num(), TEXT("MetaSound asset query tags did not return that asset implements number of inputs expected"));
+				AddErrorIfFalse(ClassInterfaceInfo.Outputs.Num() == GraphClass.GetDefaultInterface().Outputs.Num(), TEXT("MetaSound asset query tags did not return that asset implements number of outputs expected"));
 			}
 		}
 	}
diff --git a/Engine/Plugins/Runtime/Metasound/Source/MetasoundFrontend/Private/Interfaces/MetasoundFrontendInterface.cpp b/Engine/Plugins/Runtime/Metasound/Source/MetasoundFrontend/Private/Interfaces/MetasoundFrontendInterface.cpp
--- a/Engine/Plugins/Runtime/Metasound/Source/MetasoundFrontend/Private/Interfaces/MetasoundFrontendInterface.cpp
+++ b/Engine/Plugins/Runtime/Metasound/Source/MetasoundFrontend/Private/Interfaces/MetasoundFrontendInterface.cpp
@@ -139,13 +139,18 @@ namespace Metasound
 		}
 
 		bool IsSubsetOfClass(const FMetasoundFrontendInterface& InSubsetInterface, const FMetasoundFrontendClass& InSupersetClass)
+		{
+			return IsSubsetOfClassInterface(InSubsetInterface, InSupersetClass.GetDefaultInterface());
+		}
+
+		bool IsSubsetOfClassInterface(const FMetasoundFrontendInterface& InSubsetInterface, const FMetasoundFrontendClassInterface& InSupersetClassInterface)
 		{
 			// TODO: Environment variables are ignored as they are poorly supported and classes describe which environment variables are required,
 			// not which are supported 
 			using namespace InterfacePrivate;
 
-			const bool bIsInputSetIncluded = IsSetIncluded(InSubsetInterface.Inputs, InSupersetClass.Interface.Inputs, &IsLessThanVertex);
-			const bool bIsOutputSetIncluded = IsSetIncluded(InSubsetInterface.Outputs, InSupersetClass.Interface.Outputs, &IsLessThanVertex);
+			const bool bIsInputSetIncluded = IsSetIncluded(InSubsetInterface.Inputs, InSupersetClassInterface.Inputs, &IsLessThanVertex);
+			const bool bIsOutputSetIncluded = IsSetIncluded(InSubsetInterface.Outputs, InSupersetClassInterface.Outputs, &IsLessThanVertex);
 
 			return bIsInputSetIncluded && bIsOutputSetIncluded;
 		}
@@ -185,11 +190,16 @@ namespace Metasound
 		}
 
 		int32 InputOutputDifferenceCount(const FMetasoundFrontendClass& InClass, const FMetasoundFrontendInterface& InInterface)
+		{
+			return InputOutputDifferenceCount(InClass.GetDefaultInterface(), InInterface);
+		}
+
+		int32 InputOutputDifferenceCount(const FMetasoundFrontendClassInterface& InClassInterface, const FMetasoundFrontendInterface& InInterface)
 		{
 			using namespace InterfacePrivate;
 
-			int32 DiffCount = SetDifferenceCount(InClass.Interface.Inputs, InInterface.Inputs, &IsLessThanVertex);
-			DiffCount += SetDifferenceCount(InClass.Interface.Outputs, InInterface.Outputs, &IsLessThanVertex);
+			int32 DiffCount = SetDifferenceCount(InClassInterface.Inputs, InInterface.Inputs, &IsLessThanVertex);
+			DiffCount += SetDifferenceCount(InClassInterface.Outputs, InInterface.Outputs, &IsLessThanVertex);
 
 			return DiffCount;
 		}
@@ -210,7 +220,9 @@ namespace Metasound
 			{
 				using namespace InterfacePrivate;
 
-				for (const FMetasoundFrontendClassEnvironmentVariable& EnvVar : InClass.Interface.Environment)
+				// Environment variables are maintained on the default interface. Changing the environment variables in a
+				// class interface override is not supported. 
+				for (const FMetasoundFrontendClassEnvironmentVariable& EnvVar : InClass.GetDefaultInterface().Environment)
 				{
 					if (EnvVar.bIsRequired)
 					{
@@ -255,7 +267,7 @@ namespace Metasound
 			// Return the interface with the least amount of differences.
 			auto DifferencesFromRootGraph = [&](const FMetasoundFrontendInterface* CandidateInterface) -> int32
 			{ 
-				return InputOutputDifferenceCount(InRootGraph, *CandidateInterface);
+				return InputOutputDifferenceCount(InRootGraph.GetDefaultInterface(), *CandidateInterface);
 			};
 
 			Algo::SortBy(CandidateInterfaces, DifferencesFromRootGraph);
diff --git a/Engine/Plugins/Runtime/Metasound/Source/MetasoundFrontend/Private/MetasoundAssetBase.cpp b/Engine/Plugins/Runtime/Metasound/Source/MetasoundFrontend/Private/MetasoundAssetBase.cpp
--- a/Engine/Plugins/Runtime/Metasound/Source/MetasoundFrontend/Private/MetasoundAssetBase.cpp
+++ b/Engine/Plugins/Runtime/Metasound/Source/MetasoundFrontend/Private/MetasoundAssetBase.cpp
@@ -114,7 +114,7 @@ namespace Metasound
 				virtual FVertexInterface GetDefaultVertexInterface() const override
 				{
 					UE_LOG(LogMetaSound, Warning, TEXT("Proxy data is not available for non runtime node %s only used for serialization, so interface will not include object literals. Please ensure calling this function is intended."), *FrontendClass.Metadata.GetClassName().ToString())
-					return CreateVertexInterfaceFromClassNoProxy(FrontendClass);
+					return CreateDefaultVertexInterfaceFromClassNoProxy(FrontendClass);
 				}
 
 				virtual TInstancedStruct<FMetaSoundFrontendNodeConfiguration> CreateFrontendNodeConfiguration() const override
@@ -600,7 +600,7 @@ PRAGMA_ENABLE_DEPRECATION_WARNINGS
 
 	using FNameDataTypePair = TPair<FName, FName>;
 	const TSet<FMetasoundFrontendVersion>& InterfaceVersions = Document->Interfaces;
-	FMetasoundFrontendClassInterface& RootGraphClassInterface = Document->RootGraph.Interface;
+	FMetasoundFrontendClassInterface& RootGraphClassInterface = Document->RootGraph.GetDefaultInterface();
 
 	// 1. Gather inputs/outputs managed by interfaces
 	TMap<FNameDataTypePair, FMetasoundFrontendClassInput*> Inputs;
diff --git a/Engine/Plugins/Runtime/Metasound/Source/MetasoundFrontend/Private/MetasoundFrontend.cpp b/Engine/Plugins/Runtime/Metasound/Source/MetasoundFrontend/Private/MetasoundFrontend.cpp
--- a/Engine/Plugins/Runtime/Metasound/Source/MetasoundFrontend/Private/MetasoundFrontend.cpp
+++ b/Engine/Plugins/Runtime/Metasound/Source/MetasoundFrontend/Private/MetasoundFrontend.cpp
@@ -36,7 +36,7 @@ namespace Metasound
 			FMetasoundFrontendClass ClassDescription;
 
 			ClassDescription.Metadata = FMetasoundFrontendClassMetadata::GenerateClassMetadata(InNodeMetadata, ClassType);
-			ClassDescription.Interface = FMetasoundFrontendClassInterface::GenerateClassInterface(InNodeMetadata.DefaultInterface);
+			ClassDescription.SetDefaultInterface(FMetasoundFrontendClassInterface::GenerateClassInterface(InNodeMetadata.DefaultInterface));
 #if WITH_EDITORONLY_DATA
 			ClassDescription.Style = FMetasoundFrontendClassStyle::GenerateClassStyle(InNodeMetadata.DisplayStyle);
 #endif // WITH_EDITORONLY_DATA
@@ -100,12 +100,12 @@ namespace Metasound
 #endif
 		}
 
-		FVertexInterface CreateVertexInterfaceFromClassInternal(const FMetasoundFrontendClass& InNodeClass, const bool bCreateProxies, const FProxyDataCache* InProxyDataCache = nullptr)
+		FVertexInterface CreateDefaultVertexInterfaceFromClassInternal(const FMetasoundFrontendClass& InNodeClass, const bool bCreateProxies, const FProxyDataCache* InProxyDataCache = nullptr)
 		{
 			FVertexInterface Interface;
 
 			FInputVertexInterface& Inputs = Interface.GetInputInterface();
-			for (const FMetasoundFrontendClassInput& ClassInput : InNodeClass.Interface.Inputs)
+			for (const FMetasoundFrontendClassInput& ClassInput : InNodeClass.GetDefaultInterface().Inputs)
 			{
 				const FMetasoundFrontendLiteral* DefaultLiteral = ClassInput.FindConstDefault(DefaultPageID);
 				if (DefaultLiteral)
@@ -149,7 +149,7 @@ namespace Metasound
 			}
 
 			FOutputVertexInterface& Outputs = Interface.GetOutputInterface();
-			for (const FMetasoundFrontendClassOutput& ClassOutput : InNodeClass.Interface.Outputs)
+			for (const FMetasoundFrontendClassOutput& ClassOutput : InNodeClass.GetDefaultInterface().Outputs)
 			{
 				Outputs.Add(
 					FOutputDataVertex
@@ -162,7 +162,7 @@ namespace Metasound
 			}
 
 			FEnvironmentVertexInterface& Environments = Interface.GetEnvironmentInterface();
-			for (const FMetasoundFrontendClassEnvironmentVariable& ClassEnvironment : InNodeClass.Interface.Environment)
+			for (const FMetasoundFrontendClassEnvironmentVariable& ClassEnvironment : InNodeClass.GetDefaultInterface().Environment)
 			{
 				Environments.Add(FEnvironmentVertex{ ClassEnvironment.Name, FText::GetEmpty() });
 			}
@@ -170,15 +170,15 @@ namespace Metasound
 			return Interface;
 		}
 
-		FVertexInterface CreateVertexInterfaceFromClass(const FMetasoundFrontendClass& InNodeClass, const Metasound::Frontend::FProxyDataCache* InProxyDataCache)
+		FVertexInterface CreateDefaultVertexInterfaceFromClass(const FMetasoundFrontendClass& InNodeClass, const Metasound::Frontend::FProxyDataCache* InProxyDataCache)
 		{
 			checkf(InProxyDataCache || IsInGameThread(), TEXT("Vertex interface creation without proxy data cache must occur on game thread to safely create UObject proxies. Populate the proxy data cache on the game thread before calling this or use CreateVertexInterfaceFromClassNoProxy."))
-				return CreateVertexInterfaceFromClassInternal(InNodeClass, /*bCreateProxies=*/true, InProxyDataCache);
+				return CreateDefaultVertexInterfaceFromClassInternal(InNodeClass, /*bCreateProxies=*/true, InProxyDataCache);
 		}
 
-		FVertexInterface CreateVertexInterfaceFromClassNoProxy(const FMetasoundFrontendClass& InNodeClass)
+		FVertexInterface CreateDefaultVertexInterfaceFromClassNoProxy(const FMetasoundFrontendClass& InNodeClass)
 		{
-			return CreateVertexInterfaceFromClassInternal(InNodeClass, /*bCreateProxies=*/false);
+			return CreateDefaultVertexInterfaceFromClassInternal(InNodeClass, /*bCreateProxies=*/false);
 		}
 	}
 }
diff --git a/Engine/Plugins/Runtime/Metasound/Source/MetasoundFrontend/Private/MetasoundFrontendDocument.cpp b/Engine/Plugins/Runtime/Metasound/Source/MetasoundFrontend/Private/MetasoundFrontendDocument.cpp
--- a/Engine/Plugins/Runtime/Metasound/Source/MetasoundFrontend/Private/MetasoundFrontendDocument.cpp
+++ b/Engine/Plugins/Runtime/Metasound/Source/MetasoundFrontend/Private/MetasoundFrontendDocument.cpp
@@ -43,18 +43,6 @@ namespace Metasound
 
 				return Frontend::DefaultPageID;
 			}
-
-			FMetasoundFrontendNodeInterface MakeNodeInterface(const FMetasoundFrontendClass& InClass, const TInstancedStruct<FMetaSoundFrontendNodeConfiguration>& InConfiguration)
-			{
-				if (const FMetaSoundFrontendNodeConfiguration* ConfigurationPtr = InConfiguration.GetPtr())
-				{
-					return ConfigurationPtr->GetNodeInterface(InClass);
-				}
-				else
-				{
-					return FMetasoundFrontendNodeInterface(InClass.Interface);
-				}
-			}
 		} // namespace DocumentPrivate
 
 #if WITH_EDITORONLY_DATA
@@ -358,9 +346,80 @@ FMetasoundFrontendNodeInterface::FMetasoundFrontendNodeInterface(const FMetasoun
 	}
 }
 
-FMetasoundFrontendNodeInterface FMetaSoundFrontendNodeConfiguration::GetNodeInterface(const FMetasoundFrontendClass& InNodeClass) const
+void FMetasoundFrontendNodeInterface::Update(const FMetasoundFrontendClassInterface& InClassInterface)
 {
-	return FMetasoundFrontendNodeInterface{InNodeClass.Interface};
+	auto NoOp = [](const FMetasoundFrontendVertex&)->void {};
+
+	Update(InClassInterface, NoOp, NoOp);
+}
+
+void FMetasoundFrontendNodeInterface::Update(const FMetasoundFrontendClassInterface& InClassInterface, TFunctionRef<void(const FMetasoundFrontendVertex&)> OnPreRemoveInput, TFunctionRef<void(const FMetasoundFrontendVertex&)> OnPreRemoveOutput)
+{
+	TArray<const FMetasoundFrontendVertex*, TInlineAllocator<32>> UnmatchedVertices;
+
+	struct FFindMatchingVertex
+	{
+		const FMetasoundFrontendClassVertex& ClassVertex;
+		bool operator()(const FMetasoundFrontendVertex* InNodeVertex)
+		{
+			return (ClassVertex.Name == InNodeVertex->Name) && (ClassVertex.TypeName == InNodeVertex->TypeName);
+		}
+	};
+
+	// Update node inputs
+	Algo::Transform(Inputs, UnmatchedVertices, [](const FMetasoundFrontendVertex& Vertex) { return &Vertex; });
+	for (const FMetasoundFrontendClassInput& ClassInput : InClassInterface.Inputs)
+	{
+		if (const int32 Index = UnmatchedVertices.IndexOfByPredicate(FFindMatchingVertex(ClassInput)); Index != INDEX_NONE)
+		{
+			// Update the node vertex with anything new from the class vertex
+			UnmatchedVertices.RemoveAtSwap(Index, 1, EAllowShrinking::No);
+		}
+		else
+		{
+			// Add class input to node inputs
+			Inputs.Add(ClassInput);
+		}
+	}
+
+	// Remove any inputs that did not exist in the class interface
+	for (const FMetasoundFrontendVertex* UnmatchedInput : UnmatchedVertices)
+	{
+		// Allow outside systems to react before removing the unmatched node inputs
+		OnPreRemoveInput(*UnmatchedInput);
+		Inputs.RemoveSingleSwap(*UnmatchedInput);
+	}
+
+	// Update node outputs
+	UnmatchedVertices.Reset();
+	Algo::Transform(Outputs, UnmatchedVertices, [](const FMetasoundFrontendVertex& Vertex) { return &Vertex; });
+	for (const FMetasoundFrontendClassOutput& ClassOutput : InClassInterface.Outputs)
+	{
+		if (const int32 Index = UnmatchedVertices.IndexOfByPredicate(FFindMatchingVertex(ClassOutput)); Index != INDEX_NONE)
+		{
+			// Update the node vertex with anything new from the class vertex
+			UnmatchedVertices.RemoveAtSwap(Index, 1, EAllowShrinking::No);
+		}
+		else
+		{
+			// Add class input to node inputs
+			Outputs.Add(ClassOutput);
+		}
+	}
+
+	// Remove any inputs that did not exist in the class interface
+	for (const FMetasoundFrontendVertex* UnmatchedOutput : UnmatchedVertices)
+	{
+		// Allow outside systems to react before removing the unmatched node inputs
+		OnPreRemoveOutput(*UnmatchedOutput);
+		Outputs.RemoveSingleSwap(*UnmatchedOutput);
+	}
+}
+
+TInstancedStruct<FMetasoundFrontendClassInterface> FMetaSoundFrontendNodeConfiguration::OverrideDefaultInterface(const FMetasoundFrontendClass& InNodeClass) const
+{
+	// By default, node configurations do not override the class interface. 
+	return TInstancedStruct<FMetasoundFrontendClassInterface>{};
 }
 
 TSharedPtr<const Metasound::IOperatorData> FMetaSoundFrontendNodeConfiguration::GetOperatorData() const
@@ -376,9 +435,24 @@ FMetasoundFrontendNode::FMetasoundFrontendNode(const FMetasoundFrontendClass& In
 FMetasoundFrontendNode::FMetasoundFrontendNode(const FMetasoundFrontendClass& InClass, TInstancedStruct<FMetaSoundFrontendNodeConfiguration> InConfiguration)
 : ClassID(InClass.ID)
 , Name(InClass.Metadata.GetClassName().Name.ToString())
-, Interface(Metasound::Frontend::DocumentPrivate::MakeNodeInterface(InClass, InConfiguration))
 , Configuration(MoveTemp(InConfiguration))
 {
+	const FMetasoundFrontendClassInterface* ClassInterfacePtr = nullptr;
+
+	// Determine whether to initialize the node interface with the class's default interface or
+	// an override
+	if (const FMetaSoundFrontendNodeConfiguration* ConfigurationPtr = Configuration.GetPtr())
+	{
+		ClassInterfaceOverride = ConfigurationPtr->OverrideDefaultInterface(InClass);
+		ClassInterfacePtr = ClassInterfaceOverride.GetPtr();
+	}
+
+	if (ClassInterfacePtr == nullptr)
+	{
+		ClassInterfacePtr = &InClass.GetDefaultInterface();
+	}
+
+	Interface = FMetasoundFrontendNodeInterface{*ClassInterfacePtr};
 }
 
 FString FMetasoundFrontendVersion::ToString() const
@@ -402,6 +476,11 @@ bool FMetasoundFrontendVertex::IsFunctionalEquivalent(const FMetasoundFrontendVe
 	return (InLHS.Name == InRHS.Name) && (InLHS.TypeName == InRHS.TypeName);
 }
 
+bool operator==(const FMetasoundFrontendVertex& InLHS, const FMetasoundFrontendVertex& InRHS)
+{
+	return (InLHS.Name == InRHS.Name) && (InLHS.TypeName == InRHS.TypeName) && (InLHS.VertexID == InRHS.VertexID);
+}
+
 void FMetasoundFrontendClassVertex::SplitName(FName& OutNamespace, FName& OutParameterName) const
 {
 	Audio::FParameterPath::SplitName(Name, OutNamespace, OutParameterName);
@@ -811,6 +890,43 @@ void FMetasoundFrontendClassMetadata::SetClassName(const FMetasoundFrontendClass
 	SetWithChangeID(InClassName, ClassName, ChangeID);
 }
 
+void FMetasoundFrontendClass::SetDefaultInterface(const FMetasoundFrontendClassInterface& InInterface)
+{
+PRAGMA_DISABLE_DEPRECATION_WARNINGS
+	Interface = InInterface;
+PRAGMA_ENABLE_DEPRECATION_WARNINGS
+}
+
+FMetasoundFrontendClassInterface& FMetasoundFrontendClass::GetDefaultInterface()
+{
+PRAGMA_DISABLE_DEPRECATION_WARNINGS
+	return Interface;
+PRAGMA_ENABLE_DEPRECATION_WARNINGS
+}
+
+const FMetasoundFrontendClassInterface& FMetasoundFrontendClass::GetDefaultInterface() const
+{
+PRAGMA_DISABLE_DEPRECATION_WARNINGS
+	return Interface;
+PRAGMA_ENABLE_DEPRECATION_WARNINGS
+}
+
+const FMetasoundFrontendClassInterface& FMetasoundFrontendClass::GetInterfaceForNode(const FMetasoundFrontendNode& InNode) const
+{
+	if (const FMetasoundFrontendClassInterface* InterfaceOverride = InNode.ClassInterfaceOverride.GetPtr())
+	{
+#if !UE_BUILD_SHIPPING
+		if (Metadata.GetType() != EMetasoundFrontendClassType::External)
+		{
+			UE_LOG(LogMetaSound, Error, TEXT("Found class interface override on class %s. Class interface overrides are intended to be used on external nodes"), *Metadata.GetClassName().ToString());
+		}
+#endif // !UE_BUILD_SHIPPING
+		return *InterfaceOverride;
+	}
+
+	return GetDefaultInterface();
+}
+
 #if WITH_EDITOR
 bool FMetasoundFrontendClass::CacheGraphDependencyMetadataFromRegistry(FMetasoundFrontendClass& InOutDependency)
 {
@@ -819,7 +935,7 @@ bool FMetasoundFrontendClass::CacheGraphDependencyMetadataFromRegistry(FMetasoun
 	const FNodeRegistryKey Key = FNodeRegistryKey(InOutDependency.Metadata);
 	FMetasoundFrontendClass RegistryClass;
 
-	FMetasoundFrontendRegistryContainer* Registry = FMetasoundFrontendRegistryContainer::Get();
+	INodeClassRegistry* Registry = INodeClassRegistry::Get();
 	if (ensure(Registry))
 	{
 		if (Registry->FindFrontendClassFromRegistered(Key, RegistryClass))
@@ -848,21 +964,21 @@ bool FMetasoundFrontendClass::CacheGraphDependencyMetadataFromRegistry(FMetasoun
 
 			FMetasoundFrontendInterfaceStyle InputStyle;
 			FVertexMetadataMap InputMembers;
-			Algo::Transform(RegistryClass.Interface.Inputs, InputMembers, [&](const FMetasoundFrontendClassInput& Input) { return MakePairFromVertex(Input); });
-			Algo::ForEach(InOutDependency.Interface.Inputs, [&](FMetasoundFrontendClassInput& Input)
+			Algo::Transform(RegistryClass.GetDefaultInterface().Inputs, InputMembers, [&](const FMetasoundFrontendClassInput& Input) { return MakePairFromVertex(Input); });
+			Algo::ForEach(InOutDependency.GetDefaultInterface().Inputs, [&](FMetasoundFrontendClassInput& Input)
 			{
 				AddRegistryVertexMetadata(InputMembers, Input, InputStyle);
 			});
-			InOutDependency.Interface.SetInputStyle(InputStyle);
+			InOutDependency.GetDefaultInterface().SetInputStyle(InputStyle);
 
 			FMetasoundFrontendInterfaceStyle OutputStyle;
 			FVertexMetadataMap OutputMembers;
-			Algo::Transform(RegistryClass.Interface.Outputs, OutputMembers, [&](const FMetasoundFrontendClassOutput& Output) { return MakePairFromVertex(Output); });
-			Algo::ForEach(InOutDependency.Interface.Outputs, [&](FMetasoundFrontendClassOutput& Output)
+			Algo::Transform(RegistryClass.GetDefaultInterface().Outputs, OutputMembers, [&](const FMetasoundFrontendClassOutput& Output) { return MakePairFromVertex(Output); });
+			Algo::ForEach(InOutDependency.GetDefaultInterface().Outputs, [&](FMetasoundFrontendClassOutput& Output)
 			{
 				AddRegistryVertexMetadata(OutputMembers, Output, OutputStyle);
 			});
-			InOutDependency.Interface.SetOutputStyle(OutputStyle);
+			InOutDependency.GetDefaultInterface().SetOutputStyle(OutputStyle);
 
 			return true;
 		}
@@ -937,6 +1053,11 @@ FMetasoundFrontendClassInputDefault::FMetasoundFrontendClassInputDefault(const F
 }
 
 bool FMetasoundFrontendClassInputDefault::IsFunctionalEquivalent(const FMetasoundFrontendClassInputDefault& InLHS, const FMetasoundFrontendClassInputDefault& InRHS)
+{
+	return InLHS == InRHS;
+}
+
+bool operator==(const FMetasoundFrontendClassInputDefault& InLHS, const FMetasoundFrontendClassInputDefault& InRHS)
 {
 	if (InLHS.PageID != InRHS.PageID)
 	{
@@ -1434,7 +1555,12 @@ namespace Metasound::Frontend
 
 	void ForEachLiteral(const FMetasoundFrontendClass& InClass, FForEachLiteralFunctionRef OnLiteral)
 	{
-		for (const FMetasoundFrontendClassInput& ClassInput : InClass.Interface.Inputs)
+		ForEachLiteral(InClass.GetDefaultInterface(), OnLiteral);
+	}
+
+	void ForEachLiteral(const FMetasoundFrontendClassInterface& InClassInterface, FForEachLiteralFunctionRef OnLiteral)
+	{
+		for (const FMetasoundFrontendClassInput& ClassInput : InClassInterface.Inputs)
 		{
 			const FGuid PageID = DocumentPrivate::ResolveTargetPageID(ClassInput);
 			const FMetasoundFrontendLiteral& DefaultLiteral = ClassInput.FindConstDefaultChecked(PageID);
@@ -1456,5 +1582,10 @@ namespace Metasound::Frontend
 				OnLiteral(InputVertex->TypeName, VertexLiteral.Value);
 			}
 		}
+
+		if (InNode.ClassInterfaceOverride.IsValid())
+		{
+			ForEachLiteral(InNode.ClassInterfaceOverride.Get(), OnLiteral);
+		}
 	}
 } // namespace Metasound::Frontend
diff --git a/Engine/Plugins/Runtime/Metasound/Source/MetasoundFrontend/Private/MetasoundFrontendDocumentAccessPtr.cpp b/Engine/Plugins/Runtime/Metasound/Source/MetasoundFrontend/Private/MetasoundFrontendDocumentAccessPtr.cpp
--- a/Engine/Plugins/Runtime/Metasound/Source/MetasoundFrontend/Private/MetasoundFrontendDocumentAccessPtr.cpp
+++ b/Engine/Plugins/Runtime/Metasound/Source/MetasoundFrontend/Private/MetasoundFrontendDocumentAccessPtr.cpp
@@ -234,15 +234,15 @@ namespace Metasound
 			/** A callable object for accessing arrays of inputs from a class. */
 			struct FGetClassInputArray 
 			{
-			 	TArray<FMetasoundFrontendClassInput>& operator()(FMetasoundFrontendClass& InClass) const { return InClass.Interface.Inputs; }
-			 	const TArray<FMetasoundFrontendClassInput>& operator()(const FMetasoundFrontendClass& InClass) const { return InClass.Interface.Inputs; }
+			 	TArray<FMetasoundFrontendClassInput>& operator()(FMetasoundFrontendClass& InClass) const { return InClass.GetDefaultInterface().Inputs; }
+			 	const TArray<FMetasoundFrontendClassInput>& operator()(const FMetasoundFrontendClass& InClass) const { return InClass.GetDefaultInterface().Inputs; }
 			};
 
 			/** A callable object for accessing arrays of outputs from a class. */
 			struct FGetClassOutputArray 
 			{
-			 	TArray<FMetasoundFrontendClassOutput>& operator()(FMetasoundFrontendClass& InClass) const { return InClass.Interface.Outputs; }
-			 	const TArray<FMetasoundFrontendClassOutput>& operator()(const FMetasoundFrontendClass& InClass) const { return InClass.Interface.Outputs; }
+			 	TArray<FMetasoundFrontendClassOutput>& operator()(FMetasoundFrontendClass& InClass) const { return InClass.GetDefaultInterface().Outputs; }
+			 	const TArray<FMetasoundFrontendClassOutput>& operator()(const FMetasoundFrontendClass& InClass) const { return InClass.GetDefaultInterface().Outputs; }
 			};
 
 			/** A callable object for accessing arrays of nodes from a graph class. */
diff --git a/Engine/Plugins/Runtime/Metasound/Source/MetasoundFrontend/Private/MetasoundFrontendDocumentBuilder.cpp b/Engine/Plugins/Runtime/Metasound/Source/MetasoundFrontend/Private/MetasoundFrontendDocumentBuilder.cpp
--- a/Engine/Plugins/Runtime/Metasound/Source/MetasoundFrontend/Private/MetasoundFrontendDocumentBuilder.cpp
+++ b/Engine/Plugins/Runtime/Metasound/Source/MetasoundFrontend/Private/MetasoundFrontendDocumentBuilder.cpp
@@ -1164,8 +1164,8 @@ const FMetasoundFrontendNode* FMetaSoundFrontendDocumentBuilder::AddGraphInput(F
 
 	if (NewNode)
 	{
-		const int32 NewIndex = RootGraph.Interface.Inputs.Num();
-		FMetasoundFrontendClassInput& NewInput = RootGraph.Interface.Inputs.Add_GetRef(MoveTemp(ClassInput));
+		const int32 NewIndex = RootGraph.GetDefaultInterface().Inputs.Num();
+		FMetasoundFrontendClassInput& NewInput = RootGraph.GetDefaultInterface().Inputs.Add_GetRef(MoveTemp(ClassInput));
 		if (!NewInput.VertexID.IsValid())
 		{
 			NewInput.VertexID = FDocumentIDGenerator::Get().CreateVertexID(Document);
@@ -1264,8 +1264,8 @@ const FMetasoundFrontendNode* FMetaSoundFrontendDocumentBuilder::AddGraphOutput(
 
 	if (bAddedNodes)
 	{
-		const int32 NewIndex = RootGraph.Interface.Outputs.Num();
-		FMetasoundFrontendClassOutput& NewOutput = RootGraph.Interface.Outputs.Add_GetRef(MoveTemp(ClassOutput));
+		const int32 NewIndex = RootGraph.GetDefaultInterface().Outputs.Num();
+		FMetasoundFrontendClassOutput& NewOutput = RootGraph.GetDefaultInterface().Outputs.Add_GetRef(MoveTemp(ClassOutput));
 		if (!NewOutput.VertexID.IsValid())
 		{
 			NewOutput.VertexID = FDocumentIDGenerator::Get().CreateVertexID(Document);
@@ -1874,12 +1874,12 @@ void FMetaSoundFrontendDocumentBuilder::ClearDocument(TSharedRef<Metasound::Fron
 	FMetasoundFrontendDocument& Doc = GetDocumentChecked();
 	FMetasoundFrontendGraphClass& GraphClass = Doc.RootGraph;
 
-	GraphClass.Interface.Inputs.Empty();
-	GraphClass.Interface.Outputs.Empty();
+	GraphClass.GetDefaultInterface().Inputs.Empty();
+	GraphClass.GetDefaultInterface().Outputs.Empty();
 
 #if WITH_EDITOR
-	GraphClass.Interface.SetInputStyle({ });
-	GraphClass.Interface.SetOutputStyle({ });
+	GraphClass.GetDefaultInterface().SetInputStyle({ });
+	GraphClass.GetDefaultInterface().SetOutputStyle({ });
 #endif // WITH_EDITOR
 
 	GraphClass.PresetOptions.InputsInheritingDefault.Empty();
@@ -1903,9 +1903,9 @@ void FMetaSoundFrontendDocumentBuilder::ClearDocument(TSharedRef<Metasound::Fron
 	});
 #endif // !WITH_EDITORONLY_DATA
 
-	GraphClass.Interface.Inputs.Empty();
-	GraphClass.Interface.Outputs.Empty();
-	GraphClass.Interface.Environment.Empty();
+	GraphClass.GetDefaultInterface().Inputs.Empty();
+	GraphClass.GetDefaultInterface().Outputs.Empty();
+	GraphClass.GetDefaultInterface().Environment.Empty();
 
 	Doc.Interfaces.Empty();
 	Doc.Dependencies.Empty();
@@ -2089,7 +2089,7 @@ bool FMetaSoundFrontendDocumentBuilder::ConvertToPreset(const FMetasoundFrontend
 	{
 		PresetAssetRootGraph.PresetOptions.InputsInheritingDefault.Reset();
 		auto GetInputName = [](const FMetasoundFrontendClassInput& Input) { return Input.Name; };
-		Algo::Transform(PresetAssetRootGraph.Interface.Inputs, PresetAssetRootGraph.PresetOptions.InputsInheritingDefault, GetInputName);
+		Algo::Transform(PresetAssetRootGraph.GetDefaultInterface().Inputs, PresetAssetRootGraph.PresetOptions.InputsInheritingDefault, GetInputName);
 		PresetAssetRootGraph.PresetOptions.bIsPreset = true;
 	}
 
@@ -2378,6 +2378,7 @@ FMetaSoundFrontendGraphComment* FMetaSoundFrontendDocumentBuilder::FindGraphComm
 }
 #endif // WITH_EDITORONLY_DATA
 
+
 const FMetasoundFrontendNode* FMetaSoundFrontendDocumentBuilder::FindHeadNodeInVariableStack(FName VariableName, const FGuid* InPageID) const
 {
 	// The variable "stack" is [GetDelayedNodes, SetNode, GetNodes].
@@ -2649,6 +2650,7 @@ const FMetasoundFrontendVertex* FMetaSoundFrontendDocumentBuilder::FindNodeInput
 const TArray<FMetasoundFrontendClassInputDefault>* FMetaSoundFrontendDocumentBuilder::FindNodeClassInputDefaults(const FGuid& InNodeID, FName InVertexName, const FGuid* InPageID) const
 {
 	using namespace Metasound;
+	using namespace Metasound::Frontend;
 
 	if (const FMetasoundFrontendNode* Node = FindNode(InNodeID, InPageID))
 	{
@@ -2660,7 +2662,8 @@ const TArray<FMetasoundFrontendClassInputDefault>* FMetaSoundFrontendDocumentBui
 				case EMetasoundFrontendClassType::External:
 				{
 					auto MatchesName = [&InVertexName](const FMetasoundFrontendClassInput& Input) { return Input.Name == InVertexName; };
-					if (const FMetasoundFrontendClassInput* Input = Class->Interface.Inputs.FindByPredicate(MatchesName))
+					const FMetasoundFrontendClassInterface& ClassInterface = Class->GetInterfaceForNode(*Node);
+					if (const FMetasoundFrontendClassInput* Input = ClassInterface.Inputs.FindByPredicate(MatchesName))
 					{
 						return &Input->GetDefaults();
 					}
@@ -2671,7 +2674,7 @@ const TArray<FMetasoundFrontendClassInputDefault>* FMetaSoundFrontendDocumentBui
 				case EMetasoundFrontendClassType::Output:
 				case EMetasoundFrontendClassType::Literal:
 				{
-					return &Class->Interface.Inputs.Last().GetDefaults();
+					return &Class->GetInterfaceForNode(*Node).Inputs.Last().GetDefaults();
 				}
 
 				case EMetasoundFrontendClassType::Variable:
@@ -2681,7 +2684,9 @@ const TArray<FMetasoundFrontendClassInputDefault>* FMetaSoundFrontendDocumentBui
 				{
 					using namespace VariableNames;
 					auto IsDataInput = [](const FMetasoundFrontendClassInput& Input) { return Input.Name == METASOUND_GET_PARAM_NAME(InputData); };
-					if (const FMetasoundFrontendClassInput* Input = Class->Interface.Inputs.FindByPredicate(IsDataInput))
+
+					const FMetasoundFrontendClassInterface& ClassInterface = Class->GetInterfaceForNode(*Node);
+					if (const FMetasoundFrontendClassInput* Input = ClassInterface.Inputs.FindByPredicate(IsDataInput))
 					{
 						return &Input->GetDefaults();
 					}
@@ -2925,7 +2930,7 @@ TArray<const FMetasoundFrontendNode*> FMetaSoundFrontendDocumentBuilder::GetGrap
 	FMetasoundFrontendGraphClass& RootGraph = GetDocumentChecked().RootGraph;
 	if (const int32* Index = DocumentCache->GetInterfaceCache().FindInputIndex(InInputName))
 	{
-		const FMetasoundFrontendClassInput& InputClass = RootGraph.Interface.Inputs[*Index];
+		const FMetasoundFrontendClassInput& InputClass = RootGraph.GetDefaultInterface().Inputs[*Index];
 		const FMetasoundFrontendGraph& Graph = RootGraph.FindConstGraphChecked(PageID);
 		const IDocumentGraphNodeCache& NodeCache = DocumentCache->GetNodeCache(PageID);
 		const IDocumentGraphEdgeCache& EdgeCache = DocumentCache->GetEdgeCache(PageID);
@@ -3073,7 +3078,8 @@ EMetasoundFrontendVertexAccessType FMetaSoundFrontendDocumentBuilder::GetNodeInp
 
 				case EMetasoundFrontendClassType::Output:
 				{
-					const FMetasoundFrontendClassInput& ClassInput = Class->Interface.Inputs.Last();
+					const FMetasoundFrontendClassInterface& ClassInterface = Class->GetInterfaceForNode(Node);
+					const FMetasoundFrontendClassInput& ClassInput = ClassInterface.Inputs.Last();
 					return ClassInput.AccessType;
 				}
 
@@ -3085,7 +3091,8 @@ EMetasoundFrontendVertexAccessType FMetaSoundFrontendDocumentBuilder::GetNodeInp
 			if (const FMetasoundFrontendVertex* Vertex = Node.Interface.Inputs.FindByPredicate(IsVertexID))
 			{
 				auto IsClassInput = [VertexName = Vertex->Name](const FMetasoundFrontendClassInput& Input) { return Input.Name == VertexName; };
-				if (const FMetasoundFrontendClassInput* ClassInput = Class->Interface.Inputs.FindByPredicate(IsClassInput))
+				const FMetasoundFrontendClassInterface& ClassInterface = Class->GetInterfaceForNode(Node);
+				if (const FMetasoundFrontendClassInput* ClassInput = ClassInterface.Inputs.FindByPredicate(IsClassInput))
 				{
 					return ClassInput->AccessType;
 				}
@@ -3099,6 +3106,7 @@ EMetasoundFrontendVertexAccessType FMetaSoundFrontendDocumentBuilder::GetNodeInp
 const FMetasoundFrontendLiteral* FMetaSoundFrontendDocumentBuilder::GetNodeInputClassDefault(const FGuid& InNodeID, const FGuid& InVertexID, const FGuid* InPageID) const
 {
 	using namespace Metasound;
+	using namespace Metasound::Frontend;
 
 	const FGuid& PageID = InPageID ? *InPageID : BuildPageID;
 	const Frontend::IDocumentGraphNodeCache& NodeCache = DocumentCache->GetNodeCache(PageID);
@@ -3112,18 +3120,19 @@ const FMetasoundFrontendLiteral* FMetaSoundFrontendDocumentBuilder::GetNodeInput
 			if (const FMetasoundFrontendClass* Class = DocumentCache->FindDependency(Node.ClassID))
 			{
 				const EMetasoundFrontendClassType ClassType = Class->Metadata.GetType();
+				const FMetasoundFrontendClassInterface& ClassInterface = Class->GetInterfaceForNode(Node);
 				switch (ClassType)
 				{
 					case EMetasoundFrontendClassType::Output:
 					{
-						const FMetasoundFrontendClassInput& ClassInput = Class->Interface.Inputs.Last();
+						const FMetasoundFrontendClassInput& ClassInput = ClassInterface.Inputs.Last();
 						return ClassInput.FindConstDefault(Frontend::DefaultPageID);
 					}
 
 					default:
 					{
 						auto IsClassInput = [VertexName = Vertex->Name](const FMetasoundFrontendClassInput& Input) { return Input.Name == VertexName; };
-						if (const FMetasoundFrontendClassInput* ClassInput = Class->Interface.Inputs.FindByPredicate(IsClassInput))
+						if (const FMetasoundFrontendClassInput* ClassInput = ClassInterface.Inputs.FindByPredicate(IsClassInput))
 						{
 							return ClassInput->FindConstDefault(Frontend::DefaultPageID);
 						}
@@ -3200,7 +3209,8 @@ EMetasoundFrontendVertexAccessType FMetaSoundFrontendDocumentBuilder::GetNodeOut
 
 				case EMetasoundFrontendClassType::Input:
 				{
-					const FMetasoundFrontendClassOutput& ClassOutput = Class->Interface.Outputs.Last();
+					const FMetasoundFrontendClassInterface& ClassInterface = Class->GetInterfaceForNode(Node);
+					const FMetasoundFrontendClassOutput& ClassOutput = ClassInterface.Outputs.Last();
 					return ClassOutput.AccessType;
 				}
 
@@ -3213,7 +3223,8 @@ EMetasoundFrontendVertexAccessType FMetaSoundFrontendDocumentBuilder::GetNodeOut
 			if (const FMetasoundFrontendVertex* Vertex = Node.Interface.Outputs.FindByPredicate(IsVertexID))
 			{
 				auto IsClassInput = [VertexName = Vertex->Name](const FMetasoundFrontendClassInput& Output) { return Output.Name == VertexName; };
-				if (const FMetasoundFrontendClassOutput* ClassOutput = Class->Interface.Outputs.FindByPredicate(IsClassInput))
+				const FMetasoundFrontendClassInterface& ClassInterface = Class->GetInterfaceForNode(Node);
+				if (const FMetasoundFrontendClassOutput* ClassOutput = ClassInterface.Outputs.FindByPredicate(IsClassInput))
 				{
 					return ClassOutput->AccessType;
 				}
@@ -3234,7 +3245,7 @@ const bool FMetaSoundFrontendDocumentBuilder::GetIsAdvancedDisplay(const FName M
 	{
 		if (const int32* Index = DocumentCache->GetInterfaceCache().FindInputIndex(MemberName))
 		{
-			const FMetasoundFrontendClassInput& GraphInput = Document.RootGraph.Interface.Inputs[*Index];
+			const FMetasoundFrontendClassInput& GraphInput = Document.RootGraph.GetDefaultInterface().Inputs[*Index];
 			return GraphInput.Metadata.bIsAdvancedDisplay;
 		}
 	}
@@ -3243,7 +3254,7 @@ const bool FMetaSoundFrontendDocumentBuilder::GetIsAdvancedDisplay(const FName M
 	{
 		if (const int32* Index = DocumentCache->GetInterfaceCache().FindOutputIndex(MemberName))
 		{
-			const FMetasoundFrontendClassOutput& GraphOutput = Document.RootGraph.Interface.Outputs[*Index];
+			const FMetasoundFrontendClassOutput& GraphOutput = Document.RootGraph.GetDefaultInterface().Outputs[*Index];
 			return GraphOutput.Metadata.bIsAdvancedDisplay;
 		}
 	}
@@ -3874,7 +3885,7 @@ bool FMetaSoundFrontendDocumentBuilder::RemoveGraphInputDefault(FName InputName,
 
 	auto NameMatchesInput = [&InputName](const FMetasoundFrontendClassInput& Input) { return Input.Name == InputName; };
 	FMetasoundFrontendDocument& Document = GetDocumentChecked();
-	TArray<FMetasoundFrontendClassInput>& Inputs = Document.RootGraph.Interface.Inputs;
+	TArray<FMetasoundFrontendClassInput>& Inputs = Document.RootGraph.GetDefaultInterface().Inputs;
 
 	const int32 Index = Inputs.IndexOfByPredicate(NameMatchesInput);
 	if (Index != INDEX_NONE)
@@ -4091,7 +4102,7 @@ bool FMetaSoundFrontendDocumentBuilder::RemoveGraphInput(FName InputName)
 	FMetasoundFrontendDocument& Document = GetDocumentChecked();
 	if (const int32* IndexPtr = DocumentCache->GetInterfaceCache().FindInputIndex(InputName))
 	{
-		TArray<FMetasoundFrontendClassInput>& Inputs = Document.RootGraph.Interface.Inputs;
+		TArray<FMetasoundFrontendClassInput>& Inputs = Document.RootGraph.GetDefaultInterface().Inputs;
 		const FGuid NodeID = Inputs[*IndexPtr].NodeID;
 		FGuid ClassID;
 		bool bNodesRemoved = true;
@@ -4189,7 +4200,7 @@ bool FMetaSoundFrontendDocumentBuilder::RemoveGraphOutput(FName OutputName)
 
 	if (bNodesRemoved)
 	{
-		TArray<FMetasoundFrontendClassOutput>& Outputs = Document.RootGraph.Interface.Outputs;
+		TArray<FMetasoundFrontendClassOutput>& Outputs = Document.RootGraph.GetDefaultInterface().Outputs;
 		auto OutputNameMatches = [OutputName](const FMetasoundFrontendClassOutput& Output) { return Output.Name == OutputName; };
 		const int32 Index = Outputs.IndexOfByPredicate(OutputNameMatches);
 		if (Index != INDEX_NONE)
@@ -4450,7 +4461,7 @@ bool FMetaSoundFrontendDocumentBuilder::ResetGraphInputDefault(FName InputName)
 
 	auto NameMatchesInput = [&InputName](const FMetasoundFrontendClassInput& Input) { return Input.Name == InputName; };
 	FMetasoundFrontendDocument& Document = GetDocumentChecked();
-	TArray<FMetasoundFrontendClassInput>& Inputs = Document.RootGraph.Interface.Inputs;
+	TArray<FMetasoundFrontendClassInput>& Inputs = Document.RootGraph.GetDefaultInterface().Inputs;
 
 	const int32 Index = Inputs.IndexOfByPredicate(NameMatchesInput);
 	if (Index != INDEX_NONE)
@@ -4543,7 +4554,7 @@ bool FMetaSoundFrontendDocumentBuilder::SetGraphInputAdvancedDisplay(const FName
 
 	if (const int32* Index = DocumentCache->GetInterfaceCache().FindInputIndex(InputName))
 	{
-		FMetasoundFrontendClassInput& GraphInput = RootGraph.Interface.Inputs[*Index];
+		FMetasoundFrontendClassInput& GraphInput = RootGraph.GetDefaultInterface().Inputs[*Index];
 		if (GraphInput.Metadata.bIsAdvancedDisplay != InAdvancedDisplay)
 		{
 			GraphInput.Metadata.SetIsAdvancedDisplay(InAdvancedDisplay);
@@ -4573,7 +4584,7 @@ bool FMetaSoundFrontendDocumentBuilder::SetGraphInputAccessType(FName InputName,
 
 	FMetasoundFrontendDocument& Document = GetDocumentChecked();
 	FMetasoundFrontendGraphClass& RootGraph = Document.RootGraph;
-	FMetasoundFrontendClassInput& GraphInput = RootGraph.Interface.Inputs[*Index];
+	FMetasoundFrontendClassInput& GraphInput = RootGraph.GetDefaultInterface().Inputs[*Index];
 
 	if (GraphInput.AccessType != AccessType)
 	{
@@ -4639,7 +4650,7 @@ bool FMetaSoundFrontendDocumentBuilder::SetGraphInputAccessType(FName InputName,
 		}
 
 #if WITH_EDITORONLY_DATA
-		RootGraph.Interface.UpdateChangeID();
+		RootGraph.GetDefaultInterface().UpdateChangeID();
 		Document.Metadata.ModifyContext.AddMemberIDModified(GraphInput.NodeID);
 #endif // WITH_EDITORONLY_DATA
 	}
@@ -4661,7 +4672,7 @@ bool FMetaSoundFrontendDocumentBuilder::SetGraphInputDataType(FName InputName, F
 
 		FMetasoundFrontendDocument& Document = GetDocumentChecked();
 		FMetasoundFrontendGraphClass& RootGraph = Document.RootGraph;
-		FMetasoundFrontendClassInput& GraphInput = RootGraph.Interface.Inputs[*Index];
+		FMetasoundFrontendClassInput& GraphInput = RootGraph.GetDefaultInterface().Inputs[*Index];
 		if (GraphInput.TypeName != DataType)
 		{
 			GraphInput.TypeName = DataType;
@@ -4711,7 +4722,7 @@ bool FMetaSoundFrontendDocumentBuilder::SetGraphInputDataType(FName InputName, F
 
 #if WITH_EDITORONLY_DATA
 			ClearMemberMetadata(GraphInput.NodeID);
-			RootGraph.Interface.UpdateChangeID();
+			RootGraph.GetDefaultInterface().UpdateChangeID();
 			Document.Metadata.ModifyContext.AddMemberIDModified(GraphInput.NodeID);
 			Document.Metadata.ModifyContext.AddNodeIDModified(GraphInput.NodeID);
 #endif // WITH_EDITORONLY_DATA
@@ -4727,7 +4738,7 @@ bool FMetaSoundFrontendDocumentBuilder::SetGraphInputDefault(FName InputName, FM
 
 	auto NameMatchesInput = [&InputName](const FMetasoundFrontendClassInput& Input) { return Input.Name == InputName; };
 	FMetasoundFrontendDocument& Document = GetDocumentChecked();
-	TArray<FMetasoundFrontendClassInput>& Inputs = Document.RootGraph.Interface.Inputs;
+	TArray<FMetasoundFrontendClassInput>& Inputs = Document.RootGraph.GetDefaultInterface().Inputs;
 
 	const int32 Index = Inputs.IndexOfByPredicate(NameMatchesInput);
 	if (Index != INDEX_NONE)
@@ -4770,7 +4781,7 @@ bool FMetaSoundFrontendDocumentBuilder::SetGraphInputDefaults(FName InputName, T
 
 	auto NameMatchesInput = [&InputName](const FMetasoundFrontendClassInput& Input) { return Input.Name == InputName; };
 	FMetasoundFrontendDocument& Document = GetDocumentChecked();
-	TArray<FMetasoundFrontendClassInput>& Inputs = Document.RootGraph.Interface.Inputs;
+	TArray<FMetasoundFrontendClassInput>& Inputs = Document.RootGraph.GetDefaultInterface().Inputs;
 
 	const int32 Index = Inputs.IndexOfByPredicate(NameMatchesInput);
 	if (Index != INDEX_NONE)
@@ -4806,7 +4817,7 @@ bool FMetaSoundFrontendDocumentBuilder::SetGraphInputDescription(FName InputName
 	if (const int32* Index = DocumentCache->GetInterfaceCache().FindInputIndex(InputName))
 	{
 		FMetasoundFrontendDocument& Document = GetDocumentChecked();
-		FMetasoundFrontendClassInput& GraphInput = Document.RootGraph.Interface.Inputs[*Index];
+		FMetasoundFrontendClassInput& GraphInput = Document.RootGraph.GetDefaultInterface().Inputs[*Index];
 		GraphInput.Metadata.SetDescription(Description);
 		Document.Metadata.ModifyContext.AddMemberIDModified(GraphInput.NodeID);
 		return true;
@@ -4822,7 +4833,7 @@ bool FMetaSoundFrontendDocumentBuilder::SetGraphInputDisplayName(FName InputName
 	if (const int32* Index = DocumentCache->GetInterfaceCache().FindInputIndex(InputName))
 	{
 		FMetasoundFrontendDocument& Document = GetDocumentChecked();
-		FMetasoundFrontendClassInput& GraphInput = Document.RootGraph.Interface.Inputs[*Index];
+		FMetasoundFrontendClassInput& GraphInput = Document.RootGraph.GetDefaultInterface().Inputs[*Index];
 		GraphInput.Metadata.SetDisplayName(DisplayName);
 		Document.Metadata.ModifyContext.AddMemberIDModified(GraphInput.NodeID);
 		return true;
@@ -4871,7 +4882,7 @@ bool FMetaSoundFrontendDocumentBuilder::SetGraphInputName(FName InputName, FName
 	FMetasoundFrontendDocument& Document = GetDocumentChecked();
 	FMetasoundFrontendGraphClass& RootGraph = Document.RootGraph;
 
-	FMetasoundFrontendClassInput& GraphInput = RootGraph.Interface.Inputs[*Index];
+	FMetasoundFrontendClassInput& GraphInput = RootGraph.GetDefaultInterface().Inputs[*Index];
 	GraphInput.Name = NewName;
 
 	RootGraph.IterateGraphPages([this, &GraphInput, &NewName](FMetasoundFrontendGraph& Graph)
@@ -4895,7 +4906,7 @@ bool FMetaSoundFrontendDocumentBuilder::SetGraphInputName(FName InputName, FName
 	DocumentDelegates->InterfaceDelegates.OnInputNameChanged.Broadcast(InputName, NewName);
 
 #if WITH_EDITORONLY_DATA
-	RootGraph.Interface.UpdateChangeID();
+	RootGraph.GetDefaultInterface().UpdateChangeID();
 	Document.Metadata.ModifyContext.AddMemberIDModified(GraphInput.NodeID);
 #endif // WITH_EDITORONLY_DATA
 
@@ -4910,7 +4921,7 @@ bool FMetaSoundFrontendDocumentBuilder::SetGraphOutputAdvancedDisplay(const FNam
 
 	if (const int32* Index = DocumentCache->GetInterfaceCache().FindOutputIndex(OutputName))
 	{
-		FMetasoundFrontendClassOutput& GraphOutput = Document.RootGraph.Interface.Outputs[*Index];
+		FMetasoundFrontendClassOutput& GraphOutput = Document.RootGraph.GetDefaultInterface().Outputs[*Index];
 		if (GraphOutput.Metadata.bIsAdvancedDisplay != InAdvancedDisplay)
 		{
 			GraphOutput.Metadata.SetIsAdvancedDisplay(InAdvancedDisplay);
@@ -4940,7 +4951,7 @@ bool FMetaSoundFrontendDocumentBuilder::SetGraphOutputAccessType(FName OutputNam
 
 	FMetasoundFrontendDocument& Document = GetDocumentChecked();
 	FMetasoundFrontendGraphClass& RootGraph = Document.RootGraph;
-	FMetasoundFrontendClassOutput& GraphOutput = RootGraph.Interface.Outputs[*Index];
+	FMetasoundFrontendClassOutput& GraphOutput = RootGraph.GetDefaultInterface().Outputs[*Index];
 	if (GraphOutput.AccessType != AccessType)
 	{
 		GraphOutput.AccessType = AccessType;
@@ -4976,7 +4987,7 @@ bool FMetaSoundFrontendDocumentBuilder::SetGraphOutputAccessType(FName OutputNam
 		}
 
 #if WITH_EDITORONLY_DATA
-		RootGraph.Interface.UpdateChangeID();
+		RootGraph.GetDefaultInterface().UpdateChangeID();
 		Document.Metadata.ModifyContext.AddMemberIDModified(GraphOutput.NodeID);
 #endif // WITH_EDITORONLY_DATA
 	}
@@ -5001,7 +5012,7 @@ bool FMetaSoundFrontendDocumentBuilder::SetGraphOutputDataType(FName OutputName,
 
 	FMetasoundFrontendDocument& Document = GetDocumentChecked();
 	FMetasoundFrontendGraphClass& RootGraph = Document.RootGraph;
-	FMetasoundFrontendClassOutput& GraphOutput = RootGraph.Interface.Outputs[*Index];
+	FMetasoundFrontendClassOutput& GraphOutput = RootGraph.GetDefaultInterface().Outputs[*Index];
 	if (GraphOutput.TypeName != DataType)
 	{
 		GraphOutput.TypeName = DataType;
@@ -5033,7 +5044,7 @@ bool FMetaSoundFrontendDocumentBuilder::SetGraphOutputDataType(FName OutputName,
 #endif // WITH_EDITOR
 
 #if WITH_EDITORONLY_DATA
-		RootGraph.Interface.UpdateChangeID();
+		RootGraph.GetDefaultInterface().UpdateChangeID();
 		ClearMemberMetadata(GraphOutput.NodeID);
 		Document.Metadata.ModifyContext.AddMemberIDModified(GraphOutput.NodeID);
 #endif // WITH_EDITORONLY_DATA
@@ -5050,7 +5061,7 @@ bool FMetaSoundFrontendDocumentBuilder::SetGraphOutputDescription(FName OutputNa
 	if (const int32* Index = DocumentCache->GetInterfaceCache().FindOutputIndex(OutputName))
 	{
 		FMetasoundFrontendDocument& Document = GetDocumentChecked();
-		FMetasoundFrontendClassOutput& GraphOutput = Document.RootGraph.Interface.Outputs[*Index];
+		FMetasoundFrontendClassOutput& GraphOutput = Document.RootGraph.GetDefaultInterface().Outputs[*Index];
 		GraphOutput.Metadata.SetDescription(Description);
 		Document.Metadata.ModifyContext.AddMemberIDModified(GraphOutput.NodeID);
 		return true;
@@ -5066,7 +5077,7 @@ bool FMetaSoundFrontendDocumentBuilder::SetGraphOutputDisplayName(FName OutputNa
 	if (const int32* Index = DocumentCache->GetInterfaceCache().FindOutputIndex(OutputName))
 	{
 		FMetasoundFrontendDocument& Document = GetDocumentChecked();
-		FMetasoundFrontendClassOutput& GraphOutput = Document.RootGraph.Interface.Outputs[*Index];
+		FMetasoundFrontendClassOutput& GraphOutput = Document.RootGraph.GetDefaultInterface().Outputs[*Index];
 		GraphOutput.Metadata.SetDisplayName(DisplayName);
 		Document.Metadata.ModifyContext.AddMemberIDModified(GraphOutput.NodeID);
 		return true;
@@ -5093,7 +5104,7 @@ bool FMetaSoundFrontendDocumentBuilder::SetGraphOutputName(FName OutputName, FNa
 
 	FMetasoundFrontendDocument& Document = GetDocumentChecked();
 	FMetasoundFrontendGraphClass& GraphClass = Document.RootGraph;
-	FMetasoundFrontendClassInterface& Interface = GraphClass.Interface;
+	FMetasoundFrontendClassInterface& Interface = GraphClass.GetDefaultInterface();
 	Interface.UpdateChangeID();
 
 	FMetasoundFrontendClassOutput& GraphOutput = Interface.Outputs[*Index];
@@ -5119,7 +5130,7 @@ bool FMetaSoundFrontendDocumentBuilder::SetGraphOutputName(FName OutputName, FNa
 	DocumentDelegates->InterfaceDelegates.OnOutputNameChanged.Broadcast(OutputName, NewName);
 	
 #if WITH_EDITORONLY_DATA
-	GraphClass.Interface.UpdateChangeID();
+	GraphClass.GetDefaultInterface().UpdateChangeID();
 	Document.Metadata.ModifyContext.AddMemberIDModified(GraphOutput.NodeID);
 #endif // WITH_EDITORONLY_DATA
 
diff --git a/Engine/Plugins/Runtime/Metasound/Source/MetasoundFrontend/Private/MetasoundFrontendDocumentCache.cpp b/Engine/Plugins/Runtime/Metasound/Source/MetasoundFrontend/Private/MetasoundFrontendDocumentCache.cpp
--- a/Engine/Plugins/Runtime/Metasound/Source/MetasoundFrontend/Private/MetasoundFrontendDocumentCache.cpp
+++ b/Engine/Plugins/Runtime/Metasound/Source/MetasoundFrontend/Private/MetasoundFrontendDocumentCache.cpp
@@ -340,13 +340,13 @@ namespace Metasound::Frontend
 		: Parent(ParentCache)
 	{
 		const FMetasoundFrontendGraphClass& GraphClass = Parent->GetDocument().RootGraph;
-		const TArray<FMetasoundFrontendClassInput>& Inputs = GraphClass.Interface.Inputs;
+		const TArray<FMetasoundFrontendClassInput>& Inputs = GraphClass.GetDefaultInterface().Inputs;
 		for (int32 Index = 0; Index < Inputs.Num(); ++Index)
 		{
 			InputNameToIndex.Add(Inputs[Index].Name, Index);
 		};
 
-		const TArray<FMetasoundFrontendClassOutput>& Outputs = GraphClass.Interface.Outputs;
+		const TArray<FMetasoundFrontendClassOutput>& Outputs = GraphClass.GetDefaultInterface().Outputs;
 		for (int32 Index = 0; Index < Outputs.Num(); ++Index)
 		{
 			OutputNameToIndex.Add(Outputs[Index].Name, Index);
@@ -368,7 +368,7 @@ namespace Metasound::Frontend
 		if (const int32* Index = InputNameToIndex.Find(InputName))
 		{
 			const FMetasoundFrontendDocument& Doc = Parent->GetDocument();
-			return &Doc.RootGraph.Interface.Inputs[*Index];
+			return &Doc.RootGraph.GetDefaultInterface().Inputs[*Index];
 		}
 
 		return nullptr;
@@ -384,7 +384,7 @@ namespace Metasound::Frontend
 		if (const int32* Index = OutputNameToIndex.Find(OutputName))
 		{
 			const FMetasoundFrontendDocument& Doc = Parent->GetDocument();
-			return &Doc.RootGraph.Interface.Outputs[*Index];
+			return &Doc.RootGraph.GetDefaultInterface().Outputs[*Index];
 		}
 
 		return nullptr;
@@ -421,7 +421,7 @@ namespace Metasound::Frontend
 	{
 		const FMetasoundFrontendDocument& Document = Parent->GetDocument();
 		const FMetasoundFrontendGraphClass& GraphClass = Document.RootGraph;
-		const FMetasoundFrontendClassInput& Input = GraphClass.Interface.Inputs[NewIndex];
+		const FMetasoundFrontendClassInput& Input = GraphClass.GetDefaultInterface().Inputs[NewIndex];
 		InputNameToIndex.Add(Input.Name, NewIndex);
 
 		++TransactionCount;
@@ -453,7 +453,7 @@ namespace Metasound::Frontend
 	{
 		const FMetasoundFrontendDocument& Document = Parent->GetDocument();
 		const FMetasoundFrontendGraphClass& GraphClass = Document.RootGraph;
-		const FMetasoundFrontendClassOutput& Output = GraphClass.Interface.Outputs[NewIndex];
+		const FMetasoundFrontendClassOutput& Output = GraphClass.GetDefaultInterface().Outputs[NewIndex];
 		OutputNameToIndex.Add(Output.Name, NewIndex);
 
 		++TransactionCount;
@@ -480,7 +480,7 @@ namespace Metasound::Frontend
 	{
 		const FMetasoundFrontendDocument& Document = Parent->GetDocument();
 		const FMetasoundFrontendGraphClass& GraphClass = Document.RootGraph;
-		const FMetasoundFrontendClassInput& Input = GraphClass.Interface.Inputs[IndexBeingRemoved];
+		const FMetasoundFrontendClassInput& Input = GraphClass.GetDefaultInterface().Inputs[IndexBeingRemoved];
 		InputNameToIndex.Remove(Input.Name);
 
 		++TransactionCount;
@@ -490,7 +490,7 @@ namespace Metasound::Frontend
 	{
 		const FMetasoundFrontendDocument& Document = Parent->GetDocument();
 		const FMetasoundFrontendGraphClass& GraphClass = Document.RootGraph;
-		const FMetasoundFrontendClassOutput& Output = GraphClass.Interface.Outputs[IndexBeingRemoved];
+		const FMetasoundFrontendClassOutput& Output = GraphClass.GetDefaultInterface().Outputs[IndexBeingRemoved];
 		OutputNameToIndex.Remove(Output.Name);
 
 		++TransactionCount;
diff --git a/Engine/Plugins/Runtime/Metasound/Source/MetasoundFrontend/Private/MetasoundFrontendDocumentController.cpp b/Engine/Plugins/Runtime/Metasound/Source/MetasoundFrontend/Private/MetasoundFrontendDocumentController.cpp
--- a/Engine/Plugins/Runtime/Metasound/Source/MetasoundFrontend/Private/MetasoundFrontendDocumentController.cpp
+++ b/Engine/Plugins/Runtime/Metasound/Source/MetasoundFrontend/Private/MetasoundFrontendDocumentController.cpp
@@ -460,8 +460,8 @@ namespace Metasound
 					else
 					{
 						// if IDs and version numbers are equal, sort by number of inputs & outputs descending
-						const int32 NumLHSVertices = InLHS.Interface.Inputs.Num() + InLHS.Interface.Outputs.Num();
-						const int32 NumRHSVertices = InRHS.Interface.Inputs.Num() + InRHS.Interface.Outputs.Num();
+						const int32 NumLHSVertices = InLHS.GetDefaultInterface().Inputs.Num() + InLHS.GetDefaultInterface().Outputs.Num();
+						const int32 NumRHSVertices = InRHS.GetDefaultInterface().Inputs.Num() + InRHS.GetDefaultInterface().Outputs.Num();
 
 						return NumLHSVertices > NumRHSVertices;
 					}
diff --git a/Engine/Plugins/Runtime/Metasound/Source/MetasoundFrontend/Private/MetasoundFrontendDocumentVersioning.cpp b/Engine/Plugins/Runtime/Metasound/Source/MetasoundFrontend/Private/MetasoundFrontendDocumentVersioning.cpp
--- a/Engine/Plugins/Runtime/Metasound/Source/MetasoundFrontend/Private/MetasoundFrontendDocumentVersioning.cpp
+++ b/Engine/Plugins/Runtime/Metasound/Source/MetasoundFrontend/Private/MetasoundFrontendDocumentVersioning.cpp
@@ -50,11 +50,16 @@ PRAGMA_ENABLE_DEPRECATION_WARNINGS
 				};
 
 				FMetasoundFrontendDocument& Document = GetDocumentUnsafe(OutBuilder);
+				// For all class definitions we are going to access the default interface instead of inspecting the
+				// interface override. This is safe here because the class interface override did not exist in this
+				// version of the document. 
+				checkf(Document.Metadata.Version.Number <= FMetasoundFrontendVersionNumber(1, 14), TEXT("Migration of page properties needs to happen before the introduction of node configuration to the document"));
+
 				FMetasoundFrontendGraphClass& GraphClass = Document.RootGraph;
-				MigrateInterfaceInputDefaults(GraphClass.Interface);
+				MigrateInterfaceInputDefaults(GraphClass.GetDefaultInterface());
 				for (FMetasoundFrontendClass& Dependency : Document.Dependencies)
 				{
-					MigrateInterfaceInputDefaults(Dependency.Interface);
+					MigrateInterfaceInputDefaults(Dependency.GetDefaultInterface());
 				}
 
 				struct FMigratePageGraphs : public FMetasoundFrontendGraphClass::IPropertyVersionTransform
@@ -485,6 +490,11 @@ PRAGMA_ENABLE_DEPRECATION_WARNINGS
 			void TransformInternal(FDocumentHandle InDocument) const override
 			{
 	#if WITH_EDITOR
+				// For all class definitions we are going to access the default interface instead of inspecting the
+				// interface override. This is safe here because the class interface override did not exist in this
+				// version of the document. 
+				checkf(InDocument->GetMetadata()->Version.Number <= FMetasoundFrontendVersionNumber(1, 14), TEXT("Migration of page properties needs to happen before the introduction of node configuration to the document"));
+
 				// Do not serialize MetaData text for dependencies as
 				// CacheRegistryData dynamically provides this.
 				InDocument->IterateDependencies([](FMetasoundFrontendClass& Dependency)
@@ -492,12 +502,12 @@ PRAGMA_ENABLE_DEPRECATION_WARNINGS
 					constexpr bool bSerializeText = false;
 					Dependency.Metadata.SetSerializeText(bSerializeText);
 
-					for (FMetasoundFrontendClassInput& Input : Dependency.Interface.Inputs)
+					for (FMetasoundFrontendClassInput& Input : Dependency.GetDefaultInterface().Inputs)
 					{
 						Input.Metadata.SetSerializeText(false);
 					}
 
-					for (FMetasoundFrontendClassOutput& Output : Dependency.Interface.Outputs)
+					for (FMetasoundFrontendClassOutput& Output : Dependency.GetDefaultInterface().Outputs)
 					{
 						Output.Metadata.SetSerializeText(false);
 					}
@@ -530,14 +540,14 @@ PRAGMA_ENABLE_DEPRECATION_WARNINGS
 
 				// Only serialize MetaData text for inputs owned by the graph (not by interfaces)
 				FMetasoundFrontendGraphClass RootGraphClass = InDocument->GetRootGraphClass();
-				for (FMetasoundFrontendClassInput& Input : RootGraphClass.Interface.Inputs)
+				for (FMetasoundFrontendClassInput& Input : RootGraphClass.GetDefaultInterface().Inputs)
 				{
 					const bool bSerializeText = !InterfaceInputs.Contains(FNameDataTypePair(Input.Name, Input.TypeName));
 					Input.Metadata.SetSerializeText(bSerializeText);
 				}
 
 				// Only serialize MetaData text for outputs owned by the graph (not by interfaces)
-				for (FMetasoundFrontendClassOutput& Output : RootGraphClass.Interface.Outputs)
+				for (FMetasoundFrontendClassOutput& Output : RootGraphClass.GetDefaultInterface().Outputs)
 				{
 					const bool bSerializeText = !InterfaceOutputs.Contains(FNameDataTypePair(Output.Name, Output.TypeName));
 					Output.Metadata.SetSerializeText(bSerializeText);
@@ -704,7 +714,13 @@ PRAGMA_ENABLE_DEPRECATION_WARNINGS
 				// to a MetaSound Asset by removing any extraneous default data (early values in the
 				// array were stale).
 				FMetasoundFrontendDocument& Document = const_cast<FMetasoundFrontendDocument&>(OutBuilder.GetConstDocumentChecked());
-				for (FMetasoundFrontendClassInput& Input : Document.RootGraph.Interface.Inputs)
+
+				// For all class definitions we are going to access the default interface instead of inspecting the
+				// interface override. This is safe here because the class interface override did not exist in this
+				// version of the document. 
+				checkf(Document.Metadata.Version.Number <= FMetasoundFrontendVersionNumber(1, 14), TEXT("Migration of page properties needs to happen before the introduction of node configuration to the document"));
+
+				for (FMetasoundFrontendClassInput& Input : Document.RootGraph.GetDefaultInterface().Inputs)
 				{
 					int32 PageIDIndex = INDEX_NONE;
 					TArray<FMetasoundFrontendClassInputDefault>& Defaults = const_cast<TArray<FMetasoundFrontendClassInputDefault>&>(Input.GetDefaults());
@@ -804,4 +820,4 @@ PRAGMA_ENABLE_DEPRECATION_WARNINGS
 	}
 #undef METASOUND_VERSIONING_LOG
 } // namespace Metasound::Frontend
-#endif // WITH_EDITORONLY_DATA
\ No newline at end of file
+#endif // WITH_EDITORONLY_DATA
diff --git a/Engine/Plugins/Runtime/Metasound/Source/MetasoundFrontend/Private/MetasoundFrontendGraphBuilder.cpp b/Engine/Plugins/Runtime/Metasound/Source/MetasoundFrontend/Private/MetasoundFrontendGraphBuilder.cpp
--- a/Engine/Plugins/Runtime/Metasound/Source/MetasoundFrontend/Private/MetasoundFrontendGraphBuilder.cpp
+++ b/Engine/Plugins/Runtime/Metasound/Source/MetasoundFrontend/Private/MetasoundFrontendGraphBuilder.cpp
@@ -64,6 +64,11 @@ namespace Metasound::Frontend
 			return Frontend::DefaultPageID;
 		}
 
+		bool InterfacesHaveEqualSize(const FMetasoundFrontendClassInterface& InClassInterface, const FMetasoundFrontendNodeInterface& InNodeInterface)
+		{
+			return (InClassInterface.Inputs.Num() == InNodeInterface.Inputs.Num()) && (InClassInterface.Outputs.Num() == InNodeInterface.Outputs.Num());
+		}
+
 		const FMetasoundFrontendLiteral* FindLiteralForInputVertex(const FVertexName& InVertexName, const FMetasoundFrontendNode& InNode, const FMetasoundFrontendClassInput& InNodeClassInput)
 		{
 			using namespace Frontend;
@@ -106,7 +111,7 @@ namespace Metasound::Frontend
 			return Literal;
 		}
 
-		const FMetasoundFrontendLiteral* FindInputLiteralForInputNode(const FMetasoundFrontendNode& InInputNode, const FMetasoundFrontendClass& InInputNodeClass, const FMetasoundFrontendClassInput& InOwningGraphClassInput)
+		const FMetasoundFrontendLiteral* FindInputLiteralForInputNode(const FMetasoundFrontendNode& InInputNode, const FMetasoundFrontendClassInterface& InInputNodeClassInterface, const FMetasoundFrontendClassInput& InOwningGraphClassInput)
 		{
 			using namespace Frontend;
 
@@ -149,9 +154,9 @@ namespace Metasound::Frontend
 			}
 
 			// Check for default value on input node class
-			if (nullptr == Literal && ensure(InInputNodeClass.Interface.Inputs.Num() == 1))
+			if (nullptr == Literal && ensure(InInputNodeClassInterface.Inputs.Num() == 1))
 			{
-				const FMetasoundFrontendClassInput& InputNodeClassInput = InInputNodeClass.Interface.Inputs.Last();
+				const FMetasoundFrontendClassInput& InputNodeClassInput = InInputNodeClassInterface.Inputs.Last();
 				const FGuid PageID = GraphBuilderPrivate::ResolveTargetPageID(InputNodeClassInput);
 				const FMetasoundFrontendLiteral& DefaultLiteral = InputNodeClassInput.FindConstDefaultChecked(PageID);
 				if (DefaultLiteral.IsValid())
@@ -185,13 +190,9 @@ namespace Metasound::Frontend
 			}
 		}
 
-		FOutputVertexInterface CreateOutputVertexInterface(const FBuildContext& InContext, const FMetasoundFrontendNode& InNode, const FMetasoundFrontendClass& InClass, const FMetasoundFrontendClassInput* InOwningGraphClassInput, const FMetasoundFrontendClassOutput* InOwningGraphClassOutput)
+		FOutputVertexInterface CreateOutputVertexInterface(const FBuildContext& InContext, const FMetasoundFrontendNode& InNode, const FMetasoundFrontendClass& InClass, const FMetasoundFrontendClassInterface& InClassInterface, const FMetasoundFrontendClassInput* InOwningGraphClassInput, const FMetasoundFrontendClassOutput* InOwningGraphClassOutput)
 		{
-			// This exists here until we bring sub interface functionality to the frontend. 
-			// This check will no longer be valid at that point. We will have to implement
-			// a way to lookup the class's vertex access type info from the class def, or 
-			// move that info into the node's def.
-			check(InClass.Interface.Outputs.Num() == InNode.Interface.Outputs.Num());
+			check(InClassInterface.Outputs.Num() == InNode.Interface.Outputs.Num());
 
 			// TODO: migrate sub interface builders from MetasoundVertex.cpp when
 			// sub interfaces exist on the frontend document
@@ -214,7 +215,7 @@ namespace Metasound::Frontend
 								if (InNode.Interface.Outputs.Num() == 1)
 								{
 									const FMetasoundFrontendVertex& NodeOutputVertex = InNode.Interface.Outputs[0];
-									const FMetasoundFrontendClassOutput& ClassOutputVertex = InClass.Interface.Outputs[0];
+									const FMetasoundFrontendClassOutput& ClassOutputVertex = InClassInterface.Outputs[0];
 									// Note: stripping metadata here. If we make this function public, we may want to keep metadata if WITH_EDITORONLY_DATA
 									Vertices.Add(FOutputDataVertex{
 												NodeOutputVertex.Name, 
@@ -245,7 +246,7 @@ namespace Metasound::Frontend
 							{
 								check(InNode.Interface.Outputs.Num() == 1);
 								const FMetasoundFrontendVertex& NodeOutputVertex = InNode.Interface.Outputs[0];
-								const FMetasoundFrontendClassOutput& ClassOutputVertex = InClass.Interface.Outputs[0];
+								const FMetasoundFrontendClassOutput& ClassOutputVertex = InClassInterface.Outputs[0];
 								// Note: stripping metadata here. If we make this function public, we may want to keep metadata if WITH_EDITORONLY_DATA
 								Vertices.Add(FOutputDataVertex{
 											NodeOutputVertex.Name, 
@@ -264,7 +265,7 @@ namespace Metasound::Frontend
 							// As-it-stands, node interfaces and class interfaces
 							// match for all other nodes. This will change when
 							// node config is done as a feature. 
-							for (const FMetasoundFrontendClassOutput& ClassOutput : InClass.Interface.Outputs)
+							for (const FMetasoundFrontendClassOutput& ClassOutput : InClassInterface.Outputs)
 							{
 								// Note: stripping metadata here. If we make this function public, we may want to keep metadata if WITH_EDITORONLY_DATA
 								Vertices.Add(FOutputDataVertex{ClassOutput.Name, ClassOutput.TypeName, FDataVertexMetadata{}, FrontendVertexAccessTypeToCoreVertexAccessType(ClassOutput.AccessType)});
@@ -281,13 +282,9 @@ namespace Metasound::Frontend
 			}
 		}
 
-		FInputVertexInterface CreateInputVertexInterface(const FBuildContext& InContext, const FMetasoundFrontendNode& InNode, const FMetasoundFrontendClass& InClass, const FMetasoundFrontendClassInput* InOwningGraphClassInput, const FMetasoundFrontendClassOutput* InOwningGraphClassOutput)
+		FInputVertexInterface CreateInputVertexInterface(const FBuildContext& InContext, const FMetasoundFrontendNode& InNode, const FMetasoundFrontendClass& InClass, const FMetasoundFrontendClassInterface& InClassInterface, const FMetasoundFrontendClassInput* InOwningGraphClassInput, const FMetasoundFrontendClassOutput* InOwningGraphClassOutput)
 		{
-			// This exists here until we bring sub interface functionality to the frontend. 
-			// This check will no longer be valid at that point. We will have to implement
-			// a way to lookup the class's vertex access type info from the class def, or 
-			// move that info into the node's def.
-			check(InClass.Interface.Inputs.Num() == InNode.Interface.Inputs.Num());
+			check(InClassInterface.Inputs.Num() == InNode.Interface.Inputs.Num());
 
 			// TODO: migrate sub interface builders from MetasoundVertex.cpp when
 			// sub interfaces exist on the frontend document
@@ -324,8 +321,8 @@ namespace Metasound::Frontend
 							{
 								check(InNode.Interface.Inputs.Num() == 1);
 								const FMetasoundFrontendVertex& NodeInputVertex = InNode.Interface.Inputs[0];
-								const FMetasoundFrontendClassInput& ClassInputVertex = InClass.Interface.Inputs[0];
-								const FMetasoundFrontendLiteral* FrontendLiteral = FindInputLiteralForInputNode(InNode, InClass, *InOwningGraphClassInput);
+								const FMetasoundFrontendClassInput& ClassInputVertex = InClassInterface.Inputs[0];
+								const FMetasoundFrontendLiteral* FrontendLiteral = FindInputLiteralForInputNode(InNode, InClassInterface, *InOwningGraphClassInput);
 								
 								// NOTE: The access type reported on the ClassInputVertex may be incorrect. When the access type is changed in the editor,
 								// the access type of the input node's FMetasoundFrontendClass is not updated. 
@@ -340,7 +337,7 @@ namespace Metasound::Frontend
 								if (InNode.Interface.Inputs.Num() == 1)
 								{
 									const FMetasoundFrontendVertex& NodeInputVertex = InNode.Interface.Inputs[0];
-									const FMetasoundFrontendClassInput& ClassInputVertex = InClass.Interface.Inputs[0];
+									const FMetasoundFrontendClassInput& ClassInputVertex = InClassInterface.Inputs[0];
 									const FMetasoundFrontendLiteral* FrontendLiteral = FindLiteralForInputVertex(NodeInputVertex.Name, InNode, ClassInputVertex);
 									// NOTE: The access type reported on the ClassInputVertex may be incorrect. When the access type is changed in the editor,
 									// the access type of the output node's FMetasoundFrontendClass is not updated. 
@@ -365,7 +362,7 @@ namespace Metasound::Frontend
 							// As-it-stands, node interfaces and class interfaces
 							// match for all other nodes. This will change when
 							// node config is done as a feature. 
-							for (const FMetasoundFrontendClassInput& ClassInput : InClass.Interface.Inputs)
+							for (const FMetasoundFrontendClassInput& ClassInput : InClassInterface.Inputs)
 							{
 								const FMetasoundFrontendLiteral* FrontendLiteral = FindLiteralForInputVertex(ClassInput.Name, InNode, ClassInput);
 								Vertices.Add(MakeVertex(ClassInput.Name, ClassInput.TypeName, ClassInput.AccessType, FrontendLiteral));
@@ -382,22 +379,22 @@ namespace Metasound::Frontend
 			}
 		}
 
-		FVertexInterface CreateVertexInterface(const FBuildContext& InContext, const FMetasoundFrontendNode& InNode, const FMetasoundFrontendClass& InClass, const FMetasoundFrontendClassInput* InOwningGraphClassInput, const FMetasoundFrontendClassOutput* InOwningGraphClassOutput)
+		FVertexInterface CreateVertexInterface(const FBuildContext& InContext, const FMetasoundFrontendNode& InNode, const FMetasoundFrontendClass& InClass, const FMetasoundFrontendClassInterface& InClassInterface, const FMetasoundFrontendClassInput* InOwningGraphClassInput, const FMetasoundFrontendClassOutput* InOwningGraphClassOutput)
 		{
 			return FVertexInterface(
-				CreateInputVertexInterface(InContext, InNode, InClass, InOwningGraphClassInput, InOwningGraphClassOutput), 
-				CreateOutputVertexInterface(InContext, InNode, InClass, InOwningGraphClassInput, InOwningGraphClassOutput), 
+				CreateInputVertexInterface(InContext, InNode, InClass, InClassInterface, InOwningGraphClassInput, InOwningGraphClassOutput), 
+				CreateOutputVertexInterface(InContext, InNode, InClass, InClassInterface, InOwningGraphClassInput, InOwningGraphClassOutput), 
 				CreateEnvironmentVertexInterface(InNode.Interface));
 		}
 
-		FNodeData CreateNodeData(const FBuildContext& InContext, const FMetasoundFrontendNode& InNode, const FMetasoundFrontendClass& InClass, const FMetasoundFrontendClassInput* InOwningGraphClassInput=nullptr, const FMetasoundFrontendClassOutput* InOwningGraphClassOutput=nullptr)
+		FNodeData CreateNodeData(const FBuildContext& InContext, const FMetasoundFrontendNode& InNode, const FMetasoundFrontendClass& InClass, const FMetasoundFrontendClassInterface& InClassInterface, const FMetasoundFrontendClassInput* InOwningGraphClassInput=nullptr, const FMetasoundFrontendClassOutput* InOwningGraphClassOutput=nullptr)
 		{
 			TSharedPtr<const IOperatorData> Config;
 			if (InNode.Configuration.IsValid())
 			{
 				Config = InNode.Configuration.Get().GetOperatorData();
 			}
-			return FNodeData(InNode.Name, InNode.GetID(), CreateVertexInterface(InContext, InNode, InClass, InOwningGraphClassInput, InOwningGraphClassOutput), MoveTemp(Config));
+			return FNodeData(InNode.Name, InNode.GetID(), CreateVertexInterface(InContext, InNode, InClass, InClassInterface, InOwningGraphClassInput, InOwningGraphClassOutput), MoveTemp(Config));
 		}
 
 		const FMetasoundFrontendVariable* FindVariableForVariableNode(const FMetasoundFrontendNode& InVariableNode, const FMetasoundFrontendGraph& InGraph)
@@ -406,22 +403,11 @@ namespace Metasound::Frontend
 			return InGraph.Variables.FindByPredicate([&](const FMetasoundFrontendVariable& InVar) { return InVar.VariableNodeID == DesiredID; });
 		}
 
-		TUniquePtr<INode> CreateVariableNode(const FBuildContext& InContext, const FMetasoundFrontendNode& InNode, const FMetasoundFrontendClass& InClass, const FMetasoundFrontendGraph& InGraph)
+		TUniquePtr<INode> CreateVariableNode(const FBuildContext& InContext, const FMetasoundFrontendNode& InNode, const FMetasoundFrontendClass& InClass, const FMetasoundFrontendClassInterface& InClassInterface, const FMetasoundFrontendGraph& InGraph)
 		{
 			check(InClass.Metadata.GetType() == EMetasoundFrontendClassType::Variable);
 			check(InNode.ClassID == InClass.ID);
-
-			// Check frontend node and class interfaces are compatible
-			if (InClass.Interface.Inputs.Num() != InNode.Interface.Inputs.Num())
-			{
-				UE_LOG(LogMetaSound, Error, TEXT("MetaSound %s contains mismatched number of inputs (%d / %d) on input node %s class %s"), *InContext.DebugAssetName, InNode.Interface.Inputs.Num(), InClass.Interface.Inputs.Num(), *InNode.Name.ToString(), *InClass.Metadata.GetClassName().ToString());
-				return TUniquePtr<INode>(nullptr);
-			}
-			else if (InClass.Interface.Outputs.Num() != InNode.Interface.Outputs.Num())
-			{
-				UE_LOG(LogMetaSound, Error, TEXT("MetaSound %s contains mismatched number of outputs (%d / %d) on output node %s class %s"), *InContext.DebugAssetName, InNode.Interface.Outputs.Num(), InClass.Interface.Outputs.Num(), *InNode.Name.ToString(), *InClass.Metadata.GetClassName().ToString());
-				return TUniquePtr<INode>(nullptr);
-			}
+			check(InterfacesHaveEqualSize(InClassInterface, InNode.Interface));
 
 			// Find the variable object associated with the node.
 			const FMetasoundFrontendVariable* FrontendVariable = FindVariableForVariableNode(InNode, InGraph);
@@ -434,7 +420,7 @@ namespace Metasound::Frontend
 				if (IsLiteralParsableByDataType)
 				{
 					FLiteral Literal = FrontendVariable->Literal.ToLiteral(FrontendVariable->TypeName, &InContext.DataTypeRegistry, &InContext.ProxyDataCache);
-					return DataTypeRegistry.CreateVariableNode(FrontendVariable->TypeName, MoveTemp(Literal), CreateNodeData(InContext, InNode, InClass));
+					return DataTypeRegistry.CreateVariableNode(FrontendVariable->TypeName, MoveTemp(Literal), CreateNodeData(InContext, InNode, InClass, InClassInterface));
 				}
 				else
 				{
@@ -450,10 +436,11 @@ namespace Metasound::Frontend
 		}
 
 
-		TUniquePtr<INode> CreateInputNode(const FBuildContext& InContext, const FMetasoundFrontendNode& InNode, const FMetasoundFrontendClass& InClass, const FMetasoundFrontendClassInput& InOwningGraphClassInput)
+		TUniquePtr<INode> CreateInputNode(const FBuildContext& InContext, const FMetasoundFrontendNode& InNode, const FMetasoundFrontendClass& InClass, const FMetasoundFrontendClassInterface& InClassInterface, const FMetasoundFrontendClassInput& InOwningGraphClassInput)
 		{
 			check(InClass.Metadata.GetType() == EMetasoundFrontendClassType::Input);
 			check(InNode.ClassID == InClass.ID);
+			check(InterfacesHaveEqualSize(InClassInterface, InNode.Interface));
 
 			// Check that the frontend node and class is correct.
 			if (InNode.Interface.Inputs.Num() != 1)
@@ -461,27 +448,17 @@ namespace Metasound::Frontend
 				UE_LOG(LogMetaSound, Error, TEXT("MetaSound %s contains invalid number of inputs (%d) on input node %s"), *InContext.DebugAssetName, InNode.Interface.Inputs.Num(), *InNode.Name.ToString());
 				return TUniquePtr<INode>(nullptr);
 			}
-			else if (InClass.Interface.Inputs.Num() != 1)
-			{
-				UE_LOG(LogMetaSound, Error, TEXT("MetaSound %s contains invalid number of inputs (%d) on input class %s"), *InContext.DebugAssetName, InClass.Interface.Inputs.Num(), *InClass.Metadata.GetClassName().ToString());
-				return TUniquePtr<INode>(nullptr);
-			}
-			else if (InClass.Interface.Outputs.Num() != InNode.Interface.Outputs.Num())
-			{
-				UE_LOG(LogMetaSound, Error, TEXT("MetaSound %s contains mismatched number of outputs (%d / %d) on output node %s class %s"), *InContext.DebugAssetName, InNode.Interface.Outputs.Num(), InClass.Interface.Outputs.Num(), *InNode.Name.ToString(), *InClass.Metadata.GetClassName().ToString());
-				return TUniquePtr<INode>(nullptr);
-			}
 
 			// Create the input node
 			const FMetasoundFrontendVertex& InputVertex = InNode.Interface.Inputs[0];
 
 			if (InOwningGraphClassInput.AccessType == EMetasoundFrontendVertexAccessType::Reference)
 			{
-				return InContext.DataTypeRegistry.CreateInputNode(InputVertex.TypeName, CreateNodeData(InContext, InNode, InClass, &InOwningGraphClassInput));
+				return InContext.DataTypeRegistry.CreateInputNode(InputVertex.TypeName, CreateNodeData(InContext, InNode, InClass, InClassInterface, &InOwningGraphClassInput));
 			}
 			else if (InOwningGraphClassInput.AccessType == EMetasoundFrontendVertexAccessType::Value)
 			{
-				return InContext.DataTypeRegistry.CreateConstructorInputNode(InputVertex.TypeName, CreateNodeData(InContext, InNode, InClass, &InOwningGraphClassInput));
+				return InContext.DataTypeRegistry.CreateConstructorInputNode(InputVertex.TypeName, CreateNodeData(InContext, InNode, InClass, InClassInterface, &InOwningGraphClassInput));
 			}
 			else
 			{
@@ -491,36 +468,27 @@ namespace Metasound::Frontend
 			return TUniquePtr<INode>(nullptr);
 		}
 
-		TUniquePtr<INode> CreateOutputNode(const FBuildContext& InContext, const FMetasoundFrontendNode& InNode, const FMetasoundFrontendClass& InClass, const FMetasoundFrontendClassOutput& InOwningGraphClassOutput)
+		TUniquePtr<INode> CreateOutputNode(const FBuildContext& InContext, const FMetasoundFrontendNode& InNode, const FMetasoundFrontendClass& InClass, const FMetasoundFrontendClassInterface& InClassInterface, const FMetasoundFrontendClassOutput& InOwningGraphClassOutput)
 		{
 			check(InClass.Metadata.GetType() == EMetasoundFrontendClassType::Output);
 			check(InNode.ClassID == InClass.ID);
+			check(InterfacesHaveEqualSize(InClassInterface, InNode.Interface));
 
-			if (InClass.Interface.Inputs.Num() != InNode.Interface.Inputs.Num())
-			{
-				UE_LOG(LogMetaSound, Error, TEXT("MetaSound %s contains mismatched number of inputs (%d / %d) on input node %s class %s"), *InContext.DebugAssetName, InNode.Interface.Inputs.Num(), InClass.Interface.Inputs.Num(), *InNode.Name.ToString(), *InClass.Metadata.GetClassName().ToString());
-				return TUniquePtr<INode>(nullptr);
-			}
-			else if (InNode.Interface.Outputs.Num() != 1)
+			if (InNode.Interface.Outputs.Num() != 1)
 			{
 				UE_LOG(LogMetaSound, Error, TEXT("MetaSound %s contains invalid number of outputs (%d) on output node %s"), *InContext.DebugAssetName, InNode.Interface.Outputs.Num(), *InNode.Name.ToString());
 				return TUniquePtr<INode>(nullptr);
 			}
-			else if (InClass.Interface.Outputs.Num() != 1)
-			{
-				UE_LOG(LogMetaSound, Error, TEXT("MetaSound %s contains invalid number of outputs (%d) on output class %s"), *InContext.DebugAssetName, InClass.Interface.Outputs.Num(), *InClass.Metadata.GetClassName().ToString());
-				return TUniquePtr<INode>(nullptr);
-			}
 
 			const FMetasoundFrontendVertex& OutputVertex = InNode.Interface.Outputs[0];
 
-			if (InClass.Interface.Outputs[0].AccessType == EMetasoundFrontendVertexAccessType::Reference)
+			if (InClassInterface.Outputs[0].AccessType == EMetasoundFrontendVertexAccessType::Reference)
 			{
-				return InContext.DataTypeRegistry.CreateOutputNode(OutputVertex.TypeName, CreateNodeData(InContext, InNode, InClass, nullptr, &InOwningGraphClassOutput));
+				return InContext.DataTypeRegistry.CreateOutputNode(OutputVertex.TypeName, CreateNodeData(InContext, InNode, InClass, InClassInterface, nullptr, &InOwningGraphClassOutput));
 			}
-			else if (InClass.Interface.Outputs[0].AccessType == EMetasoundFrontendVertexAccessType::Value)
+			else if (InClassInterface.Outputs[0].AccessType == EMetasoundFrontendVertexAccessType::Value)
 			{
-				return InContext.DataTypeRegistry.CreateConstructorOutputNode(OutputVertex.TypeName, CreateNodeData(InContext, InNode, InClass, nullptr, &InOwningGraphClassOutput));
+				return InContext.DataTypeRegistry.CreateConstructorOutputNode(OutputVertex.TypeName, CreateNodeData(InContext, InNode, InClass, InClassInterface, nullptr, &InOwningGraphClassOutput));
 			}
 			else
 			{
@@ -530,27 +498,15 @@ namespace Metasound::Frontend
 			return TUniquePtr<INode>(nullptr);
 		}
 
-		TUniquePtr<INode> CreateExternalNode(const FBuildContext& InContext, const FMetasoundFrontendNode& InNode, const FMetasoundFrontendClass& InClass)
+		TUniquePtr<INode> CreateExternalNode(const FBuildContext& InContext, const FMetasoundFrontendNode& InNode, const FMetasoundFrontendClass& InClass, const FMetasoundFrontendClassInterface& InClassInterface)
 		{
 			using namespace Frontend;
 
 			check(InNode.ClassID == InClass.ID);
-
-			// Check frontend node and class interfaces are compatible
-			if (InClass.Interface.Inputs.Num() != InNode.Interface.Inputs.Num())
-			{
-				UE_LOG(LogMetaSound, Error, TEXT("MetaSound %s contains mismatched number of inputs (%d / %d) on input node %s class %s"), *InContext.DebugAssetName, InNode.Interface.Inputs.Num(), InClass.Interface.Inputs.Num(), *InNode.Name.ToString(), *InClass.Metadata.GetClassName().ToString());
-				return TUniquePtr<INode>(nullptr);
-			}
-			else if (InClass.Interface.Outputs.Num() != InNode.Interface.Outputs.Num())
-			{
-				UE_LOG(LogMetaSound, Error, TEXT("MetaSound %s contains mismatched number of outputs (%d / %d) on output node %s class %s"), *InContext.DebugAssetName, InNode.Interface.Outputs.Num(), InClass.Interface.Outputs.Num(), *InNode.Name.ToString(), *InClass.Metadata.GetClassName().ToString());
-				return TUniquePtr<INode>(nullptr);
-			}
-
+			check(InterfacesHaveEqualSize(InClassInterface, InNode.Interface));
 
 			const FNodeRegistryKey Key = FNodeRegistryKey(InClass.Metadata);
-			return INodeClassRegistry::Get()->CreateNode(Key, CreateNodeData(InContext, InNode, InClass));
+			return INodeClassRegistry::Get()->CreateNode(Key, CreateNodeData(InContext, InNode, InClass, InClassInterface));
 		}
 
 		const FMetasoundFrontendClassInput* FindClassInputForInputNode(const FMetasoundFrontendGraphClass& InOwningGraph, const FMetasoundFrontendNode& InInputNode)
@@ -565,7 +521,7 @@ namespace Metasound::Frontend
 					return (InInputNode.GetID() == GraphInput.NodeID);
 				};
 
-				return InOwningGraph.Interface.Inputs.FindByPredicate(IsMatchingInput);
+				return InOwningGraph.GetDefaultInterface().Inputs.FindByPredicate(IsMatchingInput);
 			}
 			return nullptr;
 		}
@@ -582,7 +538,7 @@ namespace Metasound::Frontend
 					return (InOutputNode.GetID() == GraphOutput.NodeID);
 				};
 
-				return InOwningGraph.Interface.Outputs.FindByPredicate(IsMatchingOutput);
+				return InOwningGraph.GetDefaultInterface().Outputs.FindByPredicate(IsMatchingOutput);
 			}
 			return nullptr;
 		}
@@ -600,6 +556,21 @@ namespace Metasound::Frontend
 
 				if (ensure(nullptr != NodeClass))
 				{
+					const FMetasoundFrontendClassInterface& ClassInterface = NodeClass->GetInterfaceForNode(Node);
+
+					// Check frontend node and class interfaces are compatible
+					if (ClassInterface.Inputs.Num() != Node.Interface.Inputs.Num())
+					{
+						UE_LOG(LogMetaSound, Error, TEXT("Skipping node. MetaSound %s contains mismatched number of inputs (%d / %d) on node %s class %s"), *InGraphContext.BuildContext.DebugAssetName, Node.Interface.Inputs.Num(), ClassInterface.Inputs.Num(), *Node.Name.ToString(), *NodeClass->Metadata.GetClassName().ToString());
+						continue;
+					}
+					else if (ClassInterface.Outputs.Num() != Node.Interface.Outputs.Num())
+					{
+						UE_LOG(LogMetaSound, Error, TEXT("Skipping node. MetaSound %s contains mismatched number of outputs (%d / %d) on node %s class %s"), *InGraphContext.BuildContext.DebugAssetName, Node.Interface.Outputs.Num(), ClassInterface.Outputs.Num(), *Node.Name.ToString(), *NodeClass->Metadata.GetClassName().ToString());
+						continue;
+					}
+
+
 					switch (NodeClass->Metadata.GetType())
 					{
 						case EMetasoundFrontendClassType::Input:
@@ -608,7 +579,7 @@ namespace Metasound::Frontend
 
 							if (nullptr != ClassInput)
 							{
-								TUniquePtr<INode> InputNode = CreateInputNode(InGraphContext.BuildContext, Node, *NodeClass, *ClassInput);
+								TUniquePtr<INode> InputNode = CreateInputNode(InGraphContext.BuildContext, Node, *NodeClass, ClassInterface, *ClassInput);
 								if (InputNode.IsValid())
 								{
 									InGraphContext.Graph->AddInputNode(Node.GetID(), ClassInput->Name, MoveTemp(InputNode));
@@ -632,7 +603,7 @@ namespace Metasound::Frontend
 							const FMetasoundFrontendClassOutput* ClassOutput = FindClassOutputForOutputNode(InGraphContext.GraphClass, Node);
 							if (nullptr != ClassOutput)
 							{
-								TUniquePtr<INode> OutputNode = CreateOutputNode(InGraphContext.BuildContext, Node, *NodeClass, *ClassOutput);
+								TUniquePtr<INode> OutputNode = CreateOutputNode(InGraphContext.BuildContext, Node, *NodeClass, ClassInterface, *ClassOutput);
 								if (OutputNode)
 								{
 									InGraphContext.Graph->AddOutputNode(Node.GetID(), ClassOutput->Name, MoveTemp(OutputNode));
@@ -677,7 +648,7 @@ PRAGMA_ENABLE_DEPRECATION_WARNINGS
 
 						case EMetasoundFrontendClassType::Variable:
 						{
-							TUniquePtr<INode> VariableNode = CreateVariableNode(InGraphContext.BuildContext, Node, *NodeClass, InGraphContext.PagedGraph);
+							TUniquePtr<INode> VariableNode = CreateVariableNode(InGraphContext.BuildContext, Node, *NodeClass, ClassInterface, InGraphContext.PagedGraph);
 							if (VariableNode)
 							{
 								InGraphContext.Graph->AddNode(Node.GetID(), MoveTemp(VariableNode));
@@ -698,7 +669,7 @@ PRAGMA_ENABLE_DEPRECATION_WARNINGS
 						case EMetasoundFrontendClassType::External:
 						default:
 						{
-							TUniquePtr<INode> ExternalNode = CreateExternalNode(InGraphContext.BuildContext, Node, *NodeClass);
+							TUniquePtr<INode> ExternalNode = CreateExternalNode(InGraphContext.BuildContext, Node, *NodeClass, ClassInterface);
 							if (ExternalNode)
 							{
 								InGraphContext.Graph->AddNode(Node.GetID(), MoveTemp(ExternalNode));
diff --git a/Engine/Plugins/Runtime/Metasound/Source/MetasoundFrontend/Private/MetasoundFrontendGraphController.cpp b/Engine/Plugins/Runtime/Metasound/Source/MetasoundFrontend/Private/MetasoundFrontendGraphController.cpp
--- a/Engine/Plugins/Runtime/Metasound/Source/MetasoundFrontend/Private/MetasoundFrontendGraphController.cpp
+++ b/Engine/Plugins/Runtime/Metasound/Source/MetasoundFrontend/Private/MetasoundFrontendGraphController.cpp
@@ -131,7 +131,7 @@ namespace Metasound
 
 			if (const FMetasoundFrontendGraphClass* GraphClass = GraphClassPtr.Get())
 			{
-				for (const FMetasoundFrontendClassInput& Input : GraphClass->Interface.Inputs)
+				for (const FMetasoundFrontendClassInput& Input : GraphClass->GetDefaultInterface().Inputs)
 				{
 					Names.Add(Input.Name);
 				}
@@ -146,7 +146,7 @@ namespace Metasound
 
 			if (const FMetasoundFrontendGraphClass* GraphClass = GraphClassPtr.Get())
 			{
-				for (const FMetasoundFrontendClassOutput& Output : GraphClass->Interface.Outputs)
+				for (const FMetasoundFrontendClassOutput& Output : GraphClass->GetDefaultInterface().Outputs)
 				{
 					Names.Add(Output.Name);
 				}
@@ -226,7 +226,7 @@ namespace Metasound
 		{
 			if (FMetasoundFrontendGraphClass* GraphClass = GraphClassPtr.Get())
 			{
-				return GraphClass->Interface.GetInputStyle();
+				return GraphClass->GetDefaultInterface().GetInputStyle();
 			}
 
 			return Invalid::GetInvalidInterfaceStyle();
@@ -236,7 +236,7 @@ namespace Metasound
 		{
 			if (FMetasoundFrontendGraphClass* GraphClass = GraphClassPtr.Get())
 			{
-				return GraphClass->Interface.GetOutputStyle();
+				return GraphClass->GetDefaultInterface().GetOutputStyle();
 			}
 
 			return Invalid::GetInvalidInterfaceStyle();
@@ -254,14 +254,14 @@ namespace Metasound
 		{
 			if (FMetasoundFrontendGraphClass* GraphClass = GraphClassPtr.Get())
 			{
-				TArray<FMetasoundFrontendClassInput>& Inputs = GraphClass->Interface.Inputs;
+				TArray<FMetasoundFrontendClassInput>& Inputs = GraphClass->GetDefaultInterface().Inputs;
 				Style.DefaultSortOrder.SetNumZeroed(Inputs.Num());
 
 				for (int32 i = 0; i < Inputs.Num(); ++i)
 				{
 					Inputs[i].Metadata.SortOrderIndex = Style.DefaultSortOrder[i];
 				}
-				GraphClass->Interface.SetInputStyle(MoveTemp(Style));
+				GraphClass->GetDefaultInterface().SetInputStyle(MoveTemp(Style));
 			}
 		}
 
@@ -269,14 +269,14 @@ namespace Metasound
 		{
 			if (FMetasoundFrontendGraphClass* GraphClass = GraphClassPtr.Get())
 			{
-				TArray<FMetasoundFrontendClassOutput>& Outputs = GraphClass->Interface.Outputs;
+				TArray<FMetasoundFrontendClassOutput>& Outputs = GraphClass->GetDefaultInterface().Outputs;
 				Style.DefaultSortOrder.SetNumZeroed(Outputs.Num());
 
 				for (int32 i = 0; i < Outputs.Num(); ++i)
 				{
 					Outputs[i].Metadata.SortOrderIndex = Style.DefaultSortOrder[i];
 				}
-				GraphClass->Interface.SetOutputStyle(MoveTemp(Style));
+				GraphClass->GetDefaultInterface().SetOutputStyle(MoveTemp(Style));
 			}
 		}
 #endif // WITH_EDITOR
@@ -778,8 +778,8 @@ namespace Metasound
 				FMetasoundFrontendGraph& Graph = FindBuildGraphChecked(*GraphClass);
 				Graph.Nodes.Reset();
 				Graph.Edges.Reset();
-				GraphClass->Interface.Inputs.Reset();
-				GraphClass->Interface.Outputs.Reset();
+				GraphClass->GetDefaultInterface().Inputs.Reset();
+				GraphClass->GetDefaultInterface().Outputs.Reset();
 				GraphClass->PresetOptions.InputsInheritingDefault.Reset();
 				OwningDocument->ClearInterfaceVersions();
 				OwningDocument->RemoveUnreferencedDependencies();
@@ -919,7 +919,7 @@ namespace Metasound
 				{
 					return ClassOutput.Name == InName && ClassOutput.TypeName == InTypeName;
 				};
-				return GraphClass->Interface.Outputs.ContainsByPredicate(IsOutputVertexWithSameNameAndType);
+				return GraphClass->GetDefaultInterface().Outputs.ContainsByPredicate(IsOutputVertexWithSameNameAndType);
 			}
 			return false;
 		}
@@ -932,7 +932,7 @@ namespace Metasound
 				{
 					return ClassOutput.Name == InName;
 				};
-				return GraphClass->Interface.Outputs.ContainsByPredicate(IsOutputVertexWithSameName);
+				return GraphClass->GetDefaultInterface().Outputs.ContainsByPredicate(IsOutputVertexWithSameName);
 			}
 			return false;
 		}
@@ -945,7 +945,7 @@ namespace Metasound
 				{
 					return ClassInput.Name == InName && ClassInput.TypeName == InTypeName;
 				};
-				return GraphClass->Interface.Inputs.ContainsByPredicate(IsInputVertexWithSameNameAndType);
+				return GraphClass->GetDefaultInterface().Inputs.ContainsByPredicate(IsInputVertexWithSameNameAndType);
 			}
 			return false;
 		}
@@ -958,7 +958,7 @@ namespace Metasound
 				{ 
 					return ClassInput.Name == InName;
 				};
-				return GraphClass->Interface.Inputs.ContainsByPredicate(IsInputVertexWithSameName);
+				return GraphClass->GetDefaultInterface().Inputs.ContainsByPredicate(IsInputVertexWithSameName);
 			}
 			return false;
 		}
@@ -1006,7 +1006,7 @@ namespace Metasound
 			if (FMetasoundFrontendGraphClass* GraphClass = GraphClassPtr.Get())
 			{
 				auto IsInputWithSameName = [&](const FMetasoundFrontendClassInput& ExistingDesc) { return ExistingDesc.Name == InClassInput.Name; };
-				if (Algo::NoneOf(GraphClass->Interface.Inputs, IsInputWithSameName))
+				if (Algo::NoneOf(GraphClass->GetDefaultInterface().Inputs, IsInputWithSameName))
 				{
 					FNodeRegistryKey Key;
 					if (FRegistry::GetInputNodeRegistryKeyForDataType(InClassInput.TypeName, InClassInput.AccessType, Key))
@@ -1052,7 +1052,7 @@ namespace Metasound
 							}
 
 							// Add input to this graph class interface
-							FMetasoundFrontendClassInput& NewInput = GraphClass->Interface.Inputs.Add_GetRef(InClassInput);
+							FMetasoundFrontendClassInput& NewInput = GraphClass->GetDefaultInterface().Inputs.Add_GetRef(InClassInput);
 
 							NewInput.NodeID = Node.GetID();
 							if (!NewInput.VertexID.IsValid())
@@ -1071,7 +1071,7 @@ namespace Metasound
 								DocMetadata->ModifyContext.AddMemberIDModified(NewInput.NodeID);
 							}
 #endif // WITH_EDITOR
-							GraphClass->Interface.UpdateChangeID();
+							GraphClass->GetDefaultInterface().UpdateChangeID();
 
 							FNodeAccessPtr NodePtr = GraphClassPtr.GetNodeWithNodeID(Node.GetID());
 							return GetNodeHandle(FGraphController::FNodeAndClass { NodePtr, InputClassPtr });
@@ -1129,7 +1129,7 @@ namespace Metasound
 			if (FMetasoundFrontendGraphClass* GraphClass = GraphClassPtr.Get())
 			{
 				auto IsOutputWithSameName = [&](const FMetasoundFrontendClassOutput& ExistingDesc) { return ExistingDesc.Name == InClassOutput.Name; };
-				if (Algo::NoneOf(GraphClass->Interface.Outputs, IsOutputWithSameName))
+				if (Algo::NoneOf(GraphClass->GetDefaultInterface().Outputs, IsOutputWithSameName))
 				{
 					FNodeRegistryKey Key;
 					if (FRegistry::GetOutputNodeRegistryKeyForDataType(InClassOutput.TypeName, InClassOutput.AccessType, Key))
@@ -1175,7 +1175,7 @@ namespace Metasound
 							}
 
 							// Add output to graph interface
-							FMetasoundFrontendClassOutput& NewOutput = GraphClass->Interface.Outputs.Add_GetRef(InClassOutput);
+							FMetasoundFrontendClassOutput& NewOutput = GraphClass->GetDefaultInterface().Outputs.Add_GetRef(InClassOutput);
 							
 							// Setup new output
 							NewOutput.NodeID = Node.GetID();
@@ -1197,7 +1197,7 @@ namespace Metasound
 #endif // WITH_EDITOR
 
 							// Mark interface as changed.
-							GraphClass->Interface.UpdateChangeID();
+							GraphClass->GetDefaultInterface().UpdateChangeID();
 
 							FNodeAccessPtr NodePtr = GraphClassPtr.GetNodeWithNodeID(Node.GetID());
 							return GetNodeHandle(FGraphController::FNodeAndClass { NodePtr, OutputClassPtr });
@@ -1392,9 +1392,9 @@ namespace Metasound
 		{
 			if (FMetasoundFrontendGraphClass* GraphClass = GraphClassPtr.Get())
 			{
-				FMetasoundFrontendInterfaceStyle Style = GraphClass->Interface.GetInputStyle();
+				FMetasoundFrontendInterfaceStyle Style = GraphClass->GetDefaultInterface().GetInputStyle();
 				Style.DefaultSortOrder.Reset();
-				for (FMetasoundFrontendClassInput& Input : GraphClass->Interface.Inputs)
+				for (FMetasoundFrontendClassInput& Input : GraphClass->GetDefaultInterface().Inputs)
 				{
 					if (Input.Name == InName)
 					{
@@ -1402,7 +1402,7 @@ namespace Metasound
 					}
 					Style.DefaultSortOrder.Add(Input.Metadata.SortOrderIndex);
 				}
-				GraphClass->Interface.SetInputStyle(Style);
+				GraphClass->GetDefaultInterface().SetInputStyle(Style);
 			}
 		}
 
@@ -1410,9 +1410,9 @@ namespace Metasound
 		{
 			if (FMetasoundFrontendGraphClass* GraphClass = GraphClassPtr.Get())
 			{
-				FMetasoundFrontendInterfaceStyle Style = GraphClass->Interface.GetOutputStyle();
+				FMetasoundFrontendInterfaceStyle Style = GraphClass->GetDefaultInterface().GetOutputStyle();
 				Style.DefaultSortOrder.Reset();
-				for (FMetasoundFrontendClassOutput& Output : GraphClass->Interface.Outputs)
+				for (FMetasoundFrontendClassOutput& Output : GraphClass->GetDefaultInterface().Outputs)
 				{
 					if (Output.Name == InName)
 					{
@@ -1420,7 +1420,7 @@ namespace Metasound
 					}
 					Style.DefaultSortOrder.Add(Output.Metadata.SortOrderIndex);
 				}
-				GraphClass->Interface.SetOutputStyle(MoveTemp(Style));
+				GraphClass->GetDefaultInterface().SetOutputStyle(MoveTemp(Style));
 			}
 		}
 #endif // WITH_EDITOR
@@ -1767,7 +1767,7 @@ namespace Metasound
 			{
 				auto IsInputWithSameNodeID = [&](const FMetasoundFrontendClassInput& ClassInput) { return ClassInput.NodeID == InNode.GetID(); };
 
-				int32 NumInputsRemoved = GraphClass->Interface.Inputs.RemoveAll(IsInputWithSameNodeID);
+				int32 NumInputsRemoved = GraphClass->GetDefaultInterface().Inputs.RemoveAll(IsInputWithSameNodeID);
 
 				if (NumInputsRemoved > 0)
 				{
@@ -1777,7 +1777,7 @@ namespace Metasound
 						DocMetadata->ModifyContext.AddMemberIDModified(InNode.GetID());
 					}
 #endif // WITH_EDITOR
-					GraphClass->Interface.UpdateChangeID();
+					GraphClass->GetDefaultInterface().UpdateChangeID();
 				}
 
 				bool bDidRemoveNode = RemoveNode(InNode);
@@ -1794,7 +1794,7 @@ namespace Metasound
 			{
 				auto IsOutputWithSameNodeID = [&](const FMetasoundFrontendClassOutput& ClassOutput) { return ClassOutput.NodeID == InNode.GetID(); };
 
-				int32 NumOutputsRemoved = GraphClass->Interface.Outputs.RemoveAll(IsOutputWithSameNodeID);
+				int32 NumOutputsRemoved = GraphClass->GetDefaultInterface().Outputs.RemoveAll(IsOutputWithSameNodeID);
 
 				if (NumOutputsRemoved)
 				{
@@ -1804,7 +1804,7 @@ namespace Metasound
 						DocMetadata->ModifyContext.AddMemberIDModified(InNode.GetID());
 					}
 #endif // WITH_EDITOR
-					GraphClass->Interface.UpdateChangeID();
+					GraphClass->GetDefaultInterface().UpdateChangeID();
 				}
 
 				bool bDidRemoveNode = RemoveNode(InNode);
@@ -1819,7 +1819,7 @@ namespace Metasound
 		{
 			if (FMetasoundFrontendGraphClass* GraphClass = GraphClassPtr.Get())
 			{
-				GraphClass->Interface.UpdateChangeID();
+				GraphClass->GetDefaultInterface().UpdateChangeID();
 			}
 		}
 
@@ -2245,12 +2245,12 @@ namespace Metasound
 		{
 			if (FMetasoundFrontendGraphClass* GraphClass = GraphClassPtr.Get())
 			{
-				FMetasoundFrontendClassInput* ClassInput = GraphClass->Interface.Inputs.FindByPredicate([&](const FMetasoundFrontendClassInput& Desc) { return Desc.Name == InName; });
+				FMetasoundFrontendClassInput* ClassInput = GraphClass->GetDefaultInterface().Inputs.FindByPredicate([&](const FMetasoundFrontendClassInput& Desc) { return Desc.Name == InName; });
 				if (ClassInput)
 				{
 					// TODO: This assumes the class input is being mutated due to the adjacent const correct call not being utilized.
 					// Make this more explicit rather than risking whether or not the caller is using proper const correctness.
-					GraphClass->Interface.UpdateChangeID();
+					GraphClass->GetDefaultInterface().UpdateChangeID();
 					return ClassInput;
 				}
 			}
@@ -2261,7 +2261,7 @@ namespace Metasound
 		{
 			if (const FMetasoundFrontendGraphClass* GraphClass = GraphClassPtr.Get())
 			{
-				return GraphClass->Interface.Inputs.FindByPredicate([&](const FMetasoundFrontendClassInput& Desc) { return Desc.Name == InName; });
+				return GraphClass->GetDefaultInterface().Inputs.FindByPredicate([&](const FMetasoundFrontendClassInput& Desc) { return Desc.Name == InName; });
 			}
 			return nullptr;
 		}
@@ -2270,12 +2270,12 @@ namespace Metasound
 		{
 			if (FMetasoundFrontendGraphClass* GraphClass = GraphClassPtr.Get())
 			{
-				FMetasoundFrontendClassOutput* ClassOutput = GraphClass->Interface.Outputs.FindByPredicate([&](const FMetasoundFrontendClassOutput& Desc) { return Desc.Name == InName; });
+				FMetasoundFrontendClassOutput* ClassOutput = GraphClass->GetDefaultInterface().Outputs.FindByPredicate([&](const FMetasoundFrontendClassOutput& Desc) { return Desc.Name == InName; });
 				if (ClassOutput)
 				{
 					// TODO: This assumes the class input is being mutated due to the adjacent const correct call not being utilized.
 					// Make this more explicit rather than risking whether or not the caller is using proper const correctness.
-					GraphClass->Interface.UpdateChangeID();
+					GraphClass->GetDefaultInterface().UpdateChangeID();
 					return ClassOutput;
 				}
 			}
@@ -2286,7 +2286,7 @@ namespace Metasound
 		{
 			if (const FMetasoundFrontendGraphClass* GraphClass = GraphClassPtr.Get())
 			{
-				return GraphClass->Interface.Outputs.FindByPredicate([&](const FMetasoundFrontendClassOutput& Desc) { return Desc.Name == InName; });
+				return GraphClass->GetDefaultInterface().Outputs.FindByPredicate([&](const FMetasoundFrontendClassOutput& Desc) { return Desc.Name == InName; });
 			}
 			return nullptr;
 		}
@@ -2295,12 +2295,12 @@ namespace Metasound
 		{
 			if (FMetasoundFrontendGraphClass* GraphClass = GraphClassPtr.Get())
 			{
-				FMetasoundFrontendClassInput* ClassInput = GraphClass->Interface.Inputs.FindByPredicate([&](const FMetasoundFrontendClassInput& Desc) { return Desc.VertexID == InVertexID; });
+				FMetasoundFrontendClassInput* ClassInput = GraphClass->GetDefaultInterface().Inputs.FindByPredicate([&](const FMetasoundFrontendClassInput& Desc) { return Desc.VertexID == InVertexID; });
 				if (ClassInput)
 				{
 					// TODO: This assumes the class input is being mutated due to the adjacent const correct call not being utilized.
 					// Make this more explicit rather than risking whether or not the caller is using proper const correctness.
-					GraphClass->Interface.UpdateChangeID();
+					GraphClass->GetDefaultInterface().UpdateChangeID();
 					return ClassInput;
 				}
 			}
@@ -2311,7 +2311,7 @@ namespace Metasound
 		{
 			if (const FMetasoundFrontendGraphClass* GraphClass = GraphClassPtr.Get())
 			{
-				return GraphClass->Interface.Inputs.FindByPredicate([&](const FMetasoundFrontendClassInput& Desc) { return Desc.VertexID == InVertexID; });
+				return GraphClass->GetDefaultInterface().Inputs.FindByPredicate([&](const FMetasoundFrontendClassInput& Desc) { return Desc.VertexID == InVertexID; });
 			}
 			return nullptr;
 		}
@@ -2320,12 +2320,12 @@ namespace Metasound
 		{
 			if (FMetasoundFrontendGraphClass* GraphClass = GraphClassPtr.Get())
 			{
-				FMetasoundFrontendClassOutput* ClassOutput = GraphClass->Interface.Outputs.FindByPredicate([&](const FMetasoundFrontendClassOutput& Desc) { return Desc.VertexID == InVertexID; });
+				FMetasoundFrontendClassOutput* ClassOutput = GraphClass->GetDefaultInterface().Outputs.FindByPredicate([&](const FMetasoundFrontendClassOutput& Desc) { return Desc.VertexID == InVertexID; });
 				if (ClassOutput)
 				{
 					// TODO: This assumes the class input is being mutated due to the adjacent const correct call not being utilized.
 					// Make this more explicit rather than risking whether or not the caller is using proper const correctness.
-					GraphClass->Interface.UpdateChangeID();
+					GraphClass->GetDefaultInterface().UpdateChangeID();
 					return ClassOutput;
 				}
 			}
@@ -2336,7 +2336,7 @@ namespace Metasound
 		{
 			if (const FMetasoundFrontendGraphClass* GraphClass = GraphClassPtr.Get())
 			{
-				return GraphClass->Interface.Outputs.FindByPredicate([&](const FMetasoundFrontendClassOutput& Desc) { return Desc.VertexID == InVertexID; });
+				return GraphClass->GetDefaultInterface().Outputs.FindByPredicate([&](const FMetasoundFrontendClassOutput& Desc) { return Desc.VertexID == InVertexID; });
 			}
 			return nullptr;
 		}
diff --git a/Engine/Plugins/Runtime/Metasound/Source/MetasoundFrontend/Private/MetasoundFrontendInvalidController.h b/Engine/Plugins/Runtime/Metasound/Source/MetasoundFrontend/Private/MetasoundFrontendInvalidController.h
--- a/Engine/Plugins/Runtime/Metasound/Source/MetasoundFrontend/Private/MetasoundFrontendInvalidController.h
+++ b/Engine/Plugins/Runtime/Metasound/Source/MetasoundFrontend/Private/MetasoundFrontendInvalidController.h
@@ -253,7 +253,7 @@ namespace Metasound
 			virtual const FMetasoundFrontendVersion& GetInterfaceVersion() const { return FMetasoundFrontendVersion::GetInvalid(); }
 			virtual bool IsInterfaceMember() const override { return false; }
 
-			virtual bool DiffAgainstRegistryInterface(FClassInterfaceUpdates& OutInterfaceUpdates, bool bInUseHighestMinorVersion) const override { return false; }
+			virtual bool DiffAgainstRegistryInterface(FClassInterfaceUpdates& OutInterfaceUpdates, bool bInUseHighestMinorVersion, bool bInForceRegeneratedClassInterfaceOverride) const override { return false; }
 			virtual bool CanAutoUpdate(FClassInterfaceUpdates& OutInterfaceUpdates) const override { OutInterfaceUpdates = { }; return false; }
 
 			virtual TSharedRef<IGraphController> AsGraph() override;
@@ -277,6 +277,9 @@ namespace Metasound
 
 			virtual const FVertexName& GetNodeName() const override { return Invalid::GetInvalidName(); }
 
+	   		virtual void SetNodeConfiguration(TInstancedStruct<FMetaSoundFrontendNodeConfiguration> InConfig) override {}
+	   		virtual TConstStructView<FMetaSoundFrontendNodeConfiguration> GetNodeConfiguration() const override { return TConstStructView<FMetaSoundFrontendNodeConfiguration>{}; }
+
 #if WITH_EDITOR
 			virtual FText GetDisplayName() const override { return Invalid::GetInvalidText(); }
 			virtual const FText& GetDisplayTitle() const override { return Invalid::GetInvalidText(); }
diff --git a/Engine/Plugins/Runtime/Metasound/Source/MetasoundFrontend/Private/MetasoundFrontendNodeClassRegistry.cpp b/Engine/Plugins/Runtime/Metasound/Source/MetasoundFrontend/Private/MetasoundFrontendNodeClassRegistry.cpp
--- a/Engine/Plugins/Runtime/Metasound/Source/MetasoundFrontend/Private/MetasoundFrontendNodeClassRegistry.cpp
+++ b/Engine/Plugins/Runtime/Metasound/Source/MetasoundFrontend/Private/MetasoundFrontendNodeClassRegistry.cpp
@@ -299,9 +299,12 @@ namespace Metasound::Frontend
 				}
 			}
 
-			virtual TInstancedStruct<FMetaSoundFrontendNodeConfiguration> CreateFrontendNodeConfiguration() const override
+			virtual TInstancedStruct<FMetaSoundFrontendNodeConfiguration> CreateFrontendNodeConfiguration() const override final
 			{
-				// Document based nodes do not current offer extensions. 
+				// Document based nodes do not support node configuration because
+				// many MetaSound systems assume that a graph defined in a FMetasoundFrontendDocument
+				// only supplies a default interface. The use of class interface
+				// overrides in FMetasoundFrontendDocument based nodes is unsupported. 
 				return TInstancedStruct<FMetaSoundFrontendNodeConfiguration>();
 			}
 		private:
diff --git a/Engine/Plugins/Runtime/Metasound/Source/MetasoundFrontend/Private/MetasoundFrontendNodeController.cpp b/Engine/Plugins/Runtime/Metasound/Source/MetasoundFrontend/Private/MetasoundFrontendNodeController.cpp
--- a/Engine/Plugins/Runtime/Metasound/Source/MetasoundFrontend/Private/MetasoundFrontendNodeController.cpp
+++ b/Engine/Plugins/Runtime/Metasound/Source/MetasoundFrontend/Private/MetasoundFrontendNodeController.cpp
@@ -130,7 +130,13 @@ namespace Metasound
 						{
 							return ClassInput.VertexID == InVertexID;
 						};
-						const FMetasoundFrontendClassInput* ClassInput = ClassPtr.Get()->Interface.Inputs.FindByPredicate(GetInputWithVertexID);
+
+						const FMetasoundFrontendClassInput* ClassInput = nullptr;
+						if (const FMetasoundFrontendClassInterface* ClassInterface = GetClassInterfaceForNode())
+						{
+							ClassInput = ClassInterface->Inputs.FindByPredicate(GetInputWithVertexID);
+						}
+
 						if (ClassInput)
 						{
 							// Check if setting back to class default literal
@@ -180,9 +186,9 @@ namespace Metasound
 
 		const FMetasoundFrontendClassInterface& FBaseNodeController::GetClassInterface() const
 		{
-			if (const FMetasoundFrontendClass* Class = ClassPtr.Get())
+			if (const FMetasoundFrontendClassInterface* Interface = GetClassInterfaceForNode())
 			{
-				return Class->Interface;
+				return *Interface;
 			}
 			return Invalid::GetInvalidClassInterface();
 		}
@@ -208,9 +214,9 @@ namespace Metasound
 #if WITH_EDITOR
 		const FMetasoundFrontendInterfaceStyle& FBaseNodeController::GetInputStyle() const
 		{
-			if (const FMetasoundFrontendClass* Class = ClassPtr.Get())
+			if (const FMetasoundFrontendClassInterface* ClassInterface = GetClassInterfaceForNode())
 			{
-				return Class->Interface.GetInputStyle();
+				return ClassInterface->GetInputStyle();
 			}
 
 			return Invalid::GetInvalidInterfaceStyle();
@@ -218,9 +224,9 @@ namespace Metasound
 
 		const FMetasoundFrontendInterfaceStyle& FBaseNodeController::GetOutputStyle() const
 		{
-			if (const FMetasoundFrontendClass* Class = ClassPtr.Get())
+			if (const FMetasoundFrontendClassInterface* ClassInterface = GetClassInterfaceForNode())
 			{
-				return Class->Interface.GetOutputStyle();
+				return ClassInterface->GetOutputStyle();
 			}
 
 			return Invalid::GetInvalidInterfaceStyle();
@@ -710,6 +716,18 @@ namespace Metasound
 			return false;
 		}
 
+		const FMetasoundFrontendClassInterface* FBaseNodeController::GetClassInterfaceForNode() const
+		{
+			if (const FMetasoundFrontendNode* Node = NodePtr.Get())
+			{
+				if (const FMetasoundFrontendClass* Class = ClassPtr.Get())
+				{
+					return &Class->GetInterfaceForNode(*Node);
+				}
+			}
+			return nullptr;
+		}
+
 		FGraphHandle FBaseNodeController::AsGraph()
 		{
 			// TODO: consider adding support for external graph owned in another document.
@@ -737,6 +755,7 @@ namespace Metasound
 
 		FNodeHandle FBaseNodeController::ReplaceWithVersion(const FMetasoundFrontendVersionNumber& InNewVersion, TArray<FVertexNameAndType>* OutDisconnectedInputs, TArray<FVertexNameAndType>* OutDisconnectedOutputs)
 		{
+			// TODO: Need to deal with node configuration and class interface override here. 
 			const FMetasoundFrontendClassMetadata Metadata = GetClassMetadata();
 
 			// Lookup new version in node registry
@@ -755,35 +774,69 @@ namespace Metasound
 				return this->AsShared();
 			}
 
+			// Gathers all info from the node that could be applied to the replacement node. 
+			// This is destructive in that some objects are moved off of the node.
+			FNodeInstanceReplacementData InstanceData = CaptureNodeInstanceReplacementData();
+
+			// Remove the actual node
+			if (!ensureAlways(GetOwningGraph()->RemoveNode(*this)))
+			{
+				return this->AsShared();
+			}
+
+			// Make sure classes are up-to-date with registered versions of class.
+			// Note that this may break other nodes in the graph that have stale
+			// class API, but that's on the caller to fix-up or report invalid state.
+			const FNodeRegistryKey RegistryKey = FNodeRegistryKey(NewMetasoundClass.Metadata);
+			FDocumentHandle Document = GetOwningGraph()->GetOwningDocument();
+
+			constexpr bool bRefreshFromRegistry = true;
+			ensureAlways(Document->FindOrAddClass(RegistryKey, bRefreshFromRegistry).Get() != nullptr);
+
+			FNodeHandle ReplacementNode = GetOwningGraph()->AddNode(NewMetasoundClass.Metadata, InstanceData.NodeID);
+			if (!ensureAlways(ReplacementNode->IsValid()))
+			{
+				return this->AsShared();
+			}
+
 #if WITH_EDITOR
-			FMetasoundFrontendNodeStyle Style = GetNodeStyle();
+			InstanceData.Style.bMessageNodeUpdated = ReplacementNode->GetClassMetadata().GetVersion() > Metadata.GetVersion();
 #endif // WITH_EDITOR
 
-			struct FInputConnectionInfo
+			ApplyNodeInstanceReplacementData(NewMetasoundClass, *ReplacementNode, MoveTemp(InstanceData), OutDisconnectedInputs, OutDisconnectedOutputs);
+
+			return ReplacementNode;
+		}
+
+		FBaseNodeController::FNodeInstanceReplacementData FBaseNodeController::CaptureNodeInstanceReplacementData()
+		{
+			FNodeInstanceReplacementData ReplacementData;
+			FMetasoundFrontendNode* OriginalNode = NodePtr.Get();
+			if (!OriginalNode)
 			{
-				FOutputHandle ConnectedOutput;
-				FVertexName Name;
-				FName DataType;
-				FMetasoundFrontendLiteral DefaultValue;
-				bool bLiteralSet = false;
-			};
+				// This node handle is invalid. There is no nothing to replace.
+				return ReplacementData;
+			}
+
+#if WITH_EDITOR
+			ReplacementData.Style = MoveTemp(OriginalNode->Style);
+#endif // WITH_EDITOR
 
+			// Move any configuration and override data living on the node.
+			ReplacementData.Configuration = MoveTemp(OriginalNode->Configuration);
+			ReplacementData.ClassInterfaceOverride = MoveTemp(OriginalNode->ClassInterfaceOverride);;
+			ReplacementData.NodeID = OriginalNode->GetID();
+			
 			// Cache input/output connections by name to try so they can be
 			// hooked back up after swapping to the new class version.
-			TMap<FVertexNameAndType, FInputConnectionInfo> InputConnections;
-			IterateInputs([Connections = &InputConnections](FInputHandle InputHandle)
+			IterateInputs([Connections = &ReplacementData.InputConnections](FInputHandle InputHandle)
 			{
 				bool bLiteralSet = false;
 				FMetasoundFrontendLiteral DefaultLiteral;
 				if (const FMetasoundFrontendLiteral* Literal = InputHandle->GetLiteral())
 				{
-					// Array literals are not supported in UX, so don't pass along to referencing graph
-					// TODO: Add UX in inspector to set literals (including arrays).
-					if (!Literal->IsArray())
-					{
-						DefaultLiteral = *Literal;
-						bLiteralSet = true;
-					}
+					DefaultLiteral = *Literal;
+					bLiteralSet = true;
 				}
 
 				const FVertexNameAndType ConnectionKey(InputHandle->GetName(), InputHandle->GetDataType());
@@ -797,15 +850,8 @@ namespace Metasound
 				});
 			});
 
-			struct FOutputConnectionInfo
-			{
-				TArray<FInputHandle> ConnectedInputs;
-				FVertexName VertexName;
-				FName DataType;
-			};
 
-			TMap<FVertexNameAndType, FOutputConnectionInfo> OutputConnections;
-			IterateOutputs([Connections = &OutputConnections](FOutputHandle OutputHandle)
+			IterateOutputs([Connections = &ReplacementData.OutputConnections](FOutputHandle OutputHandle)
 			{
 				const FVertexNameAndType ConnectionKey(OutputHandle->GetName(), OutputHandle->GetDataType());
 				Connections->Add(ConnectionKey, FOutputConnectionInfo
@@ -816,38 +862,27 @@ namespace Metasound
 				});
 			});
 
-			// Maintain the NodeID to ensure the editor layer maintains node representation parity.
-			const FGuid ReplacedNodeGuid = GetID();
-
-			// Maintain class ID to ensure all nodes created & preexisting reference the same dependency.
-			const FGuid ClassID = GetClassID();
-
-			if (!ensureAlways(GetOwningGraph()->RemoveNode(*this)))
-			{
-				return this->AsShared();
-			}
-
-			// Make sure classes are up-to-date with registered versions of class.
-			// Note that this may break other nodes in the graph that have stale
-			// class API, but that's on the caller to fix-up or report invalid state.
-			const FNodeRegistryKey RegistryKey = FNodeRegistryKey(NewMetasoundClass.Metadata);
-			FDocumentHandle Document = GetOwningGraph()->GetOwningDocument();
+			// nrvo.
+			return ReplacementData;
+		}
 
-			constexpr bool bRefreshFromRegistry = true;
-			ensureAlways(Document->FindOrAddClass(RegistryKey, bRefreshFromRegistry).Get() != nullptr);
+		void FBaseNodeController::ApplyNodeInstanceReplacementData(const FMetasoundFrontendClass& InReplacementClass, INodeController& InReplacementNode, FNodeInstanceReplacementData&& InInstanceData, TArray<FVertexNameAndType>* OutDisconnectedInputs, TArray<FVertexNameAndType>* OutDisconnectedOutputs)
+		{
+#if WITH_EDITOR
+			InReplacementNode.SetNodeStyle(InInstanceData.Style);
+#endif // WITH_EDITOR
 
-			FNodeHandle ReplacementNode = GetOwningGraph()->AddNode(NewMetasoundClass.Metadata, ReplacedNodeGuid);
-			if (!ensureAlways(ReplacementNode->IsValid()))
+			// The logic may appear a little backwards here because adding the node already instantiated the replacement
+			// node configuration instances. The purpose is to revert the replacement if it should not have happened in 
+			// the first place.
+			if (!ShouldReplaceExistingNodeConfig(InReplacementNode.GetNodeConfiguration(), InInstanceData.Configuration))
 			{
-				return this->AsShared();
+				// We should have kept the original node configuration. Set back to the original
+				InReplacementNode.SetNodeConfiguration(MoveTemp(InInstanceData.Configuration));
 			}
+			
 
-#if WITH_EDITOR
-			Style.bMessageNodeUpdated = ReplacementNode->GetClassMetadata().GetVersion() > Metadata.GetVersion();
-			ReplacementNode->SetNodeStyle(Style);
-#endif // WITH_EDITOR
-
-			ReplacementNode->IterateInputs([&InputConnections](FInputHandle InputHandle)
+			InReplacementNode.IterateInputs([&InputConnections = InInstanceData.InputConnections](FInputHandle InputHandle)
 			{
 				const FVertexNameAndType ConnectionKey(InputHandle->GetName(), InputHandle->GetDataType());
 				if (FInputConnectionInfo* ConnectionInfo = InputConnections.Find(ConnectionKey))
@@ -872,7 +907,7 @@ namespace Metasound
 			// Track missing input connections
 			if (nullptr != OutDisconnectedInputs)
 			{
-				for (const auto& ConnectionInfoKV : InputConnections)
+				for (const auto& ConnectionInfoKV : InInstanceData.InputConnections)
 				{
 					const FInputConnectionInfo& ConnectionInfo = ConnectionInfoKV.Value;
 					if (ConnectionInfo.ConnectedOutput->IsValid())
@@ -882,7 +917,7 @@ namespace Metasound
 				}
 			}
 
-			ReplacementNode->IterateOutputs([&OutputConnections](FOutputHandle OutputHandle)
+			InReplacementNode.IterateOutputs([&OutputConnections = InInstanceData.OutputConnections](FOutputHandle OutputHandle)
 			{
 				const FVertexNameAndType ConnectionKey(OutputHandle->GetName(), OutputHandle->GetDataType());
 				if (FOutputConnectionInfo* ConnectionInfo = OutputConnections.Find(ConnectionKey))
@@ -909,7 +944,7 @@ namespace Metasound
 			// Track missing output connections
 			if (nullptr != OutDisconnectedOutputs)
 			{
-				for (const auto& ConnectionInfoKV : OutputConnections)
+				for (const auto& ConnectionInfoKV : InInstanceData.OutputConnections)
 				{
 					const FOutputConnectionInfo& ConnectionInfo = ConnectionInfoKV.Value;
 					const bool bAnyConnectedInputs = Algo::AnyOf(ConnectionInfo.ConnectedInputs, [](const FInputHandle& Input) { return Input->IsValid(); });
@@ -919,84 +954,267 @@ namespace Metasound
 					}
 				}
 			}
-
-			return ReplacementNode;
 		}
 
-		bool FBaseNodeController::DiffAgainstRegistryInterface(FClassInterfaceUpdates& OutInterfaceUpdates, bool bInUseHighestMinorVersion) const
+		void FBaseNodeController::SetNodeConfiguration(TInstancedStruct<FMetaSoundFrontendNodeConfiguration> InConfig)
 		{
-			METASOUND_TRACE_CPUPROFILER_EVENT_SCOPE(BaseNodeController::DiffAgainstRegistryInterface);
+			if (FMetasoundFrontendNode* Node = NodePtr.Get())
+			{
+				// If neither the existing, nor new configs are valid, then there should not be any class interface override
+				const bool bUpdateClassInterfaceOverride = Node->Configuration.IsValid() || InConfig.IsValid();
 
-			OutInterfaceUpdates = FClassInterfaceUpdates();
+				Node->Configuration = MoveTemp(InConfig);
 
-			const FMetasoundFrontendClassMetadata& NodeClassMetadata = GetClassMetadata();
-			const FMetasoundFrontendClassInterface& NodeClassInterface = GetClassInterface();
+				if (bUpdateClassInterfaceOverride)
+				{
+					if (const FMetaSoundFrontendNodeConfiguration* NodeConfig = Node->Configuration.GetPtr())
+					{
+						if (const FMetasoundFrontendClass* Class = ClassPtr.Get(); ensure(Class))
+						{
+							// Update the class interface override if we have a node configration
+							SetClassInterfaceOverride(NodeConfig->OverrideDefaultInterface(*Class));
+						}
+					}
+					else
+					{
+						// Clear out the class interface override if there is no node configuration
+						SetClassInterfaceOverride(TInstancedStruct<FMetasoundFrontendClassInterface>());
+					}
+				}
+				else
+				{
+					// Class interface override should be invalid if there isn't 
+					// a node config on the node
+					check(!Node->ClassInterfaceOverride.IsValid());
+				}
+			}
+		}
 
-			Metasound::FNodeClassName NodeClassName = NodeClassMetadata.GetClassName().ToNodeClassName();
+		TConstStructView<FMetaSoundFrontendNodeConfiguration> FBaseNodeController::GetNodeConfiguration() const
+		{
+			TConstStructView<FMetaSoundFrontendNodeConfiguration> ConfigOnNode;
+			if (const FMetasoundFrontendNode* Node = NodePtr.Get())
+			{
+				ConfigOnNode = Node->Configuration;
+			}
+			return ConfigOnNode;
+		}
 
-			if (bInUseHighestMinorVersion)
+		void FBaseNodeController::SetClassInterfaceOverride(TInstancedStruct<FMetasoundFrontendClassInterface> InClassInterfaceOverride)
+		{
+			if (FMetasoundFrontendNode* Node = NodePtr.Get())
 			{
-				if (!ISearchEngine::Get().FindClassWithHighestMinorVersion(NodeClassName, NodeClassMetadata.GetVersion().Major, OutInterfaceUpdates.RegistryClass))
+				// If we are adding, removing or replacing the interface, then we need to
+				// update the node interface.  If the override never existed, we can
+				// skip the process of updating the node's interface. 
+				const bool bUpdateNodeInterface = Node->ClassInterfaceOverride.IsValid() || InClassInterfaceOverride.IsValid();
+
+				Node->ClassInterfaceOverride = MoveTemp(InClassInterfaceOverride);
+
+				if (bUpdateNodeInterface)
 				{
-					Algo::Transform(NodeClassInterface.Inputs, OutInterfaceUpdates.RemovedInputs, [&](const FMetasoundFrontendClassInput& Input) { return &Input; });
-					Algo::Transform(NodeClassInterface.Outputs, OutInterfaceUpdates.RemovedOutputs, [&](const FMetasoundFrontendClassOutput& Output) { return &Output; });
-					return false;
+					if (const FMetasoundFrontendClass* Class = ClassPtr.Get(); ensure(Class))
+					{
+						auto DisconnectInput = [this](const FMetasoundFrontendVertex& InNodeInput)
+						{
+							this->GetInputWithVertexName(InNodeInput.Name)->Disconnect();
+						};
+
+						auto DisconnectOutput = [this](const FMetasoundFrontendVertex& InNodeInput)
+						{
+							this->GetOutputWithVertexName(InNodeInput.Name)->Disconnect();
+						};
+
+						Node->Interface.Update(Class->GetInterfaceForNode(*Node), DisconnectInput, DisconnectOutput);
+					}
 				}
 			}
+		}
+
+		TConstStructView<FMetasoundFrontendClassInterface> FBaseNodeController::GetClassInterfaceOverride() const
+		{
+			TConstStructView<FMetasoundFrontendClassInterface> ClassInterfaceOverride;
+			if (const FMetasoundFrontendNode* Node = NodePtr.Get())
+			{
+				ClassInterfaceOverride = Node->ClassInterfaceOverride;
+			}
+			return ClassInterfaceOverride;
+		}
+
+		bool FBaseNodeController::FindRegisteredClass(const FMetasoundFrontendClassMetadata& InMetadata, bool bInUseHighestMinorVersion, FMetasoundFrontendClass& OutRegisteredClass) const
+		{
+			bool bFoundInRegistry = false;
+
+			const INodeClassRegistry* Registry = INodeClassRegistry::Get();
+			checkf(nullptr != Registry, TEXT("The metasound node registry should always be available if the metasound plugin is loaded"));
+
+			if (bInUseHighestMinorVersion)
+			{
+				// Find class with same name and major version in the search engine
+				bFoundInRegistry = ISearchEngine::Get().FindClassWithHighestMinorVersion(InMetadata.GetClassName().ToNodeClassName(), InMetadata.GetVersion().Major, OutRegisteredClass);
+			}
 			else
 			{
 				// Find class with same metadata in the node registry.
-				FMetasoundFrontendRegistryContainer* Registry = FMetasoundFrontendRegistryContainer::Get();
-				checkf(nullptr != Registry, TEXT("The metasound node registry should always be available if the metasound plugin is loaded"));
-				FMetasoundFrontendClass RegisteredClass;
-				bool bFoundRegisteredClass = Registry->FindFrontendClassFromRegistered(FNodeRegistryKey(GetClassMetadata()), OutInterfaceUpdates.RegistryClass);
+				bFoundInRegistry = Registry->FindFrontendClassFromRegistered(FNodeRegistryKey(InMetadata), OutRegisteredClass);
+			}
+
+			return bFoundInRegistry;
+		}
+
+		bool FBaseNodeController::ShouldReplaceExistingNodeConfig(TConstStructView<FMetaSoundFrontendNodeConfiguration> InRegisteredNodeConfig, TConstStructView<FMetaSoundFrontendNodeConfiguration> InExistingConfig) const
+		{
+			// The existing node configuration is not always replaced. Replacing
+			// the node configuration struct loses any modified state on the existing
+			// node config. In the scenario where the node configurations point to
+			// the same derived class type, we do not replace them. 
+			return InRegisteredNodeConfig.GetScriptStruct() != InExistingConfig.GetScriptStruct();
+		}
+
+		void FBaseNodeController::FindNodeConfigurationUpdates(const FMetasoundFrontendClass& InRegisteredClass, FNodeConfigurationUpdateData& OutNodeConfigurationUpdates, bool bInForceRegenerateClassInterfaceOverride) const
+		{
+			const INodeClassRegistry* Registry = INodeClassRegistry::Get();
+			checkf(nullptr != Registry, TEXT("The metasound node registry should always be available if the metasound plugin is loaded"));
+
+			OutNodeConfigurationUpdates.ExistingConfig = GetNodeConfiguration();
+			OutNodeConfigurationUpdates.ExistingClassInterfaceOverride = GetClassInterfaceOverride();
 
-				if (!bFoundRegisteredClass)
+			OutNodeConfigurationUpdates.RegisteredConfig = Registry->CreateFrontendNodeConfiguration(FNodeRegistryKey(InRegisteredClass.Metadata));
+
+			// Determine which node config will be used to create the class interface override;
+			const bool bWillReplaceNodeConfig = ShouldReplaceExistingNodeConfig(OutNodeConfigurationUpdates.RegisteredConfig, OutNodeConfigurationUpdates.ExistingConfig);
+			TConstStructView<FMetaSoundFrontendNodeConfiguration> EffectiveNodeConfig = bWillReplaceNodeConfig ? OutNodeConfigurationUpdates.RegisteredConfig : OutNodeConfigurationUpdates.ExistingConfig;
+
+			// Update the 
+			if (bWillReplaceNodeConfig || bInForceRegenerateClassInterfaceOverride)
+			{
+				// Set this flag to true even if the effective node configuration pointer 
+				// is invalid. If the effective node configuration is null, this captures
+				// the fact that a null node configuration should produce an empty 
+				// ClassInterfaceOverride. 
+				OutNodeConfigurationUpdates.bDidUpdateClassInterfaceOverride = true;
+				if (const FMetaSoundFrontendNodeConfiguration* ConfigPtr = EffectiveNodeConfig.GetPtr())
 				{
-					// If the class was not found, mark all inputs and outputs as removed.
-					UE_LOG(LogMetaSound, Warning, TEXT("Could not find registered version of interface. %s %s is not registered."), *NodeClassMetadata.GetClassName().ToString(), *NodeClassMetadata.GetVersion().ToString());
-					Algo::Transform(NodeClassInterface.Inputs, OutInterfaceUpdates.RemovedInputs, [&](const FMetasoundFrontendClassInput& Input) { return &Input; });
-					Algo::Transform(NodeClassInterface.Outputs, OutInterfaceUpdates.RemovedOutputs, [&](const FMetasoundFrontendClassOutput& Output) { return &Output; });
-					return false;
+					OutNodeConfigurationUpdates.RegeneratedClassInterfaceOverride = ConfigPtr->OverrideDefaultInterface(InRegisteredClass);
 				}
 			}
+		}
 
-			Algo::Transform(OutInterfaceUpdates.RegistryClass.Interface.Inputs, OutInterfaceUpdates.AddedInputs, [](const FMetasoundFrontendClassInput& Input) { return &Input; });
-			for (const FMetasoundFrontendClassInput& Input : NodeClassInterface.Inputs)
+		const FMetasoundFrontendClassInterface& FBaseNodeController::GetApplicableRegistryInterface(const FMetasoundFrontendClass& InRegisteredClass, const FNodeConfigurationUpdateData& InNodeConfigurationUpdates) const
+		{
+			// Check for the latest class interface overrides
+			if (InNodeConfigurationUpdates.bDidUpdateClassInterfaceOverride)
+			{
+				if (const FMetasoundFrontendClassInterface* Interface = InNodeConfigurationUpdates.RegeneratedClassInterfaceOverride.GetPtr())
+				{
+					return *Interface;
+				}
+			}
+			else
+			{
+				// The class interface override wasn't updated. Use the original if it exists. 
+				if (const FMetasoundFrontendClassInterface* Interface = InNodeConfigurationUpdates.ExistingClassInterfaceOverride.GetPtr())
+				{
+					return *Interface;
+				}
+			}
+
+			// If the override is invalid, use the interface on the class.
+			return InRegisteredClass.GetDefaultInterface();
+		}
+
+		void FBaseNodeController::DiffClassInterfaces(const FMetasoundFrontendClassInterface& InOriginalInterface, const FMetasoundFrontendClassInterface& InNewInterface, FClassInterfaceUpdates& OutClassInterfaceUpdates) const
+		{
+			Algo::Transform(InNewInterface.Inputs, OutClassInterfaceUpdates.AddedInputs, [](const FMetasoundFrontendClassInput& Input) { return &Input; });
+			for (const FMetasoundFrontendClassInput& Input : InOriginalInterface.Inputs)
 			{
 				auto IsEquivalent = [&Input](const FMetasoundFrontendClassInput* RegistryInput)
 				{
 					return FMetasoundFrontendClassInput::IsFunctionalEquivalent(Input, *RegistryInput);
 				};
 
-				const int32 Index = OutInterfaceUpdates.AddedInputs.FindLastByPredicate(IsEquivalent);
+				const int32 Index = OutClassInterfaceUpdates.AddedInputs.FindLastByPredicate(IsEquivalent);
 				if (Index == INDEX_NONE)
 				{
-					OutInterfaceUpdates.RemovedInputs.Add(&Input);
+					OutClassInterfaceUpdates.RemovedInputs.Add(&Input);
 				}
 				else
 				{
-					OutInterfaceUpdates.AddedInputs.RemoveAtSwap(Index, EAllowShrinking::No);
+					OutClassInterfaceUpdates.AddedInputs.RemoveAtSwap(Index, EAllowShrinking::No);
 				}
 			}
 
 
-			Algo::Transform(OutInterfaceUpdates.RegistryClass.Interface.Outputs, OutInterfaceUpdates.AddedOutputs, [](const FMetasoundFrontendClassOutput& Output) { return &Output; });
-			for (const FMetasoundFrontendClassOutput& Output : NodeClassInterface.Outputs)
+			Algo::Transform(InNewInterface.Outputs, OutClassInterfaceUpdates.AddedOutputs, [](const FMetasoundFrontendClassOutput& Output) { return &Output; });
+			for (const FMetasoundFrontendClassOutput& Output : InOriginalInterface.Outputs)
 			{
 				auto IsFunctionalEquivalent = [&Output](const FMetasoundFrontendClassOutput* Iter)
 				{
 					return FMetasoundFrontendClassVertex::IsFunctionalEquivalent(Output, *Iter);
 				};
 
-				const int32 Index = OutInterfaceUpdates.AddedOutputs.FindLastByPredicate(IsFunctionalEquivalent);
+				const int32 Index = OutClassInterfaceUpdates.AddedOutputs.FindLastByPredicate(IsFunctionalEquivalent);
 				if (Index == INDEX_NONE)
 				{
-					OutInterfaceUpdates.RemovedOutputs.Add(&Output);
+					OutClassInterfaceUpdates.RemovedOutputs.Add(&Output);
 				}
 				else
 				{
-					OutInterfaceUpdates.AddedOutputs.RemoveAtSwap(Index, EAllowShrinking::No);
+					OutClassInterfaceUpdates.AddedOutputs.RemoveAtSwap(Index, EAllowShrinking::No);
+				}
+			}
+		}
+
+		bool FBaseNodeController::DiffAgainstRegistryInterface(FClassInterfaceUpdates& OutInterfaceUpdates, bool bInUseHighestMinorVersion, bool bInForceRegenerateClassInterfaceOverride) const
+		{
+			METASOUND_TRACE_CPUPROFILER_EVENT_SCOPE(BaseNodeController::DiffAgainstRegistryInterface);
+
+			OutInterfaceUpdates = FClassInterfaceUpdates();
+
+			const FMetasoundFrontendClassMetadata& NodeClassMetadata = GetClassMetadata();
+			const FMetasoundFrontendClassInterface& NodeClassInterface = GetClassInterface();
+			Metasound::FNodeClassName NodeClassName = NodeClassMetadata.GetClassName().ToNodeClassName();
+
+			bool bDidFindClass = FindRegisteredClass(NodeClassMetadata, bInUseHighestMinorVersion, OutInterfaceUpdates.RegistryClass);
+
+			if (!bDidFindClass)
+			{
+				// The class couldn't be found. Mark all relevant fields as removed. 
+				Algo::Transform(NodeClassInterface.Inputs, OutInterfaceUpdates.RemovedInputs, [&](const FMetasoundFrontendClassInput& Input) { return &Input; });
+				Algo::Transform(NodeClassInterface.Outputs, OutInterfaceUpdates.RemovedOutputs, [&](const FMetasoundFrontendClassOutput& Output) { return &Output; });
+				OutInterfaceUpdates.RemovedConfiguration = GetNodeConfiguration();
+				OutInterfaceUpdates.RemovedClassInterfaceOverride = GetClassInterfaceOverride();
+
+				return false;
+			}
+
+			// Get updates of the node's configuration
+			FNodeConfigurationUpdateData NodeConfigUpdates;
+			FindNodeConfigurationUpdates(OutInterfaceUpdates.RegistryClass, NodeConfigUpdates, bInForceRegenerateClassInterfaceOverride);
+
+			// Based upon all the node configuration updates, determine which class interface is the correct up-to-date interface
+			const FMetasoundFrontendClassInterface& ApplicableRegistryInterface = GetApplicableRegistryInterface(OutInterfaceUpdates.RegistryClass, NodeConfigUpdates);
+
+			// Diff the current interface with the most recent interface
+			DiffClassInterfaces(NodeClassInterface, ApplicableRegistryInterface, OutInterfaceUpdates);
+			const bool bFoundDifferenceInInterface = OutInterfaceUpdates.ContainsRemovedMembers() || OutInterfaceUpdates.ContainsAddedMembers();
+
+			// Add node config and class interface override info if necessary. 
+			if (ShouldReplaceExistingNodeConfig(NodeConfigUpdates.RegisteredConfig, NodeConfigUpdates.ExistingConfig))
+			{
+				OutInterfaceUpdates.AddedConfiguration = MoveTemp(NodeConfigUpdates.RegisteredConfig);
+				OutInterfaceUpdates.RemovedConfiguration = NodeConfigUpdates.ExistingConfig;
+			}
+
+			if (NodeConfigUpdates.bDidUpdateClassInterfaceOverride)
+			{
+				// The class interface override is updated if there was any difference found
+				// in the interface, or the override was added/removed. 
+				const bool bReplaceClassInterfaceOverride = bFoundDifferenceInInterface || (NodeConfigUpdates.ExistingClassInterfaceOverride.IsValid() != NodeConfigUpdates.RegeneratedClassInterfaceOverride.IsValid());
+				if (bReplaceClassInterfaceOverride)
+				{
+					OutInterfaceUpdates.AddedClassInterfaceOverride = MoveTemp(NodeConfigUpdates.RegeneratedClassInterfaceOverride);
+					OutInterfaceUpdates.RemovedClassInterfaceOverride = NodeConfigUpdates.ExistingClassInterfaceOverride;
 				}
 			}
 
@@ -1010,7 +1228,8 @@ namespace Metasound
 			OutInterfaceUpdates = { };
 
 			const FMetasoundFrontendClassMetadata& NodeClassMetadata = GetClassMetadata();
-			if (!IMetaSoundAssetManager::GetChecked().CanAutoUpdate(NodeClassMetadata.GetClassName()))
+			const IMetaSoundAssetManager& AssetManager = IMetaSoundAssetManager::GetChecked();
+			if (AssetManager.IsAssetClass(NodeClassMetadata) && (!AssetManager.CanAutoUpdate(NodeClassMetadata.GetClassName())))
 			{
 				return false;
 			}
@@ -1035,14 +1254,14 @@ namespace Metasound
 				// define changes in native vs asset class definitions. Need to
 				// hash a changeID from natively defined node classes in order to
 				// merge branches.
-				const IMetaSoundAssetManager* AssetManager = IMetaSoundAssetManager::Get();
-				const bool bIsAssetClass = AssetManager && AssetManager->IsAssetClass(RegistryClass.Metadata);
+				
+				const bool bIsAssetClass = AssetManager.IsAssetClass(RegistryClass.Metadata);
 				if (bIsAssetClass)
 				{
 					if (RegistryClass.Metadata.GetChangeID() == NodeClassMetadata.GetChangeID())
 					{
 						const FGuid& NodeClassInterfaceChangeID = GetClassInterface().GetChangeID();
-						if (RegistryClass.Interface.GetChangeID() == NodeClassInterfaceChangeID)
+						if (RegistryClass.GetDefaultInterface().GetChangeID() == NodeClassInterfaceChangeID)
 						{
 							return false;
 						}
@@ -1056,14 +1275,19 @@ namespace Metasound
 					}
 				}
 
-				DiffAgainstRegistryInterface(OutInterfaceUpdates, true /* bUseHighestMinorVersion */);
+				constexpr bool bUseHighestMinorVersion = true;
+				constexpr bool bForceRegenerateClassInterfaceOverride = true;
+				DiffAgainstRegistryInterface(OutInterfaceUpdates, bUseHighestMinorVersion, bForceRegenerateClassInterfaceOverride);
+
 				return OutInterfaceUpdates.ContainsChanges();
 			}
 
 			// 3. Document's class version is out-of-date, so dif and always return true that can auto-update
 			// (Unlike the case where the version is equal, the version must be updated even if the interface
 			// contains no changes).
-			DiffAgainstRegistryInterface(OutInterfaceUpdates, true /* bUseHighestMinorVersion */);
+			constexpr bool bUseHighestMinorVersion = true;
+			constexpr bool bForceRegenerateClassInterfaceOverride = true;
+			DiffAgainstRegistryInterface(OutInterfaceUpdates, bUseHighestMinorVersion, bForceRegenerateClassInterfaceOverride);
 			return true;
 		}
 
diff --git a/Engine/Plugins/Runtime/Metasound/Source/MetasoundFrontend/Private/MetasoundFrontendNodeController.h b/Engine/Plugins/Runtime/Metasound/Source/MetasoundFrontend/Private/MetasoundFrontendNodeController.h
--- a/Engine/Plugins/Runtime/Metasound/Source/MetasoundFrontend/Private/MetasoundFrontendNodeController.h
+++ b/Engine/Plugins/Runtime/Metasound/Source/MetasoundFrontend/Private/MetasoundFrontendNodeController.h
@@ -60,7 +60,10 @@ namespace Metasound
 			virtual void SetNodeStyle(const FMetasoundFrontendNodeStyle& InStyle) override;
 #endif // WITH_EDITORONLY_DATA
 
-			virtual bool DiffAgainstRegistryInterface(FClassInterfaceUpdates& OutInterfaceUpdates, bool bInUseHighestMinorVersion) const override;
+	   		virtual void SetNodeConfiguration(TInstancedStruct<FMetaSoundFrontendNodeConfiguration> InConfig) override;
+	   		virtual TConstStructView<FMetaSoundFrontendNodeConfiguration> GetNodeConfiguration() const override;
+
+			virtual bool DiffAgainstRegistryInterface(FClassInterfaceUpdates& OutInterfaceUpdates, bool bInUseHighestMinorVersion, bool bInForceRegenerateClassInterfaceOverride) const override;
 
 			virtual bool CanAutoUpdate(FClassInterfaceUpdates& OutInterfaceUpdates) const override;
 			virtual FNodeHandle ReplaceWithVersion(const FMetasoundFrontendVersionNumber& InNewVersion, TArray<FVertexNameAndType>* OutDisconnectedInputs, TArray<FVertexNameAndType>* OutDisconnectedOutputs) override;
@@ -155,9 +158,12 @@ namespace Metasound
 
 		protected:
 
+			const FMetasoundFrontendClassInterface* GetClassInterfaceForNode() const;
+
 			virtual FDocumentAccess ShareAccess() override;
 			virtual FConstDocumentAccess ShareAccess() const override;
 
+
 			FNodeAccessPtr NodePtr;
 			FConstClassAccessPtr ClassPtr;
 			FGraphHandle OwningGraph;
@@ -187,6 +193,62 @@ namespace Metasound
 
 			virtual FInputHandle CreateInputController(FGuid InVertexID, FConstVertexAccessPtr InNodeVertexPtr, FConstClassInputAccessPtr InClassInputPtr, FNodeHandle InOwningNode) const = 0;
 			virtual FOutputHandle CreateOutputController(FGuid InVertexID, FConstVertexAccessPtr InNodeVertexPtr, FConstClassOutputAccessPtr InClassOutputPtr, FNodeHandle InOwningNode) const = 0;
+
+		private:
+	   		void SetClassInterfaceOverride(TInstancedStruct<FMetasoundFrontendClassInterface> InClassInterfaceOverride);
+	   		TConstStructView<FMetasoundFrontendClassInterface> GetClassInterfaceOverride() const;
+
+			struct FNodeConfigurationUpdateData
+			{
+				TConstStructView<FMetaSoundFrontendNodeConfiguration> ExistingConfig;
+				TConstStructView<FMetasoundFrontendClassInterface> ExistingClassInterfaceOverride;
+				TInstancedStruct<FMetaSoundFrontendNodeConfiguration> RegisteredConfig;
+				TInstancedStruct<FMetasoundFrontendClassInterface> RegeneratedClassInterfaceOverride;
+				bool bDidUpdateClassInterfaceOverride = false;
+			};
+
+
+			bool FindRegisteredClass(const FMetasoundFrontendClassMetadata& InMetadata, bool bInUseHighestMinorVersion, FMetasoundFrontendClass& OutInterface) const;
+
+			bool ShouldReplaceExistingNodeConfig(TConstStructView<FMetaSoundFrontendNodeConfiguration> InRegisteredNodeConfig, TConstStructView<FMetaSoundFrontendNodeConfiguration> InExistingConfig) const;
+			const FMetasoundFrontendClassInterface& GetApplicableRegistryInterface(const FMetasoundFrontendClass& InRegisteredClass, const FNodeConfigurationUpdateData& InNodeConfigUpdates) const;
+			void FindNodeConfigurationUpdates(const FMetasoundFrontendClass& InRegisteredClass, FNodeConfigurationUpdateData& OutNodeConfigUpdates, bool bInForceRegenerateClassInterfaceOverride) const;
+
+			void DiffClassInterfaces(const FMetasoundFrontendClassInterface& InOriginalInterface, const FMetasoundFrontendClassInterface& InNewInterface, FClassInterfaceUpdates& OutClassInterfaceUpdates) const;
+
+			struct FInputConnectionInfo
+			{
+				FOutputHandle ConnectedOutput;
+				FVertexName Name;
+				FName DataType;
+				FMetasoundFrontendLiteral DefaultValue;
+				bool bLiteralSet = false;
+			};
+
+			struct FOutputConnectionInfo
+			{
+				TArray<FInputHandle> ConnectedInputs;
+				FVertexName VertexName;
+				FName DataType;
+			};
+
+			struct FNodeInstanceReplacementData
+			{
+
+#if WITH_EDITOR
+				FMetasoundFrontendNodeStyle Style;
+#endif // WITH_EDITOR
+
+				TInstancedStruct<FMetaSoundFrontendNodeConfiguration> Configuration;
+				TInstancedStruct<FMetasoundFrontendClassInterface> ClassInterfaceOverride;
+				TMap<FVertexNameAndType, FInputConnectionInfo> InputConnections;
+				TMap<FVertexNameAndType, FOutputConnectionInfo> OutputConnections;
+				FGuid NodeID;
+			};
+
+			FNodeInstanceReplacementData CaptureNodeInstanceReplacementData();
+			void ApplyNodeInstanceReplacementData(const FMetasoundFrontendClass& InReplacementClass, INodeController& InReplacementNode, FNodeInstanceReplacementData&& InInstanceData, TArray<FVertexNameAndType>* OutDisconnectedInputs, TArray<FVertexNameAndType>* OutDisconnectedOutputs);
+
 		};
 
 		/** Represents an external node (defined in either code or by an asset's root graph). */
diff --git a/Engine/Plugins/Runtime/Metasound/Source/MetasoundFrontend/Private/MetasoundFrontendNodeTemplateRegistry.cpp b/Engine/Plugins/Runtime/Metasound/Source/MetasoundFrontend/Private/MetasoundFrontendNodeTemplateRegistry.cpp
--- a/Engine/Plugins/Runtime/Metasound/Source/MetasoundFrontend/Private/MetasoundFrontendNodeTemplateRegistry.cpp
+++ b/Engine/Plugins/Runtime/Metasound/Source/MetasoundFrontend/Private/MetasoundFrontendNodeTemplateRegistry.cpp
@@ -117,9 +117,11 @@ namespace Metasound::Frontend
 			check(Node);
 			if (const FMetasoundFrontendClass* Class = InBuilder.FindDependency(Node->ClassID))
 			{
+				const FMetasoundFrontendClassInterface& ClassInterface = Class->GetInterfaceForNode(*Node);
+
 				const EMetasoundFrontendClassType ClassType = Class->Metadata.GetType();
 				auto MatchesName = [&Vertex](const FMetasoundFrontendClassInput& Input) { return Input.Name == Vertex->Name; };
-				if (const FMetasoundFrontendClassInput* Input = Class->Interface.Inputs.FindByPredicate(MatchesName))
+				if (const FMetasoundFrontendClassInput* Input = ClassInterface.Inputs.FindByPredicate(MatchesName))
 				{
 					return &Input->GetDefaults();
 				}
diff --git a/Engine/Plugins/Runtime/Metasound/Source/MetasoundFrontend/Private/MetasoundFrontendQuery.cpp b/Engine/Plugins/Runtime/Metasound/Source/MetasoundFrontend/Private/MetasoundFrontendQuery.cpp
--- a/Engine/Plugins/Runtime/Metasound/Source/MetasoundFrontend/Private/MetasoundFrontendQuery.cpp
+++ b/Engine/Plugins/Runtime/Metasound/Source/MetasoundFrontend/Private/MetasoundFrontendQuery.cpp
@@ -77,7 +77,7 @@ FMetaSoundClassInterfaceInfo::FMetaSoundClassInterfaceInfo(const IMetaSoundDocum
 				};
 			};
 
-			const FMetasoundFrontendClassInterface& ClassInterface = Document.RootGraph.Interface;
+			const FMetasoundFrontendClassInterface& ClassInterface = Document.RootGraph.GetDefaultInterface();
 			Algo::Transform(ClassInterface.Inputs, Inputs, CopyVertexInfo);
 			Algo::Transform(ClassInterface.Outputs, Outputs, CopyVertexInfo);
 		}
diff --git a/Engine/Plugins/Runtime/Metasound/Source/MetasoundFrontend/Private/MetasoundFrontendQuerySteps.cpp b/Engine/Plugins/Runtime/Metasound/Source/MetasoundFrontend/Private/MetasoundFrontendQuerySteps.cpp
--- a/Engine/Plugins/Runtime/Metasound/Source/MetasoundFrontend/Private/MetasoundFrontendQuerySteps.cpp
+++ b/Engine/Plugins/Runtime/Metasound/Source/MetasoundFrontend/Private/MetasoundFrontendQuerySteps.cpp
@@ -166,7 +166,7 @@ namespace Metasound
 	{
 		check(InEntry.Value.IsType<FMetasoundFrontendClass>());
 
-		return InEntry.Value.Get<FMetasoundFrontendClass>().Interface.Inputs.ContainsByPredicate(
+		return InEntry.Value.Get<FMetasoundFrontendClass>().GetDefaultInterface().Inputs.ContainsByPredicate(
 			[this](const FMetasoundFrontendClassInput& InDesc)
 			{
 				return InDesc.TypeName == InputVertexTypeName;
@@ -181,7 +181,7 @@ namespace Metasound
 
 	bool FFilterClassesByOutputVertexDataType::Filter(const FFrontendQueryEntry& InEntry) const
 	{
-		return InEntry.Value.Get<FMetasoundFrontendClass>().Interface.Outputs.ContainsByPredicate(
+		return InEntry.Value.Get<FMetasoundFrontendClass>().GetDefaultInterface().Outputs.ContainsByPredicate(
 			[this](const FMetasoundFrontendClassOutput& InDesc)
 			{
 				return InDesc.TypeName == OutputVertexTypeName;
diff --git a/Engine/Plugins/Runtime/Metasound/Source/MetasoundFrontend/Private/MetasoundFrontendSubgraphNodeController.cpp b/Engine/Plugins/Runtime/Metasound/Source/MetasoundFrontend/Private/MetasoundFrontendSubgraphNodeController.cpp
--- a/Engine/Plugins/Runtime/Metasound/Source/MetasoundFrontend/Private/MetasoundFrontendSubgraphNodeController.cpp
+++ b/Engine/Plugins/Runtime/Metasound/Source/MetasoundFrontend/Private/MetasoundFrontendSubgraphNodeController.cpp
@@ -194,15 +194,12 @@ namespace Metasound
 		{
 			using namespace FrontendSubgraphNodeControllerPrivate;
 
-			// TODO: will need to update this logic to handle external classes and 
-			// dynamic interfaces. These updates will get piped down to the INodeClass
-			// interface.
 			if (FMetasoundFrontendNode* Node = NodePtr.Get())
 			{
-				if (const FMetasoundFrontendClass* NodeClass = ClassPtr.Get())
+				if (const FMetasoundFrontendClassInterface* ClassInterface = GetClassInterfaceForNode())
 				{
-					ConformNodeVertexArrayToClassVertexArray(Node->Interface.Inputs, NodeClass->Interface.Inputs);
-					ConformNodeVertexArrayToClassVertexArray(Node->Interface.Outputs, NodeClass->Interface.Outputs);
+					ConformNodeVertexArrayToClassVertexArray(Node->Interface.Inputs, ClassInterface->Inputs);
+					ConformNodeVertexArrayToClassVertexArray(Node->Interface.Outputs, ClassInterface->Outputs);
 					// TODO: Conform environment variables by aggregating all dependent environment variables.
 				}
 			}
diff --git a/Engine/Plugins/Runtime/Metasound/Source/MetasoundFrontend/Private/NodeTemplates/MetasoundFrontendNodeTemplateInput.cpp b/Engine/Plugins/Runtime/Metasound/Source/MetasoundFrontend/Private/NodeTemplates/MetasoundFrontendNodeTemplateInput.cpp
--- a/Engine/Plugins/Runtime/Metasound/Source/MetasoundFrontend/Private/NodeTemplates/MetasoundFrontendNodeTemplateInput.cpp
+++ b/Engine/Plugins/Runtime/Metasound/Source/MetasoundFrontend/Private/NodeTemplates/MetasoundFrontendNodeTemplateInput.cpp
@@ -213,9 +213,11 @@ namespace Metasound::Frontend
 		const FMetasoundFrontendNode* ConnectedInputNode = nullptr;
 		if (const FMetasoundFrontendVertex* ConnectedInputOutput = InBuilder.FindNodeOutputConnectedToNodeInput(InNodeID, InVertexID, &ConnectedInputNode, &InPageID))
 		{
+			check(ConnectedInputNode);
 			const FMetasoundFrontendClass* InputClass = InBuilder.FindDependency(ConnectedInputNode->ClassID);
 			check(InputClass);
-			return (InputClass->Interface.Outputs.Last().AccessType);
+
+			return InputClass->GetInterfaceForNode(*ConnectedInputNode).Outputs.Last().AccessType;
 		}
 
 		return EMetasoundFrontendVertexAccessType::Unset;
@@ -285,7 +287,7 @@ namespace Metasound::Frontend
 	{
 		bool bInjectedNodes = false;
 
-		const TArray<FMetasoundFrontendClassInput>& Inputs = InOutBuilder.GetConstDocumentChecked().RootGraph.Interface.Inputs;
+		const TArray<FMetasoundFrontendClassInput>& Inputs = InOutBuilder.GetConstDocumentChecked().RootGraph.GetDefaultInterface().Inputs;
 		for (const FMetasoundFrontendClassInput& Input : Inputs)
 		{
 			// Potentially not used input, which is perfectly valid so early out
@@ -342,4 +344,4 @@ namespace Metasound::Frontend
 	}
 
 #endif // WITH_EDITOR
-} // namespace Metasound::Frontend
\ No newline at end of file
+} // namespace Metasound::Frontend
diff --git a/Engine/Plugins/Runtime/Metasound/Source/MetasoundFrontend/Public/Interfaces/MetasoundFrontendInterface.h b/Engine/Plugins/Runtime/Metasound/Source/MetasoundFrontend/Public/Interfaces/MetasoundFrontendInterface.h
--- a/Engine/Plugins/Runtime/Metasound/Source/MetasoundFrontend/Public/Interfaces/MetasoundFrontendInterface.h
+++ b/Engine/Plugins/Runtime/Metasound/Source/MetasoundFrontend/Public/Interfaces/MetasoundFrontendInterface.h
@@ -10,20 +10,26 @@ namespace Metasound
 	{
 		METASOUNDFRONTEND_API bool IsSubsetOfInterface(const FMetasoundFrontendInterface& InSubsetInterface, const FMetasoundFrontendInterface& InSupersetInterface);
 
+		METASOUNDFRONTEND_API bool IsSubsetOfClassInterface(const FMetasoundFrontendInterface& InSubsetInterface, const FMetasoundFrontendClassInterface& InSupersetClassInterface);
+
+		UE_DEPRECATED(5.6, "Use IsSubsetOfClassInterface instead")
 		METASOUNDFRONTEND_API bool IsSubsetOfClass(const FMetasoundFrontendInterface& InSubsetInterface, const FMetasoundFrontendClass& InSupersetClass);
 
+
 		METASOUNDFRONTEND_API bool IsEquivalentInterface(const FMetasoundFrontendInterface& InInputInterface, const FMetasoundFrontendInterface& InTargetInterface);
 
 		METASOUNDFRONTEND_API bool IsEqualInterface(const FMetasoundFrontendInterface& InInputInterface, const FMetasoundFrontendInterface& InTargetInterface);
 
+		UE_DEPRECATED(5.6, "Use InputOutputDifferenceCount which accepts a FMetasoundFrontendClassInterface instead")
 		METASOUNDFRONTEND_API int32 InputOutputDifferenceCount(const FMetasoundFrontendClass& InClass, const FMetasoundFrontendInterface& InInterface);
 
+		METASOUNDFRONTEND_API int32 InputOutputDifferenceCount(const FMetasoundFrontendClassInterface& InClass, const FMetasoundFrontendInterface& InInterface);
+
 		METASOUNDFRONTEND_API int32 InputOutputDifferenceCount(const FMetasoundFrontendInterface& InInterfaceA, const FMetasoundFrontendInterface& InInterfaceB);
 
 		UE_DEPRECATED(5.6, "All interface functionality is in the MetaSoundFrontendDocumentBuilder API")
 		METASOUNDFRONTEND_API bool DeclaredInterfaceVersionsMatch(const FMetasoundFrontendDocument& InDocumentA, const FMetasoundFrontendDocument& InDocumentB);
 
-		METASOUNDFRONTEND_API void GatherRequiredEnvironmentVariables(const FMetasoundFrontendClass& InClass, TArray<FMetasoundFrontendClassEnvironmentVariable>& OutEnvironmentVariables);
 
 		METASOUNDFRONTEND_API const FMetasoundFrontendInterface* FindMostSimilarInterfaceSupportingEnvironment(const FMetasoundFrontendGraphClass& InRootGraph, const TArray<FMetasoundFrontendClass>& InDependencies, const TArray<FMetasoundFrontendGraphClass>& InSubgraphs, const TArray<FMetasoundFrontendInterface>& InCandidateInterfaces);
 
diff --git a/Engine/Plugins/Runtime/Metasound/Source/MetasoundFrontend/Public/MetasoundFrontend.h b/Engine/Plugins/Runtime/Metasound/Source/MetasoundFrontend/Public/MetasoundFrontend.h
--- a/Engine/Plugins/Runtime/Metasound/Source/MetasoundFrontend/Public/MetasoundFrontend.h
+++ b/Engine/Plugins/Runtime/Metasound/Source/MetasoundFrontend/Public/MetasoundFrontend.h
@@ -83,9 +83,9 @@ namespace Metasound
 		// in a FVertexInterface. 
 		
 		// If called on the game thread, no proxy data cache is needed for object literal creation, but if called from other threads, the proxy data cache must be provided. 
-		FVertexInterface CreateVertexInterfaceFromClass(const FMetasoundFrontendClass& InNodeClass, const FProxyDataCache* InProxyDataCache = nullptr);
+		FVertexInterface CreateDefaultVertexInterfaceFromClass(const FMetasoundFrontendClass& InNodeClass, const FProxyDataCache* InProxyDataCache = nullptr);
 		// Does not create proxies, so no thread restrictions. 
-		FVertexInterface CreateVertexInterfaceFromClassNoProxy(const FMetasoundFrontendClass& InNodeClass);
+		FVertexInterface CreateDefaultVertexInterfaceFromClassNoProxy(const FMetasoundFrontendClass& InNodeClass);
 
 		// Struct that indicates whether an input and an output can be connected,
 		// and whether an intermediate node is necessary to connect the two.
diff --git a/Engine/Plugins/Runtime/Metasound/Source/MetasoundFrontend/Public/MetasoundFrontendController.h b/Engine/Plugins/Runtime/Metasound/Source/MetasoundFrontend/Public/MetasoundFrontendController.h
--- a/Engine/Plugins/Runtime/Metasound/Source/MetasoundFrontend/Public/MetasoundFrontendController.h
+++ b/Engine/Plugins/Runtime/Metasound/Source/MetasoundFrontend/Public/MetasoundFrontendController.h
@@ -11,6 +11,8 @@
 #include "MetasoundFrontendRegistries.h"
 #include "MetasoundGraph.h"
 #include "MetasoundVertex.h"
+#include "StructUtils/InstancedStruct.h"
+#include "StructUtils/StructView.h"
 #include "Templates/Function.h"
 #include "Templates/SharedPointer.h"
 #include "Templates/Tuple.h"
@@ -100,10 +102,15 @@ namespace Metasound
 		  * when querying if a node's class has been updated */
 		struct FClassInterfaceUpdates
 		{
+			FClassInterfaceUpdates() = default;
 			TArray<const FMetasoundFrontendClassInput*> AddedInputs;
 			TArray<const FMetasoundFrontendClassOutput*> AddedOutputs;
 			TArray<const FMetasoundFrontendClassInput*> RemovedInputs;
 			TArray<const FMetasoundFrontendClassOutput*> RemovedOutputs;
+			TInstancedStruct<FMetaSoundFrontendNodeConfiguration> AddedConfiguration;
+			TConstStructView<FMetaSoundFrontendNodeConfiguration> RemovedConfiguration;
+			TInstancedStruct<FMetasoundFrontendClassInterface> AddedClassInterfaceOverride;
+			TConstStructView<FMetasoundFrontendClassInterface> RemovedClassInterfaceOverride;
 
 			bool ContainsRemovedMembers() const
 			{
@@ -115,9 +122,19 @@ namespace Metasound
 				return !AddedInputs.IsEmpty() || !AddedOutputs.IsEmpty();
 			}
 
+			bool ContainsUpdatedNodeConfig() const
+			{
+				return AddedConfiguration.IsValid() || RemovedConfiguration.IsValid();
+			}
+
+			bool ContainsUpdatedNodeClassInterfaceOverride() const
+			{
+				return AddedClassInterfaceOverride.IsValid() || RemovedClassInterfaceOverride.IsValid();
+			}
+
 			bool ContainsChanges() const
 			{
-				return ContainsRemovedMembers() || ContainsAddedMembers();
+				return ContainsRemovedMembers() || ContainsAddedMembers() || ContainsUpdatedNodeConfig() || ContainsUpdatedNodeClassInterfaceOverride();
 			}
 
 			// Cached copy of registry class potentially referenced by added members
@@ -554,7 +571,7 @@ namespace Metasound
 			  * the registry's version of the class interface and that of the node.
 			  * @return Whether or not the interface was found in the registry.
 			  */
-			virtual bool DiffAgainstRegistryInterface(FClassInterfaceUpdates& OutInterfaceUpdates, bool bInUseHighestMinorVersion) const = 0;
+			virtual bool DiffAgainstRegistryInterface(FClassInterfaceUpdates& OutInterfaceUpdates, bool bInUseHighestMinorVersion, bool bInForceRegenerateClassInterfaceOverride=true) const = 0;
 
 			/** Returns whether the node is eligible for auto-updating (i.e.
 			  * has undergone minor revision or the interface has changed, but
@@ -589,6 +606,9 @@ namespace Metasound
 			virtual void SetNodeStyle(const FMetasoundFrontendNodeStyle& InStyle) = 0;
 #endif // WITH_EDITORONLY_DATA
 
+	   		virtual void SetNodeConfiguration(TInstancedStruct<FMetaSoundFrontendNodeConfiguration> InConfig) = 0;
+	   		virtual TConstStructView<FMetaSoundFrontendNodeConfiguration> GetNodeConfiguration() const = 0;
+
 			/** Returns the ID associated with this node. */
 			virtual FGuid GetID() const = 0;
 
diff --git a/Engine/Plugins/Runtime/Metasound/Source/MetasoundFrontend/Public/MetasoundFrontendDocument.h b/Engine/Plugins/Runtime/Metasound/Source/MetasoundFrontend/Public/MetasoundFrontendDocument.h
--- a/Engine/Plugins/Runtime/Metasound/Source/MetasoundFrontend/Public/MetasoundFrontendDocument.h
+++ b/Engine/Plugins/Runtime/Metasound/Source/MetasoundFrontend/Public/MetasoundFrontendDocument.h
@@ -385,6 +385,8 @@
 
 	// Returns true if vertices have equal name & type.
 	static bool IsFunctionalEquivalent(const FMetasoundFrontendVertex& InLHS, const FMetasoundFrontendVertex& InRHS);
+
+	friend METASOUNDFRONTEND_API bool operator==(const FMetasoundFrontendVertex& InLHS, const FMetasoundFrontendVertex& InRHS);
 };
 
 // Pair of guids used to address location of vertex within a FrontendDocument graph
@@ -499,6 +501,9 @@
 	// Create a node interface which satisfies an existing class interface.
 	FMetasoundFrontendNodeInterface(const FMetasoundFrontendClassInterface& InClassInterface);
 
+	void Update(const FMetasoundFrontendClassInterface& InClassInterface);
+	void Update(const FMetasoundFrontendClassInterface& InClassInterface, TFunctionRef<void(const FMetasoundFrontendVertex&)> OnPreRemoveInput, TFunctionRef<void(const FMetasoundFrontendVertex&)> OnPreRemoveOutput);
+
 	// Input vertices to node.
 	UPROPERTY()
 	TArray<FMetasoundFrontendVertex> Inputs;
@@ -518,12 +523,10 @@
 	GENERATED_BODY()
 
 	virtual ~FMetaSoundFrontendNodeConfiguration() = default;
+	// TODO: Markup with special experimental API Stuff
 
 	/* Get the current interface for the class based upon the node extension */
-	virtual FMetasoundFrontendNodeInterface GetNodeInterface(const FMetasoundFrontendClass& InNodeClass) const;
+	virtual TInstancedStruct<FMetasoundFrontendClassInterface> OverrideDefaultInterface(const FMetasoundFrontendClass& InNodeClass) const;
 
-	/* Reassign existing connections if the node's interface has changed */
-	//virtual ReassignConnections(const TMap<FName, FGuid>& OriginalInputs, TMap<FName, FGuid>& NewInputs, OriginalOutputs, NewOutputs);
-
 	/** Provide any data needed by IOperators instantiated from this node. */
 	virtual TSharedPtr<const Metasound::IOperatorData> GetOperatorData() const;
@@ -589,63 +592,6 @@
 };
 
 
-// An FMetasoundFrontendNode represents a single instance of a FMetasoundFrontendClass
-USTRUCT()
-struct METASOUNDFRONTEND_API FMetasoundFrontendNode
-{
-	GENERATED_BODY()
-
-	FMetasoundFrontendNode() = default;
-
-	UE_DEPRECATED(5.6, "Please use constructor which accepts a node extension.")
-	FMetasoundFrontendNode(const FMetasoundFrontendClass& InClass);
-
-	// Construct a node from a node class and an optional node extension 
-	FMetasoundFrontendNode(const FMetasoundFrontendClass& InClass, TInstancedStruct<FMetaSoundFrontendNodeConfiguration> InConfiguration);
-
-private:
-	// Unique ID of this node.
-	UPROPERTY()
-	FGuid ID = Metasound::FrontendInvalidID;
-
-public:
-	// ID of FMetasoundFrontendClass corresponding to this node.
-	UPROPERTY()
-	FGuid ClassID = Metasound::FrontendInvalidID;
-
-	// Name of node instance.
-	UPROPERTY()
-	FName Name;
-
-	// Interface of node instance.
-	UPROPERTY()
-	FMetasoundFrontendNodeInterface Interface;
-
-	// Default values for node inputs.
-	UPROPERTY()
-	TArray<FMetasoundFrontendVertexLiteral> InputLiterals;
-
-	UPROPERTY()
-	TInstancedStruct<FMetaSoundFrontendNodeConfiguration> Configuration;
-
-#if WITH_EDITORONLY_DATA
-	// Style info related to a node.
-	UPROPERTY()
-	FMetasoundFrontendNodeStyle Style;
-#endif // WITH_EDITORONLY_DATA
-
-	const FGuid& GetID() const
-	{
-		return ID;
-	}
-
-	void UpdateID(const FGuid& InNewGuid)
-	{
-		ID = InNewGuid;
-	}
-};
-
-
 // Represents a single connection from one point to another.
 USTRUCT()
 struct FMetasoundFrontendEdge
@@ -752,32 +698,6 @@
 #endif // WITH_EDITORONLY_DATA
 };
 
-USTRUCT()
-struct METASOUNDFRONTEND_API FMetasoundFrontendGraph
-{
-	GENERATED_BODY()
 
-	// Node contained in graph
-	UPROPERTY()
-	TArray<FMetasoundFrontendNode> Nodes;
-
-	// Connections between points on nodes.
-	UPROPERTY()
-	TArray<FMetasoundFrontendEdge> Edges;
-
-	// Graph local variables.
-	UPROPERTY()
-	TArray<FMetasoundFrontendVariable> Variables;
-
-#if WITH_EDITORONLY_DATA
-	// Style of graph display.
-	UPROPERTY()
-	FMetasoundFrontendGraphStyle Style;
-#endif // WITH_EDITORONLY_DATA
-
-	UPROPERTY()
-	FGuid PageID;
-};
-
 // Metadata associated with a vertex.
 USTRUCT()
@@ -970,6 +890,7 @@
 	FMetasoundFrontendClassInputDefault(const FAudioParameter& InParameter);
 
 	static bool IsFunctionalEquivalent(const FMetasoundFrontendClassInputDefault& InLHS, const FMetasoundFrontendClassInputDefault& InRHS);
+	friend bool operator==(const FMetasoundFrontendClassInputDefault& InLHS, const FMetasoundFrontendClassInputDefault& InRHS);
 
 	UPROPERTY()
 	FMetasoundFrontendLiteral Literal;
@@ -1715,7 +1636,14 @@
 {
 	GENERATED_BODY()
 
+	PRAGMA_DISABLE_DEPRECATION_WARNINGS
+	FMetasoundFrontendClass() = default;
 	virtual ~FMetasoundFrontendClass() = default;
+	FMetasoundFrontendClass(const FMetasoundFrontendClass&) = default;
+	FMetasoundFrontendClass(FMetasoundFrontendClass&&) = default;
+	FMetasoundFrontendClass& operator=(const FMetasoundFrontendClass&) = default;
+	FMetasoundFrontendClass& operator=(FMetasoundFrontendClass&&) = default;
+	PRAGMA_ENABLE_DEPRECATION_WARNINGS
 
 	UPROPERTY()
 	FGuid ID = Metasound::FrontendInvalidID;
@@ -1723,9 +1651,16 @@
 	UPROPERTY(EditAnywhere, Category = CustomView)
 	FMetasoundFrontendClassMetadata Metadata;
 
+	UE_DEPRECATED(5.6, "Use Get/Set Default Interface instead")
 	UPROPERTY(EditAnywhere, Category = CustomView)
 	FMetasoundFrontendClassInterface Interface;
 
+
+	void SetDefaultInterface(const FMetasoundFrontendClassInterface& InInterface);
+	FMetasoundFrontendClassInterface& GetDefaultInterface();
+	const FMetasoundFrontendClassInterface& GetDefaultInterface() const;
+	const FMetasoundFrontendClassInterface& GetInterfaceForNode(const FMetasoundFrontendNode& InNode) const;
+
 #if WITH_EDITORONLY_DATA
 	UPROPERTY()
 	FMetasoundFrontendClassStyle Style;
@@ -1743,6 +1678,67 @@
 #endif // WITH_EDITOR
 };
 
+// An FMetasoundFrontendNode represents a single instance of a FMetasoundFrontendClass
+USTRUCT()
+struct METASOUNDFRONTEND_API FMetasoundFrontendNode
+{
+	GENERATED_BODY()
+
+	FMetasoundFrontendNode() = default;
+
+	UE_DEPRECATED(5.6, "Please use constructor which accepts a node extension.")
+	FMetasoundFrontendNode(const FMetasoundFrontendClass& InClass);
+
+	// Construct a node from a node class and an optional node extension 
+	FMetasoundFrontendNode(const FMetasoundFrontendClass& InClass, TInstancedStruct<FMetaSoundFrontendNodeConfiguration> InConfiguration);
+
+private:
+	// Unique ID of this node.
+	UPROPERTY()
+	FGuid ID = Metasound::FrontendInvalidID;
+
+public:
+	// ID of FMetasoundFrontendClass corresponding to this node.
+	UPROPERTY()
+	FGuid ClassID = Metasound::FrontendInvalidID;
+
+	// Name of node instance.
+	UPROPERTY()
+	FName Name;
+
+	// Interface of node instance.
+	UPROPERTY()
+	FMetasoundFrontendNodeInterface Interface;
+
+	// Default values for node inputs.
+	UPROPERTY()
+	TArray<FMetasoundFrontendVertexLiteral> InputLiterals;
+
+	// Instance of a configuration for this node.
+	UPROPERTY()
+	TInstancedStruct<FMetaSoundFrontendNodeConfiguration> Configuration;
+
+	// An optional override to the default class interface.
+	UPROPERTY()
+	TInstancedStruct<FMetasoundFrontendClassInterface> ClassInterfaceOverride;
+
+#if WITH_EDITORONLY_DATA
+	// Style info related to a node.
+	UPROPERTY()
+	FMetasoundFrontendNodeStyle Style;
+#endif // WITH_EDITORONLY_DATA
+
+	const FGuid& GetID() const
+	{
+		return ID;
+	}
+
+	void UpdateID(const FGuid& InNewGuid)
+	{
+		ID = InNewGuid;
+	}
+};
+
 // Preset options related to a parent graph class.  A graph class with bIsPreset set to true
 // auto-updates to mirror the interface members (inputs & outputs) of the single, referenced
 // node. It also connects all of these nodes' interface members on update to corresponding inputs
@@ -1764,5 +1760,32 @@
 };
 
 USTRUCT()
+struct METASOUNDFRONTEND_API FMetasoundFrontendGraph
+{
+	GENERATED_BODY()
+
+	// Node contained in graph
+	UPROPERTY()
+	TArray<FMetasoundFrontendNode> Nodes;
+
+	// Connections between points on nodes.
+	UPROPERTY()
+	TArray<FMetasoundFrontendEdge> Edges;
+
+	// Graph local variables.
+	UPROPERTY()
+	TArray<FMetasoundFrontendVariable> Variables;
+
+#if WITH_EDITORONLY_DATA
+	// Style of graph display.
+	UPROPERTY()
+	FMetasoundFrontendGraphStyle Style;
+#endif // WITH_EDITORONLY_DATA
+
+	UPROPERTY()
+	FGuid PageID;
+};
+
+USTRUCT()
 struct METASOUNDFRONTEND_API FMetasoundFrontendGraphClass : public FMetasoundFrontendClass
 {
@@ -1980,6 +2003,9 @@
 	/** Execute the provided function for each literal on a FMetasoundFrontendClass.*/
 	METASOUNDFRONTEND_API void ForEachLiteral(const FMetasoundFrontendClass& InClass, FForEachLiteralFunctionRef OnLiteral);
 
+	METASOUNDFRONTEND_API void ForEachLiteral(const FMetasoundFrontendClassInterface& InClassInterface, FForEachLiteralFunctionRef OnLiteral);
+
+
 	/** Execute the provided function for each literal on a FMetasoundFrontendNode.*/
 	METASOUNDFRONTEND_API void ForEachLiteral(const FMetasoundFrontendNode& InNode, FForEachLiteralFunctionRef OnLiteral);
 } // namespace Metasound::Frontend
