MetaSound CIS fix for Android

diff --git a/Engine/Plugins/Runtime/Metasound/Source/MetasoundFrontend/Private/MetasoundFrontendDocument.cpp b/Engine/Plugins/Runtime/Metasound/Source/MetasoundFrontend/Private/MetasoundFrontendDocument.cpp
--- a/Engine/Plugins/Runtime/Metasound/Source/MetasoundFrontend/Private/MetasoundFrontendDocument.cpp
+++ b/Engine/Plugins/Runtime/Metasound/Source/MetasoundFrontend/Private/MetasoundFrontendDocument.cpp
@@ -370,7 +370,7 @@ void FMetasoundFrontendNodeInterface::Update(const FMetasoundFrontendClassInterf
 	Algo::Transform(Inputs, UnmatchedVertices, [](const FMetasoundFrontendVertex& Vertex) { return &Vertex; });
 	for (const FMetasoundFrontendClassInput& ClassInput : InClassInterface.Inputs)
 	{
-		if (const int32 Index = UnmatchedVertices.IndexOfByPredicate(FFindMatchingVertex(ClassInput)); Index != INDEX_NONE)
+		if (const int32 Index = UnmatchedVertices.IndexOfByPredicate(FFindMatchingVertex{ClassInput}); Index != INDEX_NONE)
 		{
 			// Update the node vertex with anything new from the class vertex
 			UnmatchedVertices.RemoveAtSwap(Index, 1, EAllowShrinking::No);
@@ -395,7 +395,7 @@ void FMetasoundFrontendNodeInterface::Update(const FMetasoundFrontendClassInterf
 	Algo::Transform(Outputs, UnmatchedVertices, [](const FMetasoundFrontendVertex& Vertex) { return &Vertex; });
 	for (const FMetasoundFrontendClassOutput& ClassOutput : InClassInterface.Outputs)
 	{
-		if (const int32 Index = UnmatchedVertices.IndexOfByPredicate(FFindMatchingVertex(ClassOutput)); Index != INDEX_NONE)
+		if (const int32 Index = UnmatchedVertices.IndexOfByPredicate(FFindMatchingVertex{ClassOutput}); Index != INDEX_NONE)
 		{
 			// Update the node vertex with anything new from the class vertex
 			UnmatchedVertices.RemoveAtSwap(Index, 1, EAllowShrinking::No);
