Fix MetaSound construction non-determinism with root graphs having random guid

#rnx

diff --git a/Engine/Plugins/Runtime/Metasound/Source/MetasoundEngine/Private/Metasound.cpp b/Engine/Plugins/Runtime/Metasound/Source/MetasoundEngine/Private/Metasound.cpp
--- a/Engine/Plugins/Runtime/Metasound/Source/MetasoundEngine/Private/Metasound.cpp
+++ b/Engine/Plugins/Runtime/Metasound/Source/MetasoundEngine/Private/Metasound.cpp
@@ -66,6 +66,11 @@ UMetaSoundPatch::UMetaSoundPatch(const FObjectInitializer& ObjectInitializer)
 	: Super(ObjectInitializer)
 	, FMetasoundAssetBase()
 {
+	// Default Root Graph uses static ID to distinguish between a default constructed document
+	// (invalid ID) and CDO.  A MetaSoundSource asset should only be constructed using the Document
+	// Builder API to avoid ID collisions, but underlying UObjects must always be deterministically
+	// generated using NewObject for serialization (and for CDOs).
+	RootMetaSoundDocument.RootGraph.ID = FGuid(0x4d657461, 0x536f756e, 0x64506174, 0x63680000);
 }
 
 Metasound::Frontend::FDocumentAccessPtr UMetaSoundPatch::GetDocumentAccessPtr()
diff --git a/Engine/Plugins/Runtime/Metasound/Source/MetasoundEngine/Private/MetasoundSource.cpp b/Engine/Plugins/Runtime/Metasound/Source/MetasoundEngine/Private/MetasoundSource.cpp
--- a/Engine/Plugins/Runtime/Metasound/Source/MetasoundEngine/Private/MetasoundSource.cpp
+++ b/Engine/Plugins/Runtime/Metasound/Source/MetasoundEngine/Private/MetasoundSource.cpp
@@ -415,6 +415,12 @@ UMetaSoundSource::UMetaSoundSource(const FObjectInitializer& ObjectInitializer)
 	: Super(ObjectInitializer)
 	, FMetasoundAssetBase()
 {
+	// Default Root Graph uses static ID to distinguish between a default constructed document
+	// (invalid ID) and CDO.  A MetaSoundSource asset should only be constructed using the Document
+	// Builder API to avoid ID collisions, but underlying UObjects must always be deterministically
+	// generated using NewObject for serialization (and for CDOs).
+	RootMetasoundDocument.RootGraph.ID = FGuid(0x4D657461, 0x536F756E, 0x64536F75, 0x72636500);
+
 	bProcedural = true;
 	bRequiresStopFade = true;
 	NumChannels = 1;
diff --git a/Engine/Plugins/Runtime/Metasound/Source/MetasoundEngineTest/Private/EngineTestMetasoundSource.cpp b/Engine/Plugins/Runtime/Metasound/Source/MetasoundEngineTest/Private/EngineTestMetasoundSource.cpp
--- a/Engine/Plugins/Runtime/Metasound/Source/MetasoundEngineTest/Private/EngineTestMetasoundSource.cpp
+++ b/Engine/Plugins/Runtime/Metasound/Source/MetasoundEngineTest/Private/EngineTestMetasoundSource.cpp
@@ -183,6 +183,8 @@ namespace EngineTestMetaSoundSourcePrivate
 		// uses a weird mix of builder API & controller API calls.
 
 		FMetasoundFrontendDocument Document;
+		Document.RootGraph.ID = FGuid::NewGuid();
+
 		FMetaSoundFrontendDocumentBuilder::InitGraphClassMetadata(Document.RootGraph.Metadata);
 		Document.RootGraph.InitDefaultGraphPage();
 
diff --git a/Engine/Plugins/Runtime/Metasound/Source/MetasoundEngineTest/Private/NodeTestGraphBuilder.cpp b/Engine/Plugins/Runtime/Metasound/Source/MetasoundEngineTest/Private/NodeTestGraphBuilder.cpp
--- a/Engine/Plugins/Runtime/Metasound/Source/MetasoundEngineTest/Private/NodeTestGraphBuilder.cpp
+++ b/Engine/Plugins/Runtime/Metasound/Source/MetasoundEngineTest/Private/NodeTestGraphBuilder.cpp
@@ -12,6 +12,7 @@ namespace Metasound::Test
 
 	FNodeTestGraphBuilder::FNodeTestGraphBuilder()
 	{
+		Document.RootGraph.ID = FGuid::NewGuid();
 		Document.RootGraph.Metadata.SetClassName({ "Metasound", "TestNodes", *LexToString(FGuid::NewGuid()) });
 		Document.RootGraph.Metadata.SetType(EMetasoundFrontendClassType::Graph);
 		Document.RootGraph.InitDefaultGraphPage();
diff --git a/Engine/Plugins/Runtime/Metasound/Source/MetasoundFrontend/Private/MetasoundFrontendDocument.cpp b/Engine/Plugins/Runtime/Metasound/Source/MetasoundFrontend/Private/MetasoundFrontendDocument.cpp
--- a/Engine/Plugins/Runtime/Metasound/Source/MetasoundFrontend/Private/MetasoundFrontendDocument.cpp
+++ b/Engine/Plugins/Runtime/Metasound/Source/MetasoundFrontend/Private/MetasoundFrontendDocument.cpp
@@ -1447,7 +1447,6 @@ FMetasoundFrontendVersionNumber FMetasoundFrontendDocument::GetMaxVersion()
 
 FMetasoundFrontendDocument::FMetasoundFrontendDocument()
 {
-	RootGraph.ID = FGuid::NewGuid();
 	RootGraph.Metadata.SetType(EMetasoundFrontendClassType::Graph);
 
 #if WITH_EDITORONLY_DATA
diff --git a/Engine/Plugins/Runtime/Metasound/Source/MetasoundFrontend/Private/MetasoundFrontendDocumentBuilder.cpp b/Engine/Plugins/Runtime/Metasound/Source/MetasoundFrontend/Private/MetasoundFrontendDocumentBuilder.cpp
--- a/Engine/Plugins/Runtime/Metasound/Source/MetasoundFrontend/Private/MetasoundFrontendDocumentBuilder.cpp
+++ b/Engine/Plugins/Runtime/Metasound/Source/MetasoundFrontend/Private/MetasoundFrontendDocumentBuilder.cpp
@@ -708,6 +708,11 @@ namespace Metasound::Frontend
 	}
 } // namespace Metasound::Frontend
 
+UMetaSoundBuilderDocument::UMetaSoundBuilderDocument(const FObjectInitializer& ObjectInitializer)
+{
+	Document.RootGraph.ID = FGuid::NewGuid();
+}
+
 UMetaSoundBuilderDocument& UMetaSoundBuilderDocument::Create(const UClass& InMetaSoundUClass)
 {
 	UMetaSoundBuilderDocument* DocObject = NewObject<UMetaSoundBuilderDocument>();
diff --git a/Engine/Plugins/Runtime/Metasound/Source/MetasoundFrontend/Public/MetasoundFrontendDocumentBuilder.h b/Engine/Plugins/Runtime/Metasound/Source/MetasoundFrontend/Public/MetasoundFrontendDocumentBuilder.h
--- a/Engine/Plugins/Runtime/Metasound/Source/MetasoundFrontend/Public/MetasoundFrontendDocumentBuilder.h
+++ b/Engine/Plugins/Runtime/Metasound/Source/MetasoundFrontend/Public/MetasoundFrontendDocumentBuilder.h
@@ -98,6 +98,8 @@ class METASOUNDFRONTEND_API UMetaSoundBuilderDocument : public UObject, public I
 	GENERATED_BODY()
 
 public:
+	UMetaSoundBuilderDocument(const FObjectInitializer& ObjectInitializer);
+
 	UE_DEPRECATED(5.5, "Use overload supplying MetaSound to copy (builder documents no longer supported for cases outside of cloned document registration.")
 	static UMetaSoundBuilderDocument& Create(const UClass& InBuilderClass);
 
