- Fix for template nodes not being compiled out in editor scripts attempting to generate transient MetaSounds
- Removal of deprecated MS interface API functions
- Minor clean-up & comment clarifications

[RN] Fix for bug where creating transient MetaSound Presets from within the editor would cause registration errors and fail to be executable.

diff --git a/Engine/Plugins/Runtime/Metasound/Source/MetasoundEngine/Private/MetasoundBuilderSubsystem.cpp b/Engine/Plugins/Runtime/Metasound/Source/MetasoundEngine/Private/MetasoundBuilderSubsystem.cpp
--- a/Engine/Plugins/Runtime/Metasound/Source/MetasoundEngine/Private/MetasoundBuilderSubsystem.cpp
+++ b/Engine/Plugins/Runtime/Metasound/Source/MetasoundEngine/Private/MetasoundBuilderSubsystem.cpp
@@ -1010,16 +1010,20 @@ UMetaSoundPatchBuilder* UMetaSoundBuilderSubsystem::CreatePatchPresetBuilder(FNa
 	return nullptr;
 }
 
-UMetaSoundBuilderBase& UMetaSoundBuilderSubsystem::CreatePresetBuilder(FName BuilderName, const TScriptInterface<IMetaSoundDocumentInterface>& ReferencedPatchClass, EMetaSoundBuilderResult& OutResult)
+UMetaSoundBuilderBase& UMetaSoundBuilderSubsystem::CreatePresetBuilder(FName BuilderName, const TScriptInterface<IMetaSoundDocumentInterface>& ReferencedMetaSound, EMetaSoundBuilderResult& OutResult)
 {
-	const UClass& Class = ReferencedPatchClass->GetBaseMetaSoundUClass();
+	check(ReferencedMetaSound.GetObject());
+
+	const UClass& Class = ReferencedMetaSound->GetBaseMetaSoundUClass();
 	if (&Class == UMetaSoundSource::StaticClass())
 	{
-		return *CreateSourcePresetBuilder(BuilderName, ReferencedPatchClass, OutResult);
+		UMetaSoundSource* Source = CastChecked<UMetaSoundSource>(ReferencedMetaSound.GetObject());
+		return *CreateSourcePresetBuilder(BuilderName, Source, OutResult);
 	}
 	else if (&Class == UMetaSoundPatch::StaticClass())
 	{
-		return *CreatePatchPresetBuilder(BuilderName, ReferencedPatchClass, OutResult);
+		UMetaSoundPatch* Patch = CastChecked<UMetaSoundPatch>(ReferencedMetaSound.GetObject());
+		return *CreatePatchPresetBuilder(BuilderName, Patch, OutResult);
 	}
 	else
 	{
diff --git a/Engine/Plugins/Runtime/Metasound/Source/MetasoundEngine/Public/MetasoundBuilderSubsystem.h b/Engine/Plugins/Runtime/Metasound/Source/MetasoundEngine/Public/MetasoundBuilderSubsystem.h
--- a/Engine/Plugins/Runtime/Metasound/Source/MetasoundEngine/Public/MetasoundBuilderSubsystem.h
+++ b/Engine/Plugins/Runtime/Metasound/Source/MetasoundEngine/Public/MetasoundBuilderSubsystem.h
@@ -191,13 +191,15 @@ public:
 		EMetaSoundOutputAudioFormat OutputFormat = EMetaSoundOutputAudioFormat::Mono,
 		bool bIsOneShot = true);
 
+	// Creates a builder that acts on a generated, transient preset MetaSound UObject. Fails if Referenced MetaSound is not a UMetaSoundPatch.
 	UFUNCTION(BlueprintCallable, Category = "Audio|MetaSound|Builder",  meta = (ExpandEnumAsExecs = "OutResult"))
-	UPARAM(DisplayName = "Patch Preset Builder") UMetaSoundPatchBuilder* CreatePatchPresetBuilder(FName BuilderName, const TScriptInterface<IMetaSoundDocumentInterface>& ReferencedPatchClass, EMetaSoundBuilderResult& OutResult);
+	UPARAM(DisplayName = "Patch Preset Builder") UMetaSoundPatchBuilder* CreatePatchPresetBuilder(FName BuilderName, UPARAM(DisplayName = "Reference Patch")  const TScriptInterface<IMetaSoundDocumentInterface>& ReferencedPatchClass, EMetaSoundBuilderResult& OutResult);
 
-	UMetaSoundBuilderBase& CreatePresetBuilder(FName BuilderName, const TScriptInterface<IMetaSoundDocumentInterface>& ReferencedPatchClass, EMetaSoundBuilderResult& OutResult);
+	UMetaSoundBuilderBase& CreatePresetBuilder(FName BuilderName, const TScriptInterface<IMetaSoundDocumentInterface>& ReferencedMetaSound, EMetaSoundBuilderResult& OutResult);
 	
+	// Creates a builder that acts on a generated, transient preset MetaSound UObject. Fails if Referenced MetaSound is not a UMetaSoundSource.
 	UFUNCTION(BlueprintCallable, Category = "Audio|MetaSound|Builder", meta = (ExpandEnumAsExecs = "OutResult"))
-	UPARAM(DisplayName = "Source Preset Builder") UMetaSoundSourceBuilder* CreateSourcePresetBuilder(FName BuilderName, const TScriptInterface<IMetaSoundDocumentInterface>& ReferencedSourceClass, EMetaSoundBuilderResult& OutResult);
+	UPARAM(DisplayName = "Source Preset Builder") UMetaSoundSourceBuilder* CreateSourcePresetBuilder(FName BuilderName, UPARAM(DisplayName = "Reference Source") const TScriptInterface<IMetaSoundDocumentInterface>& ReferencedSourceClass, EMetaSoundBuilderResult& OutResult);
 
 	UFUNCTION(BlueprintCallable, Category = "Audio|MetaSound|Builder", meta = (DisplayName = "Create MetaSound Bool Literal"))
 	UPARAM(DisplayName = "Bool Literal") FMetasoundFrontendLiteral CreateBoolMetaSoundLiteral(bool Value, FName& DataType);
diff --git a/Engine/Plugins/Runtime/Metasound/Source/MetasoundFrontend/Private/MetasoundFrontendNodeClassRegistry.cpp b/Engine/Plugins/Runtime/Metasound/Source/MetasoundFrontend/Private/MetasoundFrontendNodeClassRegistry.cpp
--- a/Engine/Plugins/Runtime/Metasound/Source/MetasoundFrontend/Private/MetasoundFrontendNodeClassRegistry.cpp
+++ b/Engine/Plugins/Runtime/Metasound/Source/MetasoundFrontend/Private/MetasoundFrontendNodeClassRegistry.cpp
@@ -58,14 +58,9 @@ namespace Metasound::Frontend
 			const FMetasoundFrontendDocument& Document = DocumentInterface->GetConstDocument();
 
 #if WITH_EDITOR
-			// Only assets require template node processing
-			if (DocObject->IsAsset())
-			{
-				DocumentInterface = &UMetaSoundBuilderDocument::Create(*DocumentInterface.GetInterface());
-				FMetaSoundFrontendDocumentBuilder Builder(DocumentInterface);
-				Builder.TransformTemplateNodes();
-			}
-
+			DocumentInterface = &UMetaSoundBuilderDocument::Create(*DocumentInterface.GetInterface());
+			FMetaSoundFrontendDocumentBuilder Builder(DocumentInterface);
+			Builder.TransformTemplateNodes();
 			return DocumentInterface;
 #else
 			const bool bIsBuilding = DocumentInterface->IsActivelyBuilding();
diff --git a/Engine/Plugins/Runtime/Metasound/Source/MetasoundFrontend/Private/MetasoundFrontendNodeTemplateRegistry.cpp b/Engine/Plugins/Runtime/Metasound/Source/MetasoundFrontend/Private/MetasoundFrontendNodeTemplateRegistry.cpp
--- a/Engine/Plugins/Runtime/Metasound/Source/MetasoundFrontend/Private/MetasoundFrontendNodeTemplateRegistry.cpp
+++ b/Engine/Plugins/Runtime/Metasound/Source/MetasoundFrontend/Private/MetasoundFrontendNodeTemplateRegistry.cpp
@@ -104,11 +104,6 @@ namespace Metasound::Frontend
 		return Registry;
 	}
 
-	TUniquePtr<INodeTransform> INodeTemplate::GenerateNodeTransform(FMetasoundFrontendDocument& InDocument) const
-	{
-		return nullptr;
-	}
-
 	const TArray<FMetasoundFrontendClassInputDefault>* FNodeTemplateBase::FindNodeClassInputDefaults(const FMetaSoundFrontendDocumentBuilder& InBuilder, const FGuid& InPageID, const FGuid& InNodeID, FName VertexName) const
 	{
 		if (const FMetasoundFrontendVertex* Vertex = InBuilder.FindNodeInput(InNodeID, VertexName, &InPageID))
@@ -187,18 +182,6 @@ namespace Metasound::Frontend
 		static_cast<FNodeTemplateRegistry&>(INodeTemplateRegistry::Get()).Register(MoveTemp(InTemplate));
 	}
 
-	void UnregisterNodeTemplate(const FMetasoundFrontendVersion& InVersion)
-	{
-		FMetasoundFrontendClassName ClassName;
-		FMetasoundFrontendClassName::Parse(InVersion.Name.ToString(), ClassName);
-		const FNodeClassRegistryKey Key = FNodeClassRegistryKey(EMetasoundFrontendClassType::Template, ClassName, InVersion.Number);
-		if (ensure(Key.IsValid()))
-		{
-			FNodeClassRegistry::Get().UnregisterNodeTemplate(Key);
-			static_cast<FNodeTemplateRegistry&>(INodeTemplateRegistry::Get()).Unregister(Key);
-		}
-	}
-
 	void UnregisterNodeTemplate(const FMetasoundFrontendClassName& InClassName, const FMetasoundFrontendVersionNumber& InTemplateVersion)
 	{
 		const FNodeClassRegistryKey Key = FNodeClassRegistryKey(EMetasoundFrontendClassType::Template, InClassName, InTemplateVersion);
diff --git a/Engine/Plugins/Runtime/Metasound/Source/MetasoundFrontend/Public/MetasoundFrontendNodeTemplateRegistry.h b/Engine/Plugins/Runtime/Metasound/Source/MetasoundFrontend/Public/MetasoundFrontendNodeTemplateRegistry.h
--- a/Engine/Plugins/Runtime/Metasound/Source/MetasoundFrontend/Public/MetasoundFrontendNodeTemplateRegistry.h
+++ b/Engine/Plugins/Runtime/Metasound/Source/MetasoundFrontend/Public/MetasoundFrontendNodeTemplateRegistry.h
@@ -73,9 +73,6 @@ namespace Metasound::Frontend
 		// Returns note template class name.
 		virtual const FMetasoundFrontendClassName& GetClassName() const = 0;
 
-		UE_DEPRECATED(5.4, "Use version that does not provide a preprocessed document")
-		virtual TUniquePtr<INodeTransform> GenerateNodeTransform(FMetasoundFrontendDocument& InDocument) const;
-
 #if WITH_EDITOR
 		virtual FText GetNodeDisplayName(const IMetaSoundDocumentInterface& Interface, const FGuid& InPageID, const FGuid& InNodeID) const = 0;
 
@@ -104,9 +101,6 @@ namespace Metasound::Frontend
 		// Returns access type of the given output within the provided builder's document
 		virtual EMetasoundFrontendVertexAccessType GetNodeOutputAccessType(const FMetaSoundFrontendDocumentBuilder& InBuilder, const FGuid& InPageID, const FGuid& InNodeID, const FGuid& InVertexID) const = 0;
 
-		UE_DEPRECATED(5.4, "Use version number or classname instead")
-		virtual const FMetasoundFrontendVersion& GetVersion() const { const static FMetasoundFrontendVersion NullVersion; return NullVersion; }
-
 		// Returns note template class version.
 		virtual const FMetasoundFrontendVersionNumber& GetVersionNumber() const = 0;
 
@@ -164,12 +158,9 @@ namespace Metasound::Frontend
 #endif // WITH_EDITOR
 	};
 
-	// Register & Unregister are not publicly accessible implementation as the API
-	// is in beta and, currently, only to be used by internal implementation (ex. reroute nodes).
+	// Register/UnregisterNodeTemplate are limited to internal implementation to avoid document corruption. Node Config is
+	// likely the better option for most class interface "switching" applications (templates are reserved for cook-only
+	// applications/optimizations with more advanced node instance interface manipulation).
 	void RegisterNodeTemplate(TUniquePtr<INodeTemplate>&& InTemplate);
-
-	UE_DEPRECATED(5.4, "Use version that provides class name and version instead")
-	void UnregisterNodeTemplate(const FMetasoundFrontendVersion& InNodeTemplateVersion);
-
 	void UnregisterNodeTemplate(const FMetasoundFrontendClassName& InClassName, const FMetasoundFrontendVersionNumber& InTemplateVersion);
 } // namespace Metasound::Frontend
diff --git a/Engine/Plugins/Runtime/Metasound/Source/MetasoundFrontend/Public/NodeTemplates/MetasoundFrontendNodeTemplateReroute.h b/Engine/Plugins/Runtime/Metasound/Source/MetasoundFrontend/Public/NodeTemplates/MetasoundFrontendNodeTemplateReroute.h
--- a/Engine/Plugins/Runtime/Metasound/Source/MetasoundFrontend/Public/NodeTemplates/MetasoundFrontendNodeTemplateReroute.h
+++ b/Engine/Plugins/Runtime/Metasound/Source/MetasoundFrontend/Public/NodeTemplates/MetasoundFrontendNodeTemplateReroute.h
@@ -32,9 +32,6 @@ namespace Metasound::Frontend
 		virtual FText GetNodeDisplayName(const IMetaSoundDocumentInterface& DocumentInterface, const FGuid& InPageID, const FGuid& InNodeID) const override;
 #endif // WITH_EDITOR
 
-		UE_DEPRECATED(5.4, "Use version that does not require mutating a provided PreprocessedDocument")
-		virtual TUniquePtr<INodeTransform> GenerateNodeTransform(FMetasoundFrontendDocument& InPreprocessedDocument) const override { return { }; }
-
 		virtual FMetasoundFrontendNodeInterface GenerateNodeInterface(FNodeTemplateGenerateInterfaceParams InParams) const override;
 		virtual TUniquePtr<INodeTemplateTransform> GenerateNodeTransform() const override;
 		virtual const FMetasoundFrontendClass& GetFrontendClass() const override;
