MetaSound CIS fix for deprecation warnings missing in macro expansion

diff --git a/Engine/Plugins/Runtime/Metasound/Source/MetasoundFrontend/Private/MetasoundFrontendDocument.cpp b/Engine/Plugins/Runtime/Metasound/Source/MetasoundFrontend/Private/MetasoundFrontendDocument.cpp
--- a/Engine/Plugins/Runtime/Metasound/Source/MetasoundFrontend/Private/MetasoundFrontendDocument.cpp
+++ b/Engine/Plugins/Runtime/Metasound/Source/MetasoundFrontend/Private/MetasoundFrontendDocument.cpp
@@ -1284,6 +1284,10 @@ FMetasoundFrontendGraphClass::FMetasoundFrontendGraphClass()
 	Metadata.SetType(EMetasoundFrontendClassType::Graph);
 }
 
+FMetasoundFrontendGraphClass::~FMetasoundFrontendGraphClass()
+{
+}
+
 #if WITH_EDITORONLY_DATA
 const FMetasoundFrontendGraph& FMetasoundFrontendGraphClass::AddGraphPage(const FGuid& InPageID, bool bDuplicateLastGraph, bool bSetAsBuildGraph)
 {
diff --git a/Engine/Plugins/Runtime/Metasound/Source/MetasoundFrontend/Public/MetasoundFrontendDocument.h b/Engine/Plugins/Runtime/Metasound/Source/MetasoundFrontend/Public/MetasoundFrontendDocument.h
--- a/Engine/Plugins/Runtime/Metasound/Source/MetasoundFrontend/Public/MetasoundFrontendDocument.h
+++ b/Engine/Plugins/Runtime/Metasound/Source/MetasoundFrontend/Public/MetasoundFrontendDocument.h
@@ -1793,7 +1793,7 @@ struct METASOUNDFRONTEND_API FMetasoundFrontendGraphClass : public FMetasoundFro
 
 public:
 	FMetasoundFrontendGraphClass();
-	virtual ~FMetasoundFrontendGraphClass() = default;
+	virtual ~FMetasoundFrontendGraphClass();
 
 #if WITH_EDITORONLY_DATA
 	UPROPERTY(meta = (DeprecationMessage = "5.5 - GraphClasses now support multiple paged graphs. Use the provided page graph accessors"))
