- Mark new graph object transient directly (don't rely on its owning property)
- Parent builders underneath the transient package instead of the associated MetaSound object

#rnx

diff --git a/Engine/Plugins/Runtime/Metasound/Source/MetasoundEditor/Private/MetasoundEditorGraphBuilder.cpp b/Engine/Plugins/Runtime/Metasound/Source/MetasoundEditor/Private/MetasoundEditorGraphBuilder.cpp
--- a/Engine/Plugins/Runtime/Metasound/Source/MetasoundEditor/Private/MetasoundEditorGraphBuilder.cpp
+++ b/Engine/Plugins/Runtime/Metasound/Source/MetasoundEditor/Private/MetasoundEditorGraphBuilder.cpp
@@ -1560,7 +1560,7 @@ PRAGMA_ENABLE_DEPRECATION_WARNINGS
 			UMetasoundEditorGraph* Graph = Cast<UMetasoundEditorGraph>(MetaSound.GetGraph());
 			if (!Graph)
 			{
-				Graph = NewObject<UMetasoundEditorGraph>(MetaSound.GetOwningAsset(), FName(), RF_Transactional);
+				Graph = NewObject<UMetasoundEditorGraph>(MetaSound.GetOwningAsset(), FName(), RF_Transactional | RF_Transient);
 				Graph->Schema = UMetasoundEditorGraphSchema::StaticClass();
 				MetaSound.SetGraph(Graph);
 				bNewGraphBound = true;
diff --git a/Engine/Plugins/Runtime/Metasound/Source/MetasoundEngine/Public/MetasoundDocumentBuilderRegistry.h b/Engine/Plugins/Runtime/Metasound/Source/MetasoundEngine/Public/MetasoundDocumentBuilderRegistry.h
--- a/Engine/Plugins/Runtime/Metasound/Source/MetasoundEngine/Public/MetasoundDocumentBuilderRegistry.h
+++ b/Engine/Plugins/Runtime/Metasound/Source/MetasoundEngine/Public/MetasoundDocumentBuilderRegistry.h
@@ -101,7 +101,7 @@ namespace Metasound::Engine
 			TObjectPtr<UMetaSoundBuilderBase> NewBuilder;
 			{
 				FScopeLock Lock(&BuildersCriticalSection);
-				NewBuilder = CastChecked<UMetaSoundBuilderBase>(NewObject<UObject>(&InMetaSoundObject, &BuilderUClass, NewName, RF_Transactional));
+				NewBuilder = CastChecked<UMetaSoundBuilderBase>(NewObject<UObject>(GetTransientPackage(), &BuilderUClass, NewName, RF_Transactional | RF_Transient));
 				FMetaSoundFrontendDocumentBuilder& BuilderRef = NewBuilder->GetBuilder();
 				BuilderRef = FMetaSoundFrontendDocumentBuilder(DocInterface);
 				const FMetasoundFrontendDocument& Document = DocInterface->GetConstDocument();
