Test node config updating.
Mark subinterface code as experimental
#rnx

diff --git a/Engine/Plugins/Runtime/Metasound/Source/MetasoundFrontend/Public/MetasoundFrontendDocument.h b/Engine/Plugins/Runtime/Metasound/Source/MetasoundFrontend/Public/MetasoundFrontendDocument.h
--- a/Engine/Plugins/Runtime/Metasound/Source/MetasoundFrontend/Public/MetasoundFrontendDocument.h
+++ b/Engine/Plugins/Runtime/Metasound/Source/MetasoundFrontend/Public/MetasoundFrontendDocument.h
@@ -501,7 +501,10 @@ struct METASOUNDFRONTEND_API FMetasoundFrontendNodeInterface
 	// Create a node interface which satisfies an existing class interface.
 	FMetasoundFrontendNodeInterface(const FMetasoundFrontendClassInterface& InClassInterface);
 
+	UE_INTERNAL
 	void Update(const FMetasoundFrontendClassInterface& InClassInterface);
+
+	UE_INTERNAL
 	void Update(const FMetasoundFrontendClassInterface& InClassInterface, TFunctionRef<void(const FMetasoundFrontendVertex&)> OnPreRemoveInput, TFunctionRef<void(const FMetasoundFrontendVertex&)> OnPreRemoveOutput);
 
 	// Input vertices to node.
@@ -523,7 +526,6 @@ struct METASOUNDFRONTEND_API FMetaSoundFrontendNodeConfiguration
 	GENERATED_BODY()
 
 	virtual ~FMetaSoundFrontendNodeConfiguration() = default;
-	// TODO: Markup with special experimental API Stuff
 
 	/* Get the current interface for the class based upon the node extension */
 	virtual TInstancedStruct<FMetasoundFrontendClassInterface> OverrideDefaultInterface(const FMetasoundFrontendClass& InNodeClass) const;
diff --git a/Engine/Plugins/Runtime/Metasound/Source/MetasoundGraphCore/Private/MetasoundVertex.cpp b/Engine/Plugins/Runtime/Metasound/Source/MetasoundGraphCore/Private/MetasoundVertex.cpp
--- a/Engine/Plugins/Runtime/Metasound/Source/MetasoundGraphCore/Private/MetasoundVertex.cpp
+++ b/Engine/Plugins/Runtime/Metasound/Source/MetasoundGraphCore/Private/MetasoundVertex.cpp
@@ -41,6 +41,7 @@ namespace Metasound
 			SubInterfaces->Reserve(Num);
 		}
 
+PRAGMA_DISABLE_EXPERIMENTAL_WARNINGS
 		void FSubInterfaceDeclarationBuilder::Add(FBeginSubInterface&& InSubInterface)
 		{
 			PushSubInterfaceDeclaration(InSubInterface.Name);
@@ -50,6 +51,7 @@ namespace Metasound
 		{
 			PopSubInterfaceDeclaration();
 		}
+PRAGMA_ENABLE_EXPERIMENTAL_WARNINGS
 
 		void FSubInterfaceDeclarationBuilder::OnVertexAdded(const FVertexName& InName)
 		{
@@ -598,6 +600,7 @@ namespace Metasound
 		return !(InLHS == InRHS);
 	}
 
+PRAGMA_DISABLE_EXPERIMENTAL_WARNINGS
 	FClassVertexInterface::FClassVertexInterface(FVertexInterface InInterface)
 	: FClassVertexInterface(MoveTemp(InInterface.GetInputInterface()), MoveTemp(InInterface.GetOutputInterface()), MoveTemp(InInterface.GetEnvironmentInterface()))
 	{
@@ -629,6 +632,7 @@ namespace Metasound
 
 		return FVertexInterface(MoveTemp(NewInputs), MoveTemp(NewOutputs), Environment);
 	}
+PRAGMA_ENABLE_EXPERIMENTAL_WARNINGS
 }
 
 FString LexToString(Metasound::EVertexAccessType InAccessType)
diff --git a/Engine/Plugins/Runtime/Metasound/Source/MetasoundGraphCore/Private/MetasoundVertexData.cpp b/Engine/Plugins/Runtime/Metasound/Source/MetasoundGraphCore/Private/MetasoundVertexData.cpp
--- a/Engine/Plugins/Runtime/Metasound/Source/MetasoundGraphCore/Private/MetasoundVertexData.cpp
+++ b/Engine/Plugins/Runtime/Metasound/Source/MetasoundGraphCore/Private/MetasoundVertexData.cpp
@@ -479,6 +479,11 @@ namespace Metasound
 		checkf(bIsVertexInterfaceFrozen || SubInterfaces.Num() == 0, TEXT("Interfaces containing sub interfaces cannot be unfrozen"));
 	}
 
+	bool FInputVertexInterfaceData::Contains(const FVertexName& InVertexName) const
+	{
+		return VertexDataPrivate::Find(Bindings, InVertexName) != nullptr;
+	}
+
 	void FInputVertexInterfaceData::BindVertex(const FVertexName& InVertexName, const FAnyDataReference& InOutDataReference)
 	{
 		BindVertex(InVertexName, const_cast<FAnyDataReference&>(InOutDataReference));
@@ -715,6 +720,11 @@ namespace Metasound
 		checkf(bIsVertexInterfaceFrozen || SubInterfaces.Num() == 0, TEXT("Interfaces containing sub interfaces cannot be unfrozen"));
 	}
 
+	bool FOutputVertexInterfaceData::Contains(const FVertexName& InVertexName) const
+	{
+		return VertexDataPrivate::Find(Bindings, InVertexName) != nullptr;
+	}
+
 	void FOutputVertexInterfaceData::BindVertex(const FVertexName& InVertexName, const FAnyDataReference& InOutDataReference)
 	{
 		BindVertex(InVertexName, const_cast<FAnyDataReference&>(InOutDataReference));
diff --git a/Engine/Plugins/Runtime/Metasound/Source/MetasoundGraphCore/Public/MetasoundVertex.h b/Engine/Plugins/Runtime/Metasound/Source/MetasoundGraphCore/Public/MetasoundVertex.h
--- a/Engine/Plugins/Runtime/Metasound/Source/MetasoundGraphCore/Public/MetasoundVertex.h
+++ b/Engine/Plugins/Runtime/Metasound/Source/MetasoundGraphCore/Public/MetasoundVertex.h
@@ -350,6 +350,14 @@ namespace Metasound
 			AddOrUpdateVertex(MoveTemp(InVertex));
 		}
 
+		void Append(TArrayView<const VertexType> InVertices)
+		{
+			for (const VertexType& Vertex : InVertices)
+			{
+				Add(Vertex);
+			}
+		}
+
 		/** Remove a vertex by key. */
 		bool Remove(const FVertexName& InKey)
 		{
@@ -445,12 +453,14 @@ namespace Metasound
 	using TVertexInterfaceGroup UE_DEPRECATED(5.6, "Use TVertexInterfaceImpl instead") = TVertexInterfaceImpl<VertexType>;
 
 	/** Declare the beginning of a sub interface in a vertex interface declaration. */
+	struct UE_EXPERIMENTAL(5.6, "Sub interfaces are experimental") FBeginSubInterface;
 	struct FBeginSubInterface
 	{
 		FName Name; //< Name of the sub interface. 
 	};
 
 	/** Declare the ending of a sub interface in a vertex interface declaration. */
+	struct UE_EXPERIMENTAL(5.6, "Sub interfaces are experimental") FEndSubInterface;
 	struct FEndSubInterface
 	{
 	};
@@ -500,10 +510,12 @@ namespace Metasound
 			virtual ~FSubInterfaceDeclarationBuilder();
 
 			// Handler for the beginning of a sub interface. 
+PRAGMA_DISABLE_EXPERIMENTAL_WARNINGS
 			void Add(FBeginSubInterface&& InSubInterface);
 
 			// Handler for the ending of a sub interface. 
 			void Add(FEndSubInterface&& InSubInterface);
+PRAGMA_ENABLE_EXPERIMENTAL_WARNINGS
 
 			// Allocate memory for a known number of sub interfaces.
 			void ReserveSubInterfaces(int32 Num);
@@ -537,8 +549,10 @@ namespace Metasound
 					Vertices.Reserve(NumVertices);
 				}
 
+PRAGMA_DISABLE_EXPERIMENTAL_WARNINGS
 				// Size sub interfaces exactly to avoid wasted memory in array allocations
 				constexpr uint32 NumSubInterfaces = TNumOfTypeInPack<FBeginSubInterface>::Get<ArgTypes...>();
+PRAGMA_ENABLE_EXPERIMENTAL_WARNINGS
 				if constexpr (NumSubInterfaces > 0)
 				{
 					ReserveSubInterfaces(NumSubInterfaces);
@@ -574,8 +588,10 @@ namespace Metasound
 					Vertices.Reserve(NumVertices);
 				}
 
+PRAGMA_DISABLE_EXPERIMENTAL_WARNINGS
 				// Size sub interfaces exactly to avoid wasted memory in array allocations
 				constexpr uint32 NumSubInterfaces = TNumOfTypeInPack<FBeginSubInterface>::Get<ArgTypes...>();
+PRAGMA_ENABLE_EXPERIMENTAL_WARNINGS
 				if constexpr (NumSubInterfaces > 0)
 				{
 					ReserveSubInterfaces(NumSubInterfaces);
@@ -646,7 +662,7 @@ namespace Metasound
 			InterfaceBuilder.Build(Forward<ArgTypes>(InArgs)...);
 		}
 
-		explicit FInputVertexInterface(TArray<FInputDataVertex> InVertices, TArray<VertexPrivate::FSubInterfaceLayout> InSubInterfaces);
+		FInputVertexInterface(TArray<FInputDataVertex> InVertices, TArray<VertexPrivate::FSubInterfaceLayout> InSubInterfaces={});
 
 		/** Iterate through all repetitions of a sub interface. */
 		void ForEachSubInterfaceInstance(const FName& InSubInterfaceName, TFunctionRef<void (TConstArrayView<FInputDataVertex>)> Callable) const;
@@ -688,7 +704,7 @@ namespace Metasound
 			InterfaceBuilder.Build(Forward<ArgTypes>(InArgs)...);
 		}
 
-		explicit FOutputVertexInterface(TArray<FOutputDataVertex> InVertices, TArray<VertexPrivate::FSubInterfaceLayout> InSubInterfaces);
+		FOutputVertexInterface(TArray<FOutputDataVertex> InVertices, TArray<VertexPrivate::FSubInterfaceLayout> InSubInterfaces={});
 
 		/** Iterate through all repetitions of a sub interface. */
 		void ForEachSubInterfaceInstance(const FName& InSubInterfaceName, TFunctionRef<void (TConstArrayView<FOutputDataVertex>)> Callable) const;
@@ -788,6 +804,7 @@ namespace Metasound
 	};
 
 	/** A description of a sub interface which is used when declaring a FClassVertexInterface */
+	struct UE_EXPERIMENTAL(5.6, "Sub interfaces are experimental") FSubInterfaceDescription;
 	struct FSubInterfaceDescription
 	{
 		FName SubInterfaceName;
@@ -798,6 +815,7 @@ namespace Metasound
 
 	/** A sub interface configuration is used to create a FVertexInterface from 
 	 * a FClassVertexInterface. */
+	struct UE_EXPERIMENTAL(5.6, "Sub interfaces are experimental") FSubInterfaceConfiguration;
 	struct FSubInterfaceConfiguration
 	{
 		FName SubInterfaceName;
@@ -806,20 +824,27 @@ namespace Metasound
 
 	/** FClassVertexInterface describes the interface of a node class. It is an 
 	 * immutable factory for FVertexInterfaces. */
+	class UE_EXPERIMENTAL(5.6, "FClassVertexInterfaces are experimental") FClassVertexInterface;
+
 	class METASOUNDGRAPHCORE_API FClassVertexInterface
 	{
 	public:
+PRAGMA_DISABLE_EXPERIMENTAL_WARNINGS
 		FClassVertexInterface() = default;
+		~FClassVertexInterface() = default;
 		FClassVertexInterface(FVertexInterface InInterface);
 		FClassVertexInterface(FInputVertexInterface InInputs, FOutputVertexInterface InOutputs, FEnvironmentVertexInterface InEnvironment={});
 		FClassVertexInterface(TArray<FSubInterfaceDescription> InSubInterfactDescriptions, FInputVertexInterface InInputs, FOutputVertexInterface InOutputs, FEnvironmentVertexInterface InEnvironment={});
 
 		/** Create a FVertexInterface with the given configuration. */
 		FVertexInterface CreateVertexInterface(TConstArrayView<FSubInterfaceConfiguration> InSubInterfaceConfigurations) const;
+PRAGMA_ENABLE_EXPERIMENTAL_WARNINGS
 
 	private:
 
+PRAGMA_DISABLE_EXPERIMENTAL_WARNINGS
 		TArray<FSubInterfaceDescription> SubInterfaces;
+PRAGMA_ENABLE_EXPERIMENTAL_WARNINGS
 		FInputVertexInterface Inputs;
 		FOutputVertexInterface Outputs;
 		FEnvironmentVertexInterface Environment;
diff --git a/Engine/Plugins/Runtime/Metasound/Source/MetasoundGraphCore/Public/MetasoundVertexData.h b/Engine/Plugins/Runtime/Metasound/Source/MetasoundGraphCore/Public/MetasoundVertexData.h
--- a/Engine/Plugins/Runtime/Metasound/Source/MetasoundGraphCore/Public/MetasoundVertexData.h
+++ b/Engine/Plugins/Runtime/Metasound/Source/MetasoundGraphCore/Public/MetasoundVertexData.h
@@ -565,6 +565,10 @@ namespace Metasound
 		 */
 		void SetIsVertexInterfaceFrozen(bool bInIsVertexInterfaceFrozen);
 
+		/** Returns true if a vertex exists with the provided name. This does not
+		 * reflect whether the vertex is bound or not. */
+		bool Contains(const FVertexName& InVertexName) const;
+
 		/** Set the value of a vertex. */
 		template<typename DataType>
 		void SetValue(const FVertexName& InVertexName, const DataType& InValue)
@@ -642,6 +646,7 @@ namespace Metasound
 		void BindVertex(const FVertexName& InVertexName, FAnyDataReference& InOutDataReference);
 
 		/** Return the number of instances of a particular sub interface. */
+		UE_EXPERIMENTAL(5.6, "Sub interfaces are still under development")
 		uint32 GetNumSubInterfaceInstances(const FName& InSubInterfaceName) const;
 
 		/** Bind the vertices of a sub interface vertex. 
@@ -650,6 +655,7 @@ namespace Metasound
 		 * @param InOutDataReferences - An indexable sequence container of data references. 
 		 */
 		template<typename RangeType>
+		UE_EXPERIMENTAL(5.6, "Sub interfaces are still under development")
 		void BindSubInterfaceVertices(const FName& InSubInterfaceName, const FName& InVertexName, RangeType& InOutDataReferences)
 		{
 			auto BindData = [&InOutDataReferences](uint32 Index, FInputBinding& Binding)
@@ -665,6 +671,7 @@ namespace Metasound
 		 * @param InOutDataReferences - An indexable sequence container of data references. 
 		 */
 		template<typename RangeType>
+		UE_EXPERIMENTAL(5.6, "Sub interfaces are still under development")
 		void BindSubInterfaceWriteVertices(const FName& InSubInterfaceName, const FName& InVertexName, RangeType& InOutDataReferences)
 		{
 			auto BindData = [&InOutDataReferences](uint32 Index, FInputBinding& Binding)
@@ -680,6 +687,7 @@ namespace Metasound
 		 * @param InOutDataReferences - An indexable sequence container of data references. 
 		 */
 		template<typename RangeType>
+		UE_EXPERIMENTAL(5.6, "Sub interfaces are still under development")
 		void BindSubInterfaceReadVertices(const FName& InSubInterfaceName, const FName& InVertexName, RangeType& InOutDataReferences)
 		{
 			auto BindData = [&InOutDataReferences](uint32 Index, FInputBinding& Binding)
@@ -695,6 +703,7 @@ namespace Metasound
 		 * @param InOutDataReferences - An indexable sequence container of data references. 
 		 */
 		template<typename RangeType>
+		UE_EXPERIMENTAL(5.6, "Sub interfaces are still under development")
 		void SetSubInterfaceValues(const FName& InSubInterfaceName, const FName& InVertexName, const RangeType& InValues)
 		{
 			auto BindData = [&InValues](uint32 Index, FInputBinding& Binding)
@@ -905,6 +914,7 @@ namespace Metasound
 		 * @param InSettings         - Operator settings to use if creating a data reference.
 		 */
 		template<typename DataReferenceType>
+		UE_EXPERIMENTAL(5.6, "Sub interfaces are still under development")
 		TArray<DataReferenceType> GetOrCreateDefaultSubInterfaceAs(const FName& InSubInterfaceName, const FVertexName& InVertexName, const FOperatorSettings& InSettings) const
 		{
 			TArray<DataReferenceType> References;
@@ -924,6 +934,7 @@ namespace Metasound
 		 * @param InSettings         - Operator settings to use if creating a data reference.
 		 */
 		template<typename DataType>
+		UE_EXPERIMENTAL(5.6, "Sub interfaces are still under development")
 		TArray<TDataReadReference<DataType>> GetOrCreateDefaultSubInterfaceDataReadReferences(const FName& InSubInterfaceName, const FVertexName& InVertexName, const FOperatorSettings& InSettings) const
 		{
 			return GetOrCreateDefaultSubInterfaceAs<TDataReadReference<DataType>>(InSubInterfaceName, InVertexName, InSettings);
@@ -935,6 +946,7 @@ namespace Metasound
 		 * @param InSettings         - Operator settings to use if creating a data reference.
 		 */
 		template<typename DataType>
+		UE_EXPERIMENTAL(5.6, "Sub interfaces are still under development")
 		TArray<TDataWriteReference<DataType>> GetOrCreateDefaultSubInterfaceDataWriteReferences(const FName& InSubInterfaceName, const FVertexName& InVertexName, const FOperatorSettings& InSettings) const
 		{
 			return GetOrCreateDefaultSubInterfaceAs<TDataWriteReference<DataType>>(InSubInterfaceName, InVertexName, InSettings);
@@ -946,6 +958,7 @@ namespace Metasound
 		 * @param InSettings         - Operator settings to use if creating a data reference.
 		 */
 		template<typename DataType>
+		UE_EXPERIMENTAL(5.6, "Sub interfaces are still under development")
 		TArray<DataType> GetOrCreateDefaultSubInterfaceValues(const FName& InSubInterfaceName, const FVertexName& InVertexName, const FOperatorSettings& InSettings) const
 		{
 			return GetOrCreateDefaultSubInterfaceAs<DataType>(InSubInterfaceName, InVertexName, InSettings);
@@ -958,6 +971,7 @@ namespace Metasound
 		 * @param InSubInterfaceName - Name of the sub interface
 		 * @param InVertexName       - Name of the prototype vertex in the sub interface.
 		 */
+		UE_EXPERIMENTAL(5.6, "Sub interfaces are still under development")
 		TArray<const FAnyDataReference*> FindSubInterfaceReferences(const FName& InSubInterfaceName, const FVertexName& InVertexName) const;
 	
 	private:
@@ -1010,6 +1024,10 @@ namespace Metasound
 		 */
 		void SetIsVertexInterfaceFrozen(bool bInIsVertexInterfaceFrozen);
 
+		/** Returns true if a vertex exists with the provided name. This does not
+		 * reflect whether the vertex is bound or not. */
+		bool Contains(const FVertexName& InVertexName) const;
+
 		/** Set the value of a vertex. */
 		template<typename DataType>
 		void SetValue(const FVertexName& InVertexName, const DataType& InValue)
@@ -1110,6 +1128,7 @@ namespace Metasound
 		void BindVertex(const FVertexName& InVertexName, FAnyDataReference& InOutDataReference);
 
 		/** Return the number of instances of a particular sub interface. */
+		UE_EXPERIMENTAL(5.6, "Sub interfaces are still under development")
 		uint32 GetNumSubInterfaceInstances(const FName& InSubInterfaceName) const;
 
 		/** Bind the vertices of a sub interface vertex. 
@@ -1118,6 +1137,7 @@ namespace Metasound
 		 * @param InOutDataReferences - An indexable sequence container of data references. 
 		 */
 		template<typename RangeType>
+		UE_EXPERIMENTAL(5.6, "Sub interfaces are still under development")
 		void BindSubInterfaceVertices(const FName& InSubInterfaceName, const FName& InVertexName, RangeType& InOutDataReferences)
 		{
 			using namespace VertexDataPrivate;
@@ -1135,6 +1155,7 @@ namespace Metasound
 		 * @param InOutDataReferences - An indexable sequence container of data references. 
 		 */
 		template<typename RangeType>
+		UE_EXPERIMENTAL(5.6, "Sub interfaces are still under development")
 		void BindSubInterfaceWriteVertices(const FName& InSubInterfaceName, const FName& InVertexName, RangeType& InOutDataReferences)
 		{
 			using namespace VertexDataPrivate;
@@ -1152,6 +1173,7 @@ namespace Metasound
 		 * @param InOutDataReferences - An indexable sequence container of data references. 
 		 */
 		template<typename RangeType>
+		UE_EXPERIMENTAL(5.6, "Sub interfaces are still under development")
 		void BindSubInterfaceReadVertices(const FName& InSubInterfaceName, const FName& InVertexName, RangeType& InOutDataReferences)
 		{
 			using namespace VertexDataPrivate;
@@ -1169,6 +1191,7 @@ namespace Metasound
 		 * @param InOutDataReferences - An indexable sequence container of data references. 
 		 */
 		template<typename RangeType>
+		UE_EXPERIMENTAL(5.6, "Sub interfaces are still under development")
 		void SetSubInterfaceValues(const FName& InSubInterfaceName, const FName& InVertexName, const RangeType& InValues)
 		{
 			using namespace VertexDataPrivate;
@@ -1323,6 +1346,7 @@ namespace Metasound
 		 * @param InSubInterfaceName - Name of the sub interface
 		 * @param InVertexName       - Name of the prototype vertex in the sub interface.
 		 */
+		UE_EXPERIMENTAL(5.6, "Sub interfaces are still under development")
 		TArray<const FAnyDataReference*> FindSubInterfaceReferences(const FName& InSubInterfaceName, const FVertexName& InVertexName) const;
 	
 	private:
