MetaSound node configuration editor layer
- Node configuration values can be updated in the node's detail panel and changes are reflected on the node

diff --git a/Engine/Plugins/Runtime/Metasound/Source/MetasoundEditor/MetasoundEditor.Build.cs b/Engine/Plugins/Runtime/Metasound/Source/MetasoundEditor/MetasoundEditor.Build.cs
--- a/Engine/Plugins/Runtime/Metasound/Source/MetasoundEditor/MetasoundEditor.Build.cs
+++ b/Engine/Plugins/Runtime/Metasound/Source/MetasoundEditor/MetasoundEditor.Build.cs
@@ -29,6 +29,7 @@ namespace UnrealBuildTool.Rules
 					"MetasoundGraphCore",
 					"MetasoundStandardNodes",
 					"SignalProcessing",
+					"StructUtilsEditor",
 					"SubobjectEditor",
 					"ToolMenus",
 					"ToolWidgets",
diff --git a/Engine/Plugins/Runtime/Metasound/Source/MetasoundEditor/Private/MetasoundEditorModule.cpp b/Engine/Plugins/Runtime/Metasound/Source/MetasoundEditor/Private/MetasoundEditorModule.cpp
--- a/Engine/Plugins/Runtime/Metasound/Source/MetasoundEditor/Private/MetasoundEditorModule.cpp
+++ b/Engine/Plugins/Runtime/Metasound/Source/MetasoundEditor/Private/MetasoundEditorModule.cpp
@@ -38,6 +38,7 @@
 #include "MetasoundEditorGraphBuilder.h"
 #include "MetasoundEditorGraphConnectionDrawingPolicy.h"
 #include "MetasoundEditorGraphMemberDefaults.h"
+#include "MetasoundEditorGraphNode.h"
 #include "MetasoundEditorGraphNodeFactory.h"
 #include "MetasoundEditorGraphNodeVisualization.h"
 #include "MetasoundEditorSettings.h"
@@ -865,6 +866,10 @@ namespace Metasound::Editor
 				UMetasoundPagesView::StaticClass()->GetFName(),
 				FOnGetDetailCustomizationInstance::CreateLambda([]() { return MakeShared<FMetasoundPagesDetailCustomization>(); }));
 
+			PropertyModule.RegisterCustomClassLayout(
+				UMetasoundEditorGraphNode::StaticClass()->GetFName(),
+				FOnGetDetailCustomizationInstance::CreateLambda([]() { return MakeShared<FMetaSoundNodeDetailCustomization>(); }));
+
 			PropertyModule.RegisterCustomClassLayout(
 				UMetasoundEditorGraphInput::StaticClass()->GetFName(),
 				FOnGetDetailCustomizationInstance::CreateLambda([]() { return MakeShared<FMetasoundInputDetailCustomization>(); }));
diff --git a/Engine/Plugins/Runtime/Metasound/Source/MetasoundEditor/Private/MetasoundNodeDetailCustomization.cpp b/Engine/Plugins/Runtime/Metasound/Source/MetasoundEditor/Private/MetasoundNodeDetailCustomization.cpp
--- a/Engine/Plugins/Runtime/Metasound/Source/MetasoundEditor/Private/MetasoundNodeDetailCustomization.cpp
+++ b/Engine/Plugins/Runtime/Metasound/Source/MetasoundEditor/Private/MetasoundNodeDetailCustomization.cpp
@@ -1062,6 +1062,67 @@ PRAGMA_ENABLE_DEPRECATION_WARNINGS
 				}
 			}
 		}
+		
+		void FMetaSoundNodeConfigurationDataDetails::OnChildRowAdded(IDetailPropertyRow& ChildRow)
+		{
+			FInstancedStructDataDetails::OnChildRowAdded(ChildRow);
+
+			TSharedPtr<IPropertyHandle> ChildHandle = ChildRow.GetPropertyHandle();
+			if (ChildHandle)
+			{
+				const FSimpleDelegate OnChildPropertyChangedDelegate = FSimpleDelegate::CreateSP(this, &FMetaSoundNodeConfigurationDataDetails::OnChildPropertyChanged);
+				ChildHandle->SetOnPropertyValueChanged(OnChildPropertyChangedDelegate);
+			}
+		}
+
+		void FMetaSoundNodeConfigurationDataDetails::OnChildPropertyChanged()
+		{
+			if (GraphNode.IsValid())
+			{
+				FMetaSoundFrontendDocumentBuilder& DocBuilder = GraphNode->GetBuilderChecked().GetBuilder();
+				const FGuid& NodeID = GraphNode->GetNodeID();
+				DocBuilder.UpdateNodeInterfaceFromConfiguration(NodeID);
+
+				FMetasoundFrontendDocumentModifyContext& ModifyContext = FGraphBuilder::GetOutermostMetaSoundChecked(*GraphNode).GetModifyContext();
+				ModifyContext.AddNodeIDModified(NodeID);
+			}
+		}
+
+		void FMetaSoundNodeDetailCustomization::CustomizeDetails(IDetailLayoutBuilder& InDetailLayout)
+		{
+			TArray<TWeakObjectPtr<UObject>> Objects;
+			InDetailLayout.GetObjectsBeingCustomized(Objects);
+			if (!Objects.IsEmpty())
+			{
+				UMetasoundEditorGraphNode* Node = Cast<UMetasoundEditorGraphNode>(Objects.Last().Get());
+				const UEdGraph* EdGraph = Node->GetGraph();
+				const UMetasoundEditorGraph* MetasoundGraph = CastChecked<UMetasoundEditorGraph>(EdGraph);
+				const FMetaSoundFrontendDocumentBuilder& DocBuilder = MetasoundGraph->GetBuilderChecked().GetBuilder();
+				
+				// Walk the property path to find the node configuration handle
+				TSharedPtr<IPropertyHandle> DocumentHandle = InDetailLayout.AddObjectPropertyData({ MetasoundGraph->GetMetasound() }, FName(TEXT("RootMetasoundDocument")));
+				check(DocumentHandle);
+				TSharedPtr<IPropertyHandle> RootGraphHandle = DocumentHandle->GetChildHandle(GET_MEMBER_NAME_CHECKED(FMetasoundFrontendDocument, RootGraph));
+				TSharedPtr<IPropertyHandle> PagedGraphsHandle = RootGraphHandle->GetChildHandle(FName(TEXT("PagedGraphs")));
+
+				const FGuid& PageID = DocBuilder.GetBuildPageID();
+				const int32 PageIndex = DocBuilder.FindPageIndex(PageID);
+				TSharedRef<IPropertyHandle> GraphHandle = PagedGraphsHandle->AsArray()->GetElement(PageIndex);
+				TSharedPtr<IPropertyHandle> NodesHandle = GraphHandle->GetChildHandle(GET_MEMBER_NAME_CHECKED(FMetasoundFrontendGraph, Nodes));
+				
+				const int32* NodeIndex = DocBuilder.FindNodeIndex(Node->GetNodeID());
+				if (NodeIndex != nullptr)
+				{
+					TSharedRef<IPropertyHandle> NodeHandle = NodesHandle->AsArray()->GetElement(*NodeIndex);
+					TSharedPtr<IPropertyHandle> ConfigurationHandle = NodeHandle->GetChildHandle(GET_MEMBER_NAME_CHECKED(FMetasoundFrontendNode, Configuration));
+
+					// Use a custom details builder to add custom child property update behavior
+					// and hide the struct picker since users should only be able edit child values, not the struct type 
+					TSharedRef<IDetailCustomNodeBuilder> DataDetailsBuilder = MakeShared<FMetaSoundNodeConfigurationDataDetails>(ConfigurationHandle, TWeakObjectPtr<UMetasoundEditorGraphNode>(Node));
+					InDetailLayout.EditCategory("General").AddCustomBuilder(DataDetailsBuilder);
+				}
+			}
+		}
 
 		void FMetasoundMemberDetailCustomization::UpdateRenameDelegate(UMetasoundEditorGraphMember& InMember)
 		{
diff --git a/Engine/Plugins/Runtime/Metasound/Source/MetasoundEditor/Private/MetasoundNodeDetailCustomization.h b/Engine/Plugins/Runtime/Metasound/Source/MetasoundEditor/Private/MetasoundNodeDetailCustomization.h
--- a/Engine/Plugins/Runtime/Metasound/Source/MetasoundEditor/Private/MetasoundNodeDetailCustomization.h
+++ b/Engine/Plugins/Runtime/Metasound/Source/MetasoundEditor/Private/MetasoundNodeDetailCustomization.h
@@ -11,6 +11,7 @@
 #include "IDetailCustomization.h"
 #include "IDetailPropertyExtensionHandler.h"
 #include "IDetailPropertyRow.h"
+#include "InstancedStructDetails.h"
 #include "IPropertyTypeCustomization.h"
 #include "Layout/Visibility.h"
 #include "MetasoundAssetBase.h"
@@ -318,6 +319,30 @@ namespace Metasound
 			FName ArrayTypeName;
 		};
 
+		class FMetaSoundNodeConfigurationDataDetails : public FInstancedStructDataDetails
+		{
+		public:
+			FMetaSoundNodeConfigurationDataDetails(TSharedPtr<IPropertyHandle> InStructProperty, TWeakObjectPtr<UMetasoundEditorGraphNode> InNode)
+				: FInstancedStructDataDetails(InStructProperty)
+				, GraphNode(InNode)
+			{
+			}
+
+			virtual void OnChildRowAdded(IDetailPropertyRow& ChildRow) override;
+		
+		private: 
+			void OnChildPropertyChanged();
+
+			TWeakObjectPtr<UMetasoundEditorGraphNode> GraphNode;
+		};
+
+		class FMetaSoundNodeDetailCustomization : public IDetailCustomization
+		{
+			// IDetailCustomization interface
+			virtual void CustomizeDetails(IDetailLayoutBuilder& InDetailLayout) override;
+			// End of IDetailCustomization interface
+		};
+
 		class FMetasoundMemberDetailCustomization : public IDetailCustomization
 		{
 		public:
diff --git a/Engine/Plugins/Runtime/Metasound/Source/MetasoundFrontend/Private/MetasoundFrontendDocumentBuilder.cpp b/Engine/Plugins/Runtime/Metasound/Source/MetasoundFrontend/Private/MetasoundFrontendDocumentBuilder.cpp
--- a/Engine/Plugins/Runtime/Metasound/Source/MetasoundFrontend/Private/MetasoundFrontendDocumentBuilder.cpp
+++ b/Engine/Plugins/Runtime/Metasound/Source/MetasoundFrontend/Private/MetasoundFrontendDocumentBuilder.cpp
@@ -2617,6 +2617,14 @@ const FMetasoundFrontendNode* FMetaSoundFrontendDocumentBuilder::FindNode(const
 	return NodeCache.FindNode(InNodeID);
 }
 
+const int32* FMetaSoundFrontendDocumentBuilder::FindNodeIndex(const FGuid& InNodeID, const FGuid* InPageID) const
+{
+	using namespace Metasound::Frontend;
+	const FGuid& PageID = InPageID ? *InPageID : BuildPageID;
+	const IDocumentGraphNodeCache& NodeCache = DocumentCache->GetNodeCache(PageID);
+	return NodeCache.FindNodeIndex(InNodeID);
+}
+
 bool FMetaSoundFrontendDocumentBuilder::FindNodeClassInterfaces(const FGuid& InNodeID, TSet<FMetasoundFrontendVersion>& OutInterfaces, const FGuid& InPageID) const
 {
 	using namespace Metasound;
@@ -2847,6 +2855,14 @@ const FMetasoundFrontendVertex* FMetaSoundFrontendDocumentBuilder::FindNodeOutpu
 	return nullptr;
 }
 
+int32 FMetaSoundFrontendDocumentBuilder::FindPageIndex(const FGuid& InPageID) const
+{
+	const FMetasoundFrontendDocument& Document = GetDocumentChecked();
+	const TArray<FMetasoundFrontendGraph>& GraphPages = Document.RootGraph.GetConstGraphPages();
+	auto MatchesPageID = [this, &InPageID](const FMetasoundFrontendGraph& Iter) { return Iter.PageID == InPageID; };
+	return GraphPages.IndexOfByPredicate(MatchesPageID);
+}
+
 #if WITH_EDITORONLY_DATA
 FMetaSoundFrontendGraphComment& FMetaSoundFrontendDocumentBuilder::FindOrAddGraphComment(const FGuid& InCommentID, const FGuid* InPageID)
 {
@@ -5591,6 +5607,37 @@ bool FMetaSoundFrontendDocumentBuilder::SwapGraphOutput(const FMetasoundFrontend
 	return true;
 }
 
+bool FMetaSoundFrontendDocumentBuilder::UpdateNodeInterfaceFromConfiguration(const FGuid& InNodeID, const FGuid* InPageID)
+{
+	if (FMetasoundFrontendNode* Node = FindNodeInternal(InNodeID, InPageID))
+	{
+		if (const FMetaSoundFrontendNodeConfiguration* ConfigurationPtr = Node->Configuration.GetPtr())
+		{
+			const FMetasoundFrontendClass* Class = FindDependency(Node->ClassID);
+			check(Class);
+
+			Node->ClassInterfaceOverride = ConfigurationPtr->OverrideDefaultInterface(*Class);
+			const FMetasoundFrontendClassInterface & ClassInterface = Class->GetInterfaceForNode(*Node);
+
+			auto DisconnectInput = [this, &InNodeID, &InPageID](const FMetasoundFrontendVertex& InNodeInput)
+			{
+				this->RemoveEdgeToNodeInput(InNodeID, InNodeInput.VertexID, InPageID);
+			};
+
+			auto DisconnectOutput = [this, &InNodeID, &InPageID](const FMetasoundFrontendVertex& InNodeOutput)
+			{
+				this->RemoveEdgesFromNodeOutput(InNodeID, InNodeOutput.VertexID, InPageID);
+			};
+
+			Node->Interface.Update(ClassInterface, DisconnectInput, DisconnectOutput);
+		}
+
+		return true;
+	}
+
+	return false;
+}
+
 bool FMetaSoundFrontendDocumentBuilder::UnlinkVariableNode(const FGuid& InNodeID, const FGuid& InPageID)
 {
 	auto IsNodeID = [&InNodeID](const FGuid& TestID) { return TestID == InNodeID; };
diff --git a/Engine/Plugins/Runtime/Metasound/Source/MetasoundFrontend/Public/MetasoundFrontendDocument.h b/Engine/Plugins/Runtime/Metasound/Source/MetasoundFrontend/Public/MetasoundFrontendDocument.h
--- a/Engine/Plugins/Runtime/Metasound/Source/MetasoundFrontend/Public/MetasoundFrontendDocument.h
+++ b/Engine/Plugins/Runtime/Metasound/Source/MetasoundFrontend/Public/MetasoundFrontendDocument.h
@@ -1717,7 +1717,10 @@ public:
 	TArray<FMetasoundFrontendVertexLiteral> InputLiterals;
 
 	// Instance of a configuration for this node.
-	UPROPERTY()
+	// This property is EditAnywhere with a false EditCondition 
+	// so child properties (ex. node configuration) have exposed property handles for details customization code, 
+	// but they should not be editable elsewhere in the editor (ex. property matrix)
+	UPROPERTY(EditAnywhere, Category = CustomView, meta = (EditCondition = "false", HideProperty))
 	TInstancedStruct<FMetaSoundFrontendNodeConfiguration> Configuration;
 
 	// An optional override to the default class interface.
@@ -1767,7 +1770,10 @@ struct METASOUNDFRONTEND_API FMetasoundFrontendGraph
 	GENERATED_BODY()
 
 	// Node contained in graph
-	UPROPERTY()
+	// This property is EditAnywhere with a false EditCondition 
+	// so child properties (ex. node configuration) have exposed property handles for details customization code, 
+	// but they should not be editable elsewhere in the editor (ex. property matrix)
+	UPROPERTY(EditAnywhere, Category = CustomView, meta = (EditCondition = "false", HideProperty))
 	TArray<FMetasoundFrontendNode> Nodes;
 
 	// Connections between points on nodes.
@@ -1803,7 +1809,10 @@ public:
 #endif // WITH_EDITORONLY_DATA
 
 private:
-	UPROPERTY()
+	// This property is EditAnywhere with a false EditCondition 
+	// so child properties (ex. node configuration) have exposed property handles for details customization code, 
+	// but they should not be editable elsewhere in the editor (ex. property matrix)
+	UPROPERTY(EditAnywhere, Category = CustomView, meta = (EditCondition = "false", HideProperty))
 	TArray<FMetasoundFrontendGraph> PagedGraphs;
 
 public:
diff --git a/Engine/Plugins/Runtime/Metasound/Source/MetasoundFrontend/Public/MetasoundFrontendDocumentBuilder.h b/Engine/Plugins/Runtime/Metasound/Source/MetasoundFrontend/Public/MetasoundFrontendDocumentBuilder.h
--- a/Engine/Plugins/Runtime/Metasound/Source/MetasoundFrontend/Public/MetasoundFrontendDocumentBuilder.h
+++ b/Engine/Plugins/Runtime/Metasound/Source/MetasoundFrontend/Public/MetasoundFrontendDocumentBuilder.h
@@ -261,6 +261,8 @@ public:
 	const FMetasoundFrontendGraph& FindConstBuildGraphChecked() const;
 
 	const FMetasoundFrontendNode* FindNode(const FGuid& InNodeID, const FGuid* InPageID = nullptr) const;
+	// Return the node's index in the document's specified paged graph's node array 
+	const int32* FindNodeIndex(const FGuid& InNodeID, const FGuid* InPageID = nullptr) const;
 
 	const FMetasoundFrontendVertex* FindNodeInput(const FGuid& InNodeID, const FGuid& InVertexID, const FGuid* InPageID = nullptr) const;
 	const FMetasoundFrontendVertex* FindNodeInput(const FGuid& InNodeID, FName InVertexName, const FGuid* InPageID = nullptr) const;
@@ -282,6 +284,9 @@ public:
 	TArray<const FMetasoundFrontendVertex*> FindNodeOutputs(const FGuid& InNodeID, FName TypeName = FName(), const FGuid* InPageID = nullptr) const;
 	const FMetasoundFrontendVertex* FindNodeOutputConnectedToNodeInput(const FGuid& InInputNodeID, const FGuid& InInputVertexID, const FMetasoundFrontendNode** ConnectedOutputNode = nullptr, const FGuid* InPageID = nullptr) const;
 
+	// Return the index of the given page in the document's PagedGraphs array 
+	int32 FindPageIndex(const FGuid& InPageID) const;
+
 	const FMetasoundFrontendDocument& GetConstDocumentChecked() const;
 	const IMetaSoundDocumentInterface& GetConstDocumentInterfaceChecked() const;
 	const FString GetDebugName() const;
@@ -566,6 +571,10 @@ public:
 	bool UpdateDependencyRegistryData(const TMap<Metasound::Frontend::FNodeRegistryKey, Metasound::Frontend::FNodeRegistryKey>& OldToNewClassKeys);
 #endif // WITH_EDITOR
 
+	// If a node contains a node configuration, update the node class interface and interface.
+	// Returns true if node is found.
+	bool UpdateNodeInterfaceFromConfiguration(const FGuid& InNodeID, const FGuid* InPageID = nullptr);
+
 #if WITH_EDITORONLY_DATA
 	// Transforms template nodes within the given builder's document, which can include swapping associated edges and/or
 	// replacing nodes with other, registry-defined concrete node class instances. Returns true if any template nodes were processed.
