Have builder utilize node config when doing live audition
Use common codepath for creating INodes from Metasound Document subobjects.

diff --git a/Engine/Plugins/Runtime/Metasound/Source/MetasoundEngine/Private/MetasoundBuilderSubsystem.cpp b/Engine/Plugins/Runtime/Metasound/Source/MetasoundEngine/Private/MetasoundBuilderSubsystem.cpp
--- a/Engine/Plugins/Runtime/Metasound/Source/MetasoundEngine/Private/MetasoundBuilderSubsystem.cpp
+++ b/Engine/Plugins/Runtime/Metasound/Source/MetasoundEngine/Private/MetasoundBuilderSubsystem.cpp
@@ -19,6 +19,7 @@
 #include "MetasoundFrontendDataTypeRegistry.h"
 #include "MetasoundFrontendDocument.h"
 #include "MetasoundFrontendDocumentIdGenerator.h"
+#include "MetasoundFrontendGraphBuilder.h"
 #include "MetasoundFrontendRegistries.h"
 #include "MetasoundFrontendSearchEngine.h"
 #include "MetasoundFrontendTransform.h"
@@ -480,119 +481,21 @@ void UMetaSoundSourceBuilder::OnNodeAdded(int32 NodeIndex) const
 		using namespace Metasound;
 		using namespace Metasound::Frontend;
 
-		const FMetasoundFrontendNode& AddedNode = GetConstTargetPageGraphChecked().Nodes[NodeIndex];
+		const FMetasoundFrontendGraphClass& OwningGraphClass = Builder.GetConstDocumentChecked().RootGraph;
+		const FMetasoundFrontendGraph& OwningGraph = GetConstTargetPageGraphChecked();
+		const FMetasoundFrontendNode& AddedNode = OwningGraph.Nodes[NodeIndex];
 		const FMetasoundFrontendClass* NodeClass = Builder.FindDependency(AddedNode.ClassID);
 		checkf(NodeClass, TEXT("Node successfully added to graph but document is missing associated dependency"));
 
-		// The logic for creating the FVertexInterface already lives in the FMetasoundFrontendGraphBuilder. 
-		// That logic should be exposed in an API that can be called here. In order
-		// to account for class overrides from node configurations, the vertex interface
-		// will need to be built from the `FMetasoundFrontendClass` & `FMetasoundFrontendNode`
-		// instead of being retrieved from the registry. 
-		// JIRA UE-235230
-
-		const FNodeRegistryKey& ClassKey = FNodeRegistryKey(NodeClass->Metadata);
-		FMetasoundFrontendRegistryContainer& NodeRegistry = *FMetasoundFrontendRegistryContainer::Get();
-		IDataTypeRegistry& DataTypeRegistry = IDataTypeRegistry::Get();
-
-		TUniquePtr<INode> NewNode;
-
-		switch (NodeClass->Metadata.GetType())
+		FGraphBuilder::FCreateNodeParams CreateNodeParams
 		{
-			case EMetasoundFrontendClassType::VariableDeferredAccessor:
-			case EMetasoundFrontendClassType::VariableAccessor:
-			case EMetasoundFrontendClassType::VariableMutator:
-			case EMetasoundFrontendClassType::External:
-			case EMetasoundFrontendClassType::Graph:
-			{
-				FVertexInterface VertexInterface;
-				NodeRegistry.FindDefaultVertexInterface(ClassKey, VertexInterface);
-				NewNode = NodeRegistry.CreateNode(ClassKey, FNodeData{ AddedNode.Name, AddedNode.GetID() , VertexInterface });
-			}
-			break;
-
-			case EMetasoundFrontendClassType::Input:
-			{
-				const FName& DataTypeName = NodeClass->Metadata.GetClassName().Name;
-				const FMetasoundFrontendVertex& InputVertex = AddedNode.Interface.Inputs.Last();
-
-				const FMetasoundFrontendLiteral* DefaultLiteral = nullptr;
-				auto HasEqualVertexID = [&VertexID = InputVertex.VertexID](const FMetasoundFrontendVertexLiteral& InVertexLiteral)
-				{
-					return InVertexLiteral.VertexID == VertexID;
-				};
-
-				// Check for default literal on Node
-				if (const FMetasoundFrontendVertexLiteral* VertexLiteralOnNode = AddedNode.InputLiterals.FindByPredicate(HasEqualVertexID))
-				{
-					DefaultLiteral = &(VertexLiteralOnNode->Value);
-				}
-				// Check for default literal on class input
-				else if (const FMetasoundFrontendClassInput* GraphInput = Builder.FindGraphInput(InputVertex.Name))
-				{
-					DefaultLiteral = &(GraphInput->FindConstDefaultChecked(Frontend::DefaultPageID));
-				}
-				else
-				{
-					// As a last resort, get default literal on node class 
-					// JIRA UE-235230 - This should use the applicable class interface instead of the default class interface
-					DefaultLiteral = &(NodeClass->GetDefaultInterface().Inputs.Last().FindConstDefaultChecked(Frontend::DefaultPageID));
-				}
-
-				TSharedPtr<const FNodeClassMetadata> InputClassMetadata = DataTypeRegistry.GetInputClassMetadata(DataTypeName);
-				checkf(InputClassMetadata, TEXT("Input class metadata could not be found from data type registry"));
-
-				FVertexInterface InputNodeInterface = InputClassMetadata->DefaultInterface;
-				// Vertex names must be set for input nodes
-				InputNodeInterface.GetInputInterface().At(0).VertexName = InputVertex.Name;
-				InputNodeInterface.GetOutputInterface().At(0).VertexName = InputVertex.Name;
-				InputNodeInterface.GetInputInterface().At(0).SetDefaultLiteral(DefaultLiteral->ToLiteral(DataTypeName));
-
-				NewNode = DataTypeRegistry.CreateInputNode(DataTypeName, FNodeData{ AddedNode.Name, AddedNode.GetID(), InputNodeInterface });
-			}
-			break;
-
-			case EMetasoundFrontendClassType::Variable:
-			{
-				const FName& DataTypeName = NodeClass->Metadata.GetClassName().Name;
-				TSharedPtr<const FNodeClassMetadata> VariableMetadata = DataTypeRegistry.GetVariableClassMetadata(DataTypeName);
-				checkf(VariableMetadata, TEXT("Variable class metadata could not be found from data type registry"));
-				NewNode = DataTypeRegistry.CreateVariableNode(DataTypeName, DataTypeRegistry.CreateDefaultLiteral(DataTypeName),
-					FNodeData{ AddedNode.Name, AddedNode.GetID(), VariableMetadata->DefaultInterface }); 
-			}
-			break;
-
-			case EMetasoundFrontendClassType::Literal:
-			{
-				UE_LOG(LogMetaSound, Warning, TEXT("Builder '%s': Adding literal nodes is being deprecated. Please set literal values on the connected node interface directly."), *GetName());
-				const FName& DataTypeName = NodeClass->Metadata.GetClassName().Name;
-PRAGMA_DISABLE_DEPRECATION_WARNINGS
-				FDefaultLiteralNodeConstructorParams InitData { AddedNode.Name, AddedNode.GetID(), DataTypeRegistry.CreateDefaultLiteral(DataTypeName)};
-				NewNode = DataTypeRegistry.CreateLiteralNode(DataTypeName, MoveTemp(InitData));
-PRAGMA_ENABLE_DEPRECATION_WARNINGS
-			}
-			break;
-
-			case EMetasoundFrontendClassType::Output:
-			{
-				const FName& DataTypeName = NodeClass->Metadata.GetClassName().Name;
-				const FName& OutputVertexName = AddedNode.Interface.Outputs.Last().Name;
-
-				TSharedPtr<const FNodeClassMetadata> OutputClassMetadata = DataTypeRegistry.GetOutputClassMetadata(DataTypeName);
-				checkf(OutputClassMetadata, TEXT("Output class metadata could not be found from data type registry"));
-
-				FVertexInterface OutputNodeInterface = OutputClassMetadata->DefaultInterface;
-				// Vertex names must be set for output nodes
-				OutputNodeInterface.GetInputInterface().At(0).VertexName = OutputVertexName;
-				OutputNodeInterface.GetOutputInterface().At(0).VertexName = OutputVertexName;
-				NewNode = DataTypeRegistry.CreateOutputNode(DataTypeName, FNodeData{ AddedNode.Name, AddedNode.GetID(), OutputNodeInterface });
-			}
-			break;
-
-			case EMetasoundFrontendClassType::Template:
-			default:
-			static_assert(static_cast<int32>(EMetasoundFrontendClassType::Invalid) == 10, "Possible missed EMetasoundFrontendClassType case coverage");
+			GetName(), 
+			AddedNode,
+			*NodeClass,
+			OwningGraph,
+			OwningGraphClass
 		};
+		TUniquePtr<INode> NewNode = FGraphBuilder::CreateNode(CreateNodeParams);
 
 		if (!NewNode.IsValid())
 		{
diff --git a/Engine/Plugins/Runtime/Metasound/Source/MetasoundFrontend/Private/MetasoundFrontendGraphBuilder.cpp b/Engine/Plugins/Runtime/Metasound/Source/MetasoundFrontend/Private/MetasoundFrontendGraphBuilder.cpp
--- a/Engine/Plugins/Runtime/Metasound/Source/MetasoundFrontend/Private/MetasoundFrontendGraphBuilder.cpp
+++ b/Engine/Plugins/Runtime/Metasound/Source/MetasoundFrontend/Private/MetasoundFrontendGraphBuilder.cpp
@@ -16,6 +16,7 @@
 #include "MetasoundFrontendProxyDataCache.h"
 #include "MetasoundFrontendRegistries.h"
 #include "MetasoundGraph.h"
+#include "MetasoundGraphNode.h"
 #include "MetasoundLiteralNode.h"
 #include "MetasoundLog.h"
 #include "MetasoundNodeInterface.h"
@@ -27,18 +28,22 @@ namespace Metasound::Frontend
 		// Map of Input VertexID to variable data required to construct and connect default variable
 		using FNodeIDVertexID = TTuple<FGuid, FGuid>;
 		using FDependencyByIDMap = TMap<FGuid, const FMetasoundFrontendClass*>;
-		using FSharedNodeByIDMap = TMap<FGuid, TSharedPtr<const INode>>;
+		using FSharedNodeByIDMap = TMap<FGuid, TSharedPtr<const IGraph>>;
 
 		// Context used throughout entire graph build process
 		// (for both a root and nested subgraphs)
 		struct FBuildContext
 		{
-			FString DebugAssetName;
+			const FString& DebugAssetName;
+			const Frontend::IDataTypeRegistry& DataTypeRegistry;
+			const Frontend::FProxyDataCache* ProxyDataCache;
+		};
+
+		// Context related to the document being built. 
+		struct FBuildDocumentContext
+		{
 			FDependencyByIDMap FrontendClasses;
 			FSharedNodeByIDMap Graphs;
-			const Frontend::IDataTypeRegistry& DataTypeRegistry;
-			const Frontend::FProxyDataCache& ProxyDataCache;
-			FGuid GraphId;
 		};
 
 		// Transient context used for building a specific graph
@@ -47,7 +52,8 @@ namespace Metasound::Frontend
 			TUniquePtr<FFrontendGraph> Graph;
 			const FMetasoundFrontendGraphClass& GraphClass;
 			const FMetasoundFrontendGraph& PagedGraph;
-			FBuildContext& BuildContext;
+			const FBuildContext& BuildContext;
+			const FBuildDocumentContext& BuildDocumentContext;
 		};
 
 		template <typename ResolveType>
@@ -296,7 +302,7 @@ namespace Metasound::Frontend
 					// Note: stripping metadata here. If we make this function public, we may want to keep metadata if WITH_EDITORONLY_DATA
 					if (InFrontendLiteral)
 					{
-						FLiteral CoreLiteral = InFrontendLiteral->ToLiteral(DataTypeName, &InContext.DataTypeRegistry, &InContext.ProxyDataCache);
+						FLiteral CoreLiteral = InFrontendLiteral->ToLiteral(DataTypeName, &InContext.DataTypeRegistry, InContext.ProxyDataCache);
 						return FInputDataVertex{Name, DataTypeName, FDataVertexMetadata{}, AccessType, MoveTemp(CoreLiteral)};
 					}
 					else
@@ -387,6 +393,24 @@ namespace Metasound::Frontend
 				CreateEnvironmentVertexInterface(InNode.Interface));
 		}
 
+		const FName& GetNodeName(const FMetasoundFrontendNode& InNode, const FMetasoundFrontendClassInput* InOwningGraphClassInput=nullptr, const FMetasoundFrontendClassOutput* InOwningGraphClassOutput=nullptr)
+		{
+			// Make sure the node's name is correct when it is an input or output node. 
+			// The node name is used when adding the node to the FGraph as an input.
+			if (InOwningGraphClassInput)
+			{
+				return InOwningGraphClassInput->Name;
+			}
+			else if (InOwningGraphClassOutput)
+			{
+				return InOwningGraphClassOutput->Name;
+			}
+			else
+			{
+				return InNode.Name;
+			}
+		}
+
 		FNodeData CreateNodeData(const FBuildContext& InContext, const FMetasoundFrontendNode& InNode, const FMetasoundFrontendClass& InClass, const FMetasoundFrontendClassInterface& InClassInterface, const FMetasoundFrontendClassInput* InOwningGraphClassInput=nullptr, const FMetasoundFrontendClassOutput* InOwningGraphClassOutput=nullptr)
 		{
 			TSharedPtr<const IOperatorData> Config;
@@ -394,7 +418,10 @@ namespace Metasound::Frontend
 			{
 				Config = InNode.Configuration.Get().GetOperatorData();
 			}
-			return FNodeData(InNode.Name, InNode.GetID(), CreateVertexInterface(InContext, InNode, InClass, InClassInterface, InOwningGraphClassInput, InOwningGraphClassOutput), MoveTemp(Config));
+
+			const FName& NodeName = GetNodeName(InNode, InOwningGraphClassInput, InOwningGraphClassOutput);
+
+			return FNodeData(NodeName, InNode.GetID(), CreateVertexInterface(InContext, InNode, InClass, InClassInterface, InOwningGraphClassInput, InOwningGraphClassOutput), MoveTemp(Config));
 		}
 
 		const FMetasoundFrontendVariable* FindVariableForVariableNode(const FMetasoundFrontendNode& InVariableNode, const FMetasoundFrontendGraph& InGraph)
@@ -419,7 +446,7 @@ namespace Metasound::Frontend
 
 				if (IsLiteralParsableByDataType)
 				{
-					FLiteral Literal = FrontendVariable->Literal.ToLiteral(FrontendVariable->TypeName, &InContext.DataTypeRegistry, &InContext.ProxyDataCache);
+					FLiteral Literal = FrontendVariable->Literal.ToLiteral(FrontendVariable->TypeName, &InContext.DataTypeRegistry, InContext.ProxyDataCache);
 					return DataTypeRegistry.CreateVariableNode(FrontendVariable->TypeName, MoveTemp(Literal), CreateNodeData(InContext, InNode, InClass, InClassInterface));
 				}
 				else
@@ -509,6 +536,16 @@ namespace Metasound::Frontend
 			return INodeClassRegistry::Get()->CreateNode(Key, CreateNodeData(InContext, InNode, InClass, InClassInterface));
 		}
 
+		TUniquePtr<INode> CreateSubgraphNode(const FBuildContext& InContext, const TSharedRef<const IGraph>& InSubgraph, const FMetasoundFrontendNode& InNode, const FMetasoundFrontendClass& InClass, const FMetasoundFrontendClassInterface& InClassInterface)
+		{
+			using namespace Frontend;
+
+			check(InNode.ClassID == InClass.ID);
+			check(InterfacesHaveEqualSize(InClassInterface, InNode.Interface));
+
+			return MakeUnique<FGraphNode>(CreateNodeData(InContext, InNode, InClass, InClassInterface), InSubgraph);
+		}
+
 		const FMetasoundFrontendClassInput* FindClassInputForInputNode(const FMetasoundFrontendGraphClass& InOwningGraph, const FMetasoundFrontendNode& InInputNode)
 		{
 			// Input nodes should have exactly one input.
@@ -543,6 +580,7 @@ namespace Metasound::Frontend
 			return nullptr;
 		}
 
+
 		void LogFailedToCreateNode(const FString& InAssetName, const FMetasoundFrontendNode& InNode, const FMetasoundFrontendClass& InNodeClass)
 		{
 			UE_LOG(LogMetaSound, Error, TEXT("Metasound '%s': Failed to create node [NodeID:%s, NodeName:%s] from registry [Class: %s %s]"), *InAssetName, *InNode.GetID().ToString(), *InNode.Name.ToString(), *InNodeClass.Metadata.GetClassName().ToString(), *InNodeClass.Metadata.GetVersion().ToString());
@@ -550,135 +588,57 @@ namespace Metasound::Frontend
 
 		bool AddNodesToGraph(FBuildGraphContext& InGraphContext)
 		{
-			for (const FMetasoundFrontendNode& Node : InGraphContext.PagedGraph.Nodes)
+			for (const FMetasoundFrontendNode& FrontendNode : InGraphContext.PagedGraph.Nodes)
 			{
-				const FMetasoundFrontendClass* NodeClass = InGraphContext.BuildContext.FrontendClasses.FindRef(Node.ClassID);
+				const FMetasoundFrontendClass* FrontendNodeClass = InGraphContext.BuildDocumentContext.FrontendClasses.FindRef(FrontendNode.ClassID);
 
-				if (ensure(nullptr != NodeClass))
+				if (ensure(nullptr != FrontendNodeClass))
 				{
-					const FMetasoundFrontendClassInterface& ClassInterface = NodeClass->GetInterfaceForNode(Node);
-
-					// Check frontend node and class interfaces are compatible
-					if (ClassInterface.Inputs.Num() != Node.Interface.Inputs.Num())
+					const FGraphBuilder::FCreateNodeParams CreateNodeParams
 					{
-						UE_LOG(LogMetaSound, Error, TEXT("Skipping node. MetaSound %s contains mismatched number of inputs (%d / %d) on node %s class %s"), *InGraphContext.BuildContext.DebugAssetName, Node.Interface.Inputs.Num(), ClassInterface.Inputs.Num(), *Node.Name.ToString(), *NodeClass->Metadata.GetClassName().ToString());
-						continue;
-					}
-					else if (ClassInterface.Outputs.Num() != Node.Interface.Outputs.Num())
+						InGraphContext.BuildContext.DebugAssetName,
+						FrontendNode,
+						*FrontendNodeClass,
+						InGraphContext.PagedGraph,
+						InGraphContext.GraphClass,
+						InGraphContext.BuildContext.ProxyDataCache, 
+						&InGraphContext.BuildContext.DataTypeRegistry,
+						&InGraphContext.BuildDocumentContext.Graphs
+					};
+					TUniquePtr<INode> Node = FGraphBuilder::CreateNode(CreateNodeParams);
+
+					if (!Node.IsValid())
 					{
-						UE_LOG(LogMetaSound, Error, TEXT("Skipping node. MetaSound %s contains mismatched number of outputs (%d / %d) on node %s class %s"), *InGraphContext.BuildContext.DebugAssetName, Node.Interface.Outputs.Num(), ClassInterface.Outputs.Num(), *Node.Name.ToString(), *NodeClass->Metadata.GetClassName().ToString());
+						LogFailedToCreateNode(InGraphContext.BuildContext.DebugAssetName, FrontendNode, *FrontendNodeClass);
 						continue;
 					}
 
-
-					switch (NodeClass->Metadata.GetType())
+					switch (FrontendNodeClass->Metadata.GetType())
 					{
 						case EMetasoundFrontendClassType::Input:
-						{
-							const FMetasoundFrontendClassInput* ClassInput = FindClassInputForInputNode(InGraphContext.GraphClass, Node);
-
-							if (nullptr != ClassInput)
 							{
-								TUniquePtr<INode> InputNode = CreateInputNode(InGraphContext.BuildContext, Node, *NodeClass, ClassInterface, *ClassInput);
-								if (InputNode.IsValid())
-								{
-									InGraphContext.Graph->AddInputNode(Node.GetID(), ClassInput->Name, MoveTemp(InputNode));
-								}
-								else
-								{
-									LogFailedToCreateNode(InGraphContext.BuildContext.DebugAssetName, Node, *NodeClass);
-								}
+								const FName& InputName = Node->GetInstanceName();
+								InGraphContext.Graph->AddInputNode(FrontendNode.GetID(), InputName, MoveTemp(Node));
 							}
-							else
-							{
-								const FString GraphClassIDString = InGraphContext.GraphClass.ID.ToString();
-								UE_LOG(LogMetaSound, Error, TEXT("MetaSound '%s': Failed to match input node [NodeID:%s, NodeName:%s] to owning graph [ClassID:%s] input."), *InGraphContext.BuildContext.DebugAssetName, *Node.GetID().ToString(), *Node.Name.ToString(), *GraphClassIDString);
-								return false;
-							}
-						}
 						break;
 
 						case EMetasoundFrontendClassType::Output:
-						{
-							const FMetasoundFrontendClassOutput* ClassOutput = FindClassOutputForOutputNode(InGraphContext.GraphClass, Node);
-							if (nullptr != ClassOutput)
-							{
-								TUniquePtr<INode> OutputNode = CreateOutputNode(InGraphContext.BuildContext, Node, *NodeClass, ClassInterface, *ClassOutput);
-								if (OutputNode)
-								{
-									InGraphContext.Graph->AddOutputNode(Node.GetID(), ClassOutput->Name, MoveTemp(OutputNode));
-								}
-								else
-								{
-									LogFailedToCreateNode(InGraphContext.BuildContext.DebugAssetName, Node, *NodeClass);
-								}
-							}
-							else
 							{
-								const FString GraphClassIDString = InGraphContext.GraphClass.ID.ToString();
-								UE_LOG(LogMetaSound, Error, TEXT("MetaSound '%s': Failed to match output node [NodeID:%s, NodeName:%s] to owning graph [ClassID:%s] output."), *InGraphContext.BuildContext.DebugAssetName, *Node.GetID().ToString(), *Node.Name.ToString(), *GraphClassIDString);
-								return false;
+								const FName& OutputName = Node->GetInstanceName();
+								InGraphContext.Graph->AddOutputNode(FrontendNode.GetID(), OutputName, MoveTemp(Node));
 							}
-						}
-						break;
-
-						case EMetasoundFrontendClassType::Graph:
-						{
-							const TSharedPtr<const INode> SubgraphPtr = InGraphContext.BuildContext.Graphs.FindRef(Node.ClassID);
-
-							if (SubgraphPtr.IsValid())
-							{
-PRAGMA_DISABLE_DEPRECATION_WARNINGS
-								InGraphContext.Graph->AddNode(Node.GetID(), SubgraphPtr);
-PRAGMA_ENABLE_DEPRECATION_WARNINGS
-							}
-							else
-							{
-								UE_LOG(LogMetaSound, Error, TEXT("MetaSound '%s': Failed to find subgraph for node [NodeID:%s, NodeName:%s, ClassID:%s]"), *InGraphContext.BuildContext.DebugAssetName, *Node.GetID().ToString(), *Node.Name.ToString(), *Node.ClassID.ToString());
-								return false;
-							}
-						}
 						break;
 
 						case EMetasoundFrontendClassType::Literal:
-						{
-							checkNoEntry(); // Unsupported.
-							return false;
-						}
-
+						case EMetasoundFrontendClassType::Graph:
 						case EMetasoundFrontendClassType::Variable:
-						{
-							TUniquePtr<INode> VariableNode = CreateVariableNode(InGraphContext.BuildContext, Node, *NodeClass, ClassInterface, InGraphContext.PagedGraph);
-							if (VariableNode)
-							{
-								InGraphContext.Graph->AddNode(Node.GetID(), MoveTemp(VariableNode));
-							}
-							else
-							{
-								LogFailedToCreateNode(InGraphContext.BuildContext.DebugAssetName, Node, *NodeClass);
-							}
-						}
-						break;
-
-						// Templates, variable accessors and mutators are
-						// constructed with the same parameters as external nodes.
 						case EMetasoundFrontendClassType::Template:
 						case EMetasoundFrontendClassType::VariableAccessor:
 						case EMetasoundFrontendClassType::VariableDeferredAccessor:
 						case EMetasoundFrontendClassType::VariableMutator:
 						case EMetasoundFrontendClassType::External:
 						default:
-						{
-							TUniquePtr<INode> ExternalNode = CreateExternalNode(InGraphContext.BuildContext, Node, *NodeClass, ClassInterface);
-							if (ExternalNode)
-							{
-								InGraphContext.Graph->AddNode(Node.GetID(), MoveTemp(ExternalNode));
-							}
-							else
-							{
-								LogFailedToCreateNode(InGraphContext.BuildContext.DebugAssetName, Node, *NodeClass);
-							}
-						}
+							InGraphContext.Graph->AddNode(FrontendNode.GetID(), MoveTemp(Node));
 						break;
 					}
 				}
@@ -820,7 +780,7 @@ PRAGMA_ENABLE_DEPRECATION_WARNINGS
 			return bSuccess;
 		}
 
-		TUniquePtr<FFrontendGraph> CreateGraphInternal(FBuildContext& InContext, const FMetasoundFrontendGraphClass& InGraphClass)
+		TUniquePtr<FFrontendGraph> CreateGraphInternal(FBuildContext& InContext, const FBuildDocumentContext& InDocumentContext, const FMetasoundFrontendGraphClass& InGraphClass, const FGuid& InGraphId)
 		{
 			const FString GraphName = InContext.DebugAssetName;
 
@@ -828,10 +788,11 @@ PRAGMA_ENABLE_DEPRECATION_WARNINGS
 			const FMetasoundFrontendGraph& PageGraph = InGraphClass.FindConstGraphChecked(PageID);
 			FBuildGraphContext BuildGraphContext
 			{
-				MakeUnique<FFrontendGraph>(GraphName, InContext.GraphId),
+				MakeUnique<FFrontendGraph>(GraphName, InGraphId),
 				InGraphClass,
 				PageGraph,
-				InContext
+				InContext,
+				InDocumentContext
 			};
 
 			bool bSuccess = AddNodesToGraph(BuildGraphContext);
@@ -851,6 +812,125 @@ PRAGMA_ENABLE_DEPRECATION_WARNINGS
 			}
 		}
 	} // namespace GraphBuilderPrivate
+	 
+	TUniquePtr<INode> FGraphBuilder::CreateNode(const FCreateNodeParams& InParams)
+	{
+		using namespace GraphBuilderPrivate;
+
+		checkf(InParams.ProxyDataCache != nullptr || IsInGameThread(), TEXT("An FProxyDataCache must be used if creating nodes on a thread other than the game thread"));
+
+
+		TUniquePtr<INode> Node;
+
+		const FMetasoundFrontendClassInterface& ClassInterface = InParams.FrontendNodeClass.GetInterfaceForNode(InParams.FrontendNode);
+
+		if (ClassInterface.Inputs.Num() != InParams.FrontendNode.Interface.Inputs.Num())
+		{
+			UE_LOG(LogMetaSound, Error, TEXT("Cannot create node. MetaSound %s contains mismatched number of inputs (%d / %d) on node %s class %s"), *InParams.DebugAssetName, InParams.FrontendNode.Interface.Inputs.Num(), ClassInterface.Inputs.Num(), *InParams.FrontendNode.Name.ToString(), *InParams.FrontendNodeClass.Metadata.GetClassName().ToString());
+			return Node;
+		}
+		else if (ClassInterface.Outputs.Num() != InParams.FrontendNode.Interface.Outputs.Num())
+		{
+			UE_LOG(LogMetaSound, Error, TEXT("Cannot create node. MetaSound %s contains mismatched number of outputs (%d / %d) on node %s class %s"), *InParams.DebugAssetName, InParams.FrontendNode.Interface.Outputs.Num(), ClassInterface.Outputs.Num(), *InParams.FrontendNode.Name.ToString(), *InParams.FrontendNodeClass.Metadata.GetClassName().ToString());
+			return Node;
+		}
+
+		FBuildContext BuildContext
+		{
+			InParams.DebugAssetName,
+			InParams.DataTypeRegistry ? *InParams.DataTypeRegistry : IDataTypeRegistry::Get(),
+			InParams.ProxyDataCache
+		};
+
+		switch (InParams.FrontendNodeClass.Metadata.GetType())
+		{
+			case EMetasoundFrontendClassType::Input:
+			{
+				const FMetasoundFrontendClassInput* ClassInput = FindClassInputForInputNode(InParams.OwningFrontendGraphClass, InParams.FrontendNode);
+
+				if (nullptr != ClassInput)
+				{
+					Node = CreateInputNode(BuildContext, InParams.FrontendNode, InParams.FrontendNodeClass, ClassInterface, *ClassInput);
+				}
+				else
+				{
+					const FString GraphClassIDString = InParams.OwningFrontendGraphClass.ID.ToString();
+					UE_LOG(LogMetaSound, Error, TEXT("MetaSound '%s': Failed to match input node [NodeID:%s, NodeName:%s] to owning graph [ClassID:%s] input."), *InParams.DebugAssetName, *InParams.FrontendNode.GetID().ToString(), *InParams.FrontendNode.Name.ToString(), *GraphClassIDString);
+				}
+			}
+			break;
+
+			case EMetasoundFrontendClassType::Output:
+			{
+				const FMetasoundFrontendClassOutput* ClassOutput = FindClassOutputForOutputNode(InParams.OwningFrontendGraphClass, InParams.FrontendNode);
+				if (nullptr != ClassOutput)
+				{
+					Node = CreateOutputNode(BuildContext, InParams.FrontendNode, InParams.FrontendNodeClass, ClassInterface, *ClassOutput);
+				}
+				else
+				{
+					const FString GraphClassIDString = InParams.OwningFrontendGraphClass.ID.ToString();
+					UE_LOG(LogMetaSound, Error, TEXT("MetaSound '%s': Failed to match output node [NodeID:%s, NodeName:%s] to owning graph [ClassID:%s] output."), *InParams.DebugAssetName, *InParams.FrontendNode.GetID().ToString(), *InParams.FrontendNode.Name.ToString(), *GraphClassIDString);
+				}
+			}
+			break;
+
+			case EMetasoundFrontendClassType::Graph:
+			{
+				if (InParams.Subgraphs == nullptr)
+				{
+					UE_LOG(LogMetaSound, Error, TEXT("MetaSound '%s': Failed to find subgraph for node [NodeID:%s, NodeName:%s, ClassID:%s] because of missing subgraph map"), *InParams.DebugAssetName, *InParams.FrontendNode.GetID().ToString(), *InParams.FrontendNode.Name.ToString(), *InParams.FrontendNode.ClassID.ToString());
+
+				}
+				else
+				{
+					TSharedPtr<const IGraph> SubgraphPtr = InParams.Subgraphs->FindRef(InParams.FrontendNode.ClassID);
+					if (SubgraphPtr.IsValid())
+					{
+						Node = CreateSubgraphNode(BuildContext, SubgraphPtr.ToSharedRef(), InParams.FrontendNode, InParams.FrontendNodeClass, ClassInterface);
+					}
+					else
+					{
+						UE_LOG(LogMetaSound, Error, TEXT("MetaSound '%s': Found invalid subgraph for node [NodeID:%s, NodeName:%s, ClassID:%s]"), *InParams.DebugAssetName, *InParams.FrontendNode.GetID().ToString(), *InParams.FrontendNode.Name.ToString(), *InParams.FrontendNode.ClassID.ToString());
+
+					}
+				}
+			}
+			break;
+
+			case EMetasoundFrontendClassType::Literal:
+			{
+				UE_LOG(LogMetaSound, Warning, TEXT("MetaSound '%s': Adding literal nodes is being deprecated. Please set literal values on the connected node interface directly."), *InParams.DebugAssetName);
+				const FName& DataTypeName = InParams.FrontendNodeClass.Metadata.GetClassName().Name;
+PRAGMA_DISABLE_DEPRECATION_WARNINGS
+				FDefaultLiteralNodeConstructorParams InitData { InParams.FrontendNode.Name, InParams.FrontendNode.GetID(), BuildContext.DataTypeRegistry.CreateDefaultLiteral(DataTypeName)};
+				Node = BuildContext.DataTypeRegistry.CreateLiteralNode(DataTypeName, MoveTemp(InitData));
+PRAGMA_ENABLE_DEPRECATION_WARNINGS
+				break;
+			}
+
+			case EMetasoundFrontendClassType::Variable:
+			{
+				Node = CreateVariableNode(BuildContext, InParams.FrontendNode, InParams.FrontendNodeClass, ClassInterface, InParams.OwningFrontendGraph);
+			}
+			break;
+
+			// Templates, variable accessors and mutators are
+			// constructed with the same parameters as external nodes.
+			case EMetasoundFrontendClassType::Template:
+			case EMetasoundFrontendClassType::VariableAccessor:
+			case EMetasoundFrontendClassType::VariableDeferredAccessor:
+			case EMetasoundFrontendClassType::VariableMutator:
+			case EMetasoundFrontendClassType::External:
+			default:
+			{
+				Node = CreateExternalNode(BuildContext, InParams.FrontendNode, InParams.FrontendNodeClass, ClassInterface);
+			}
+			break;
+		}
+
+		return Node;
+	}
 
 	/** Check that all dependencies are C++ class dependencies. */
 	bool FGraphBuilder::IsFlat(const FMetasoundFrontendDocument& InDocument)
@@ -905,21 +985,20 @@ PRAGMA_ENABLE_DEPRECATION_WARNINGS
 		FBuildContext Context
 		{
 			InDebugAssetName, 					// DebugAssetName
-			{}, 								// FrontendClasses
-			{}, 								// Graphs
 			Frontend::IDataTypeRegistry::Get(), // DataTypeRegistry
-			InProxyDataCache, 					// ProxyDataCache
-			InGraphId							// GraphId
+			&InProxyDataCache, 					// ProxyDataCache
 		};
 
+		FBuildDocumentContext DocumentContext;
+
 		// Gather all references to node classes from external dependencies and subgraphs.
 		for (const FMetasoundFrontendClass& ExtClass : InDependencies)
 		{
-			Context.FrontendClasses.Add(ExtClass.ID, &ExtClass);
+			DocumentContext.FrontendClasses.Add(ExtClass.ID, &ExtClass);
 		}
 		for (const FMetasoundFrontendClass& ExtClass : InSubgraphs)
 		{
-			Context.FrontendClasses.Add(ExtClass.ID, &ExtClass);
+			DocumentContext.FrontendClasses.Add(ExtClass.ID, &ExtClass);
 		}
 
 		// Sort subgraphs so that dependent subgraphs are created in correct order.
@@ -936,7 +1015,7 @@ PRAGMA_ENABLE_DEPRECATION_WARNINGS
 		// Create each subgraph.
 		for (const FMetasoundFrontendGraphClass* FrontendSubgraphPtr : FrontendSubgraphPtrs)
 		{
-			TSharedPtr<const INode> Subgraph(CreateGraphInternal(Context, *FrontendSubgraphPtr).Release());
+			TSharedPtr<const IGraph> Subgraph(CreateGraphInternal(Context, DocumentContext, *FrontendSubgraphPtr, FrontendSubgraphPtr->ID).Release());
 			if (!Subgraph.IsValid())
 			{
 				UE_LOG(LogMetaSound, Warning, TEXT("Failed to create subgraph [SubgraphName: %s] in asset '%s'"), *FrontendSubgraphPtr->Metadata.GetClassName().ToString(), *InDebugAssetName);
@@ -944,12 +1023,12 @@ PRAGMA_ENABLE_DEPRECATION_WARNINGS
 			else
 			{
 				// Add subgraphs to context so they are accessible for subsequent graphs.
-				Context.Graphs.Add(FrontendSubgraphPtr->ID, Subgraph);
+				DocumentContext.Graphs.Add(FrontendSubgraphPtr->ID, Subgraph);
 			}
 		}
 
 		// Create parent graph.
-		return CreateGraphInternal(Context, InGraph);
+		return CreateGraphInternal(Context, DocumentContext, InGraph, InGraphId);
 	}
 
 	TUniquePtr<FFrontendGraph> FGraphBuilder::CreateGraph(const FMetasoundFrontendDocument& InDocument, const FString& InDebugAssetName)
diff --git a/Engine/Plugins/Runtime/Metasound/Source/MetasoundFrontend/Private/MetasoundFrontendNodeClassRegistry.cpp b/Engine/Plugins/Runtime/Metasound/Source/MetasoundFrontend/Private/MetasoundFrontendNodeClassRegistry.cpp
--- a/Engine/Plugins/Runtime/Metasound/Source/MetasoundFrontend/Private/MetasoundFrontendNodeClassRegistry.cpp
+++ b/Engine/Plugins/Runtime/Metasound/Source/MetasoundFrontend/Private/MetasoundFrontendNodeClassRegistry.cpp
@@ -26,6 +26,7 @@
 #include "MetasoundFrontendProxyDataCache.h"
 #include "MetasoundFrontendRegistryTransaction.h"
 #include "MetasoundFrontendSearchEngine.h"
+#include "MetasoundGraphNode.h"
 #include "MetasoundLog.h"
 #include "MetasoundRouter.h"
 #include "MetasoundTrace.h"
@@ -105,118 +106,6 @@ namespace Metasound::Frontend
 #endif // WITH_EDITOR
 		}
 
-		// FGraphNode is used to create unique INodes based off of a IGraph.
-		//
-		// Individual nodes need to reflect their InstanceName and InstanceID, but otherwise
-		// they simply encapsulate a shared set of behavior. To minimize memory usage, a single
-		// shared IGraph is used for all nodes referring to the same IGraph.
-		class FGraphNode : public INode
-		{
-			// This adapter class forwards the correct FBuilderOperatorParams
-			// to the graph's operator creation method. Many operator creation
-			// methods downcast the supplied INode in `FBuilderOperatorParams`
-			// and so it is required that it point to the correct runtime instance
-			// when calling CreateOperator(...)
-			class FGraphOperatorFactoryAdapter : public IOperatorFactory 
-			{
-			public:
-				FGraphOperatorFactoryAdapter(const IGraph& InGraph)
-				: Graph(&InGraph)
-				, GraphFactory(InGraph.GetDefaultOperatorFactory())
-				{
-				}
-
-				virtual ~FGraphOperatorFactoryAdapter() = default;
-
-				virtual TUniquePtr<IOperator> CreateOperator(const FBuildOperatorParams& InParams, FBuildResults& OutResults) override
-				{
-					// Create default data references for inputs which do not have data references
-					// already provided. These need to be created here, but the defaults
-					// cannot be piped into the wrapped graph. The wrapped graph cannot
-					// be modified. 
-					FInputVertexInterfaceData InputData = InParams.InputData;
-					CreateDefaultsIfNotBound(InParams.OperatorSettings, InputData);
-
-					FBuildOperatorParams ForwardParams 
-					{
-						*Graph,  // Point to correct INode instance
-						InParams.OperatorSettings,
-						InputData,
-						InParams.Environment,
-						InParams.Builder,
-						InParams.GraphRenderCost
-					};
-					return GraphFactory->CreateOperator(ForwardParams, OutResults);
-				}
-			private:
-				const IGraph* Graph; // Only store pointer because owning node keeps wrapped IGraph alive.  
-				FOperatorFactorySharedRef GraphFactory;
-			};
-
-		public:
-			FGraphNode(const FNodeInitData& InNodeInitData, TSharedRef<const IGraph> InGraphToWrap)
-			: FGraphNode(FNodeData(InNodeInitData.InstanceName, InNodeInitData.InstanceID, InGraphToWrap->GetMetadata().DefaultInterface), MoveTemp(InGraphToWrap))
-			{
-			}
-
-			FGraphNode(FNodeData InNodeData, TSharedRef<const IGraph> InGraphToWrap)
-			: NodeData(MoveTemp(InNodeData))
-			, Factory(MakeShared<FGraphOperatorFactoryAdapter>(*InGraphToWrap))
-			, Graph(MoveTemp(InGraphToWrap))
-			{
-			}
-
-			virtual const FName& GetInstanceName() const override
-			{
-				// Use the instance name of underlying graph because it refers
-				// to the actual asset name.
-				return Graph->GetInstanceName();
-			}
-
-			virtual const FGuid& GetInstanceID() const override
-			{
-				return NodeData.ID;
-			}
-
-			virtual const FNodeClassMetadata& GetMetadata() const override
-			{
-				return Graph->GetMetadata();
-			}
-
-			virtual const FVertexInterface& GetVertexInterface() const override
-			{
-				return NodeData.Interface;
-			}
-
-			virtual void SetDefaultInput(const FVertexName& InVertexName, const FLiteral& InLiteral) override
-			{
-				if (FInputDataVertex* Vertex = NodeData.Interface.GetInputInterface().Find(InVertexName))
-				{
-					Vertex->SetDefaultLiteral(InLiteral);
-				}
-				else
-				{
-					UE_LOG(LogMetaSound, Error, TEXT("Could not set default input. Could not find input vertex %s on node %s of node class %s"), *InVertexName.ToString(), *GetInstanceName().ToString(), *GetMetadata().ClassName.ToString());
-				}
-			}
-
-			virtual TSharedPtr<const IOperatorData> GetOperatorData() const override
-			{
-				return NodeData.OperatorData;
-			}
-
-			virtual FOperatorFactorySharedRef GetDefaultOperatorFactory() const override
-			{
-				return Factory;
-			}
-
-		private:
-
-			FNodeData NodeData;
-			TSharedRef<FGraphOperatorFactoryAdapter> Factory;
-			TSharedRef<const IGraph> Graph;
-		};
-
 
 		// FDocumentNodeClassRegistryEntry encapsulates a node registry entry for a FGraph
 		class FDocumentNodeClassRegistryEntry : public INodeClassRegistryEntry
diff --git a/Engine/Plugins/Runtime/Metasound/Source/MetasoundFrontend/Private/MetasoundGraphNode.cpp b/Engine/Plugins/Runtime/Metasound/Source/MetasoundFrontend/Private/MetasoundGraphNode.cpp
new file mode 100644
--- /dev/null
+++ b/Engine/Plugins/Runtime/Metasound/Source/MetasoundFrontend/Private/MetasoundGraphNode.cpp
@@ -0,0 +1,96 @@
+// Copyright Epic Games, Inc. All Rights Reserved.
+
+#include "MetasoundGraphNode.h"
+
+#include "MetasoundBuilderInterface.h"
+#include "MetasoundFrontendDataTypeRegistry.h"
+#include "MetasoundLog.h"
+#include "MetasoundNodeInterface.h"
+#include "Templates/SharedPointer.h"
+
+namespace Metasound::Frontend
+{
+	FGraphNode::FGraphOperatorFactoryAdapter::FGraphOperatorFactoryAdapter(const IGraph& InGraph)
+	: Graph(&InGraph)
+	, GraphFactory(InGraph.GetDefaultOperatorFactory())
+	{
+	}
+
+	TUniquePtr<IOperator> FGraphNode::FGraphOperatorFactoryAdapter::CreateOperator(const FBuildOperatorParams& InParams, FBuildResults& OutResults)
+	{
+		// Create default data references for inputs which do not have data references
+		// already provided. These need to be created here, but the defaults
+		// cannot be piped into the wrapped graph. The wrapped graph cannot
+		// be modified. 
+		FInputVertexInterfaceData InputData = InParams.InputData;
+		CreateDefaultsIfNotBound(InParams.OperatorSettings, InputData);
+
+		FBuildOperatorParams ForwardParams 
+		{
+			*Graph,  // Point to correct INode instance
+			InParams.OperatorSettings,
+			InputData,
+			InParams.Environment,
+			InParams.Builder,
+			InParams.GraphRenderCost
+		};
+		return GraphFactory->CreateOperator(ForwardParams, OutResults);
+	}
+
+	FGraphNode::FGraphNode(const FNodeInitData& InNodeInitData, TSharedRef<const IGraph> InGraphToWrap)
+	: FGraphNode(FNodeData(InNodeInitData.InstanceName, InNodeInitData.InstanceID, InGraphToWrap->GetMetadata().DefaultInterface), MoveTemp(InGraphToWrap))
+	{
+	}
+
+	FGraphNode::FGraphNode(FNodeData InNodeData, TSharedRef<const IGraph> InGraphToWrap)
+	: NodeData(MoveTemp(InNodeData))
+	, Factory(MakeShared<FGraphOperatorFactoryAdapter>(*InGraphToWrap))
+	, Graph(MoveTemp(InGraphToWrap))
+	{
+	}
+
+	const FName& FGraphNode::GetInstanceName() const
+	{
+		// Use the instance name of underlying graph because it refers
+		// to the actual asset name.
+		return Graph->GetInstanceName();
+	}
+
+	const FGuid& FGraphNode::GetInstanceID() const
+	{
+		return NodeData.ID;
+	}
+
+	const FNodeClassMetadata& FGraphNode::GetMetadata() const
+	{
+		return Graph->GetMetadata();
+	}
+
+	const FVertexInterface& FGraphNode::GetVertexInterface() const
+	{
+		return NodeData.Interface;
+	}
+
+	void FGraphNode::SetDefaultInput(const FVertexName& InVertexName, const FLiteral& InLiteral)
+	{
+		if (FInputDataVertex* Vertex = NodeData.Interface.GetInputInterface().Find(InVertexName))
+		{
+			Vertex->SetDefaultLiteral(InLiteral);
+		}
+		else
+		{
+			UE_LOG(LogMetaSound, Error, TEXT("Could not set default input. Could not find input vertex %s on node %s of node class %s"), *InVertexName.ToString(), *GetInstanceName().ToString(), *GetMetadata().ClassName.ToString());
+		}
+	}
+
+	TSharedPtr<const IOperatorData> FGraphNode::GetOperatorData() const
+	{
+		return NodeData.OperatorData;
+	}
+
+	FOperatorFactorySharedRef FGraphNode::GetDefaultOperatorFactory() const
+	{
+		return Factory;
+	}
+}
+
diff --git a/Engine/Plugins/Runtime/Metasound/Source/MetasoundFrontend/Private/MetasoundGraphNode.h b/Engine/Plugins/Runtime/Metasound/Source/MetasoundFrontend/Private/MetasoundGraphNode.h
new file mode 100644
--- /dev/null
+++ b/Engine/Plugins/Runtime/Metasound/Source/MetasoundFrontend/Private/MetasoundGraphNode.h
@@ -0,0 +1,62 @@
+// Copyright Epic Games, Inc. All Rights Reserved.
+
+#pragma once
+
+#include "MetasoundBuilderInterface.h"
+#include "MetasoundNodeInterface.h"
+#include "Templates/SharedPointer.h"
+
+namespace Metasound::Frontend
+{
+	// FGraphNode is used to create unique INodes based off of a IGraph.
+	//
+	// Individual nodes need to reflect their InstanceName and InstanceID, but otherwise
+	// they simply encapsulate a shared set of behavior. To minimize memory usage, a single
+	// shared IGraph is used for all nodes referring to the same IGraph.
+	class FGraphNode : public INode
+	{
+		// This adapter class forwards the correct FBuilderOperatorParams
+		// to the graph's operator creation method. Many operator creation
+		// methods downcast the supplied INode in `FBuilderOperatorParams`
+		// and so it is required that it point to the correct runtime instance
+		// when calling CreateOperator(...)
+		class FGraphOperatorFactoryAdapter : public IOperatorFactory 
+		{
+		public:
+			FGraphOperatorFactoryAdapter(const IGraph& InGraph);
+			virtual ~FGraphOperatorFactoryAdapter() = default;
+
+			virtual TUniquePtr<IOperator> CreateOperator(const FBuildOperatorParams& InParams, FBuildResults& OutResults) override;
+
+		private:
+			const IGraph* Graph; // Only store pointer because owning node keeps wrapped IGraph alive.  
+			FOperatorFactorySharedRef GraphFactory;
+		};
+
+	public:
+		FGraphNode(const FNodeInitData& InNodeInitData, TSharedRef<const IGraph> InGraphToWrap);
+
+		FGraphNode(FNodeData InNodeData, TSharedRef<const IGraph> InGraphToWrap);
+
+		virtual const FName& GetInstanceName() const override;
+
+		virtual const FGuid& GetInstanceID() const override;
+
+		virtual const FNodeClassMetadata& GetMetadata() const override;
+
+		virtual const FVertexInterface& GetVertexInterface() const override;
+
+		virtual void SetDefaultInput(const FVertexName& InVertexName, const FLiteral& InLiteral) override;
+
+		virtual TSharedPtr<const IOperatorData> GetOperatorData() const override;
+
+		virtual FOperatorFactorySharedRef GetDefaultOperatorFactory() const override;
+
+	private:
+
+		FNodeData NodeData;
+		TSharedRef<FGraphOperatorFactoryAdapter> Factory;
+		TSharedRef<const IGraph> Graph;
+	};
+}
+
diff --git a/Engine/Plugins/Runtime/Metasound/Source/MetasoundFrontend/Public/MetasoundFrontendGraphBuilder.h b/Engine/Plugins/Runtime/Metasound/Source/MetasoundFrontend/Public/MetasoundFrontendGraphBuilder.h
--- a/Engine/Plugins/Runtime/Metasound/Source/MetasoundFrontend/Public/MetasoundFrontendGraphBuilder.h
+++ b/Engine/Plugins/Runtime/Metasound/Source/MetasoundFrontend/Public/MetasoundFrontendGraphBuilder.h
@@ -25,6 +25,8 @@ namespace Metasound
 
 namespace Metasound::Frontend
 {
+	// Forward declare
+	class IDataTypeRegistry;
 
 	/** FGraphBuilder builds a FFrontendGraph from a FMetasoundDocument
 	 * or FMetasoundFrontendClass.
@@ -32,6 +34,39 @@ namespace Metasound::Frontend
 	class METASOUNDFRONTEND_API FGraphBuilder
 	{
 	public:
+
+		struct FCreateNodeParams
+		{
+			/** Debug name used for logging errors and warnings */
+			const FString& DebugAssetName;
+			/** The Frontend representation of the node being created. */
+			const FMetasoundFrontendNode& FrontendNode; 
+			/** The Frontend Node Class representing the node. */
+			const FMetasoundFrontendClass& FrontendNodeClass; 
+			/** The Frontend Graph which contains the Frontend Node being created. */
+			const FMetasoundFrontendGraph& OwningFrontendGraph;
+			/** The Frontend Graph Class which contains the OwningFrontendGraph. */
+			const FMetasoundFrontendGraphClass& OwningFrontendGraphClass;
+
+			/** Optional ProxyDataCache. This is needed for creating nodes when being
+			 * executed off of the game thread. */
+			const Frontend::FProxyDataCache* ProxyDataCache = nullptr;
+
+			/** Optional DataTypeRegistry. Providing a reusable data type registry
+			 * can avoid repeated lookups to gather the registry. */
+			const Frontend::IDataTypeRegistry* DataTypeRegistry = nullptr;
+
+			/** Optional subgraph map. Subgraphs are not fully supported in MetaSound
+			 * so this is kept optional to avoid creation of unused TMap<>s. */
+			const TMap<FGuid, TSharedPtr<const IGraph>>* Subgraphs = nullptr;
+		};
+
+		/** Create a INode from a document representation of the node. 
+		 *
+		 * Note: This does not connect the node or place it in any graph. 
+		 */
+		static TUniquePtr<INode> CreateNode(const FCreateNodeParams& InParams);
+
 		/** Check that all dependencies are C++ class dependencies. 
 		 * 
 		 * @param InDocument - Document containing dependencies.
diff --git a/Engine/Plugins/Runtime/Metasound/Source/MetasoundFrontend/Public/MetasoundNodeRegistrationMacro.h b/Engine/Plugins/Runtime/Metasound/Source/MetasoundFrontend/Public/MetasoundNodeRegistrationMacro.h
--- a/Engine/Plugins/Runtime/Metasound/Source/MetasoundFrontend/Public/MetasoundNodeRegistrationMacro.h
+++ b/Engine/Plugins/Runtime/Metasound/Source/MetasoundFrontend/Public/MetasoundNodeRegistrationMacro.h
@@ -284,12 +284,12 @@ namespace Metasound
 	}
 }
 
-#define METASOUND_REGISTER_NODE_AND_EXTENSION(NodeClass, ConfigurationClass) \
+#define METASOUND_REGISTER_NODE_AND_CONFIGURATION(NodeClass, ConfigurationClass) \
 	 static_assert(std::is_base_of<::Metasound::INodeBase, NodeClass>::value, "To be registered as a  Metasound Node," #NodeClass "need to be a derived class from Metasound::INodeBase, Metasound::INode, or Metasound::FNode."); \
 	 static_assert(::Metasound::TIsNodeConstructorSupported<NodeClass>::Value, "In order to be registered as a Metasound Node, " #NodeClass " needs to implement the following public constructor: " #NodeClass "(Metasound::FNodeData InNodeData);"); \
 	 static bool bSuccessfullyRegistered##NodeClass  = ::Metasound::Frontend::INodeClassRegistry::Get()->EnqueueInitCommand([](){ ::Metasound::Frontend::NodeRegistrationPrivate::RegisterNodeUsingMacro<NodeClass, ConfigurationClass>(); }); // This static bool is useful for debugging, but also is the only way the compiler will let us call this function outside of an expression.
 
-#define METASOUND_REGISTER_NODE(NodeClass) METASOUND_REGISTER_NODE_AND_EXTENSION(NodeClass, void)
+#define METASOUND_REGISTER_NODE(NodeClass) METASOUND_REGISTER_NODE_AND_CONFIGURATION(NodeClass, void)
 /*
 Macros to help define various FText node fields.
 */
