[MetaSounds] Cache Relative Render cost object in the Metasound Operator Cache.
#swarm https://p4-swarm.epicgames.net/reviews/41494486

diff --git a/Engine/Plugins/Runtime/Metasound/Source/MetasoundGenerator/Private/MetasoundGenerator.cpp b/Engine/Plugins/Runtime/Metasound/Source/MetasoundGenerator/Private/MetasoundGenerator.cpp
--- a/Engine/Plugins/Runtime/Metasound/Source/MetasoundGenerator/Private/MetasoundGenerator.cpp
+++ b/Engine/Plugins/Runtime/Metasound/Source/MetasoundGenerator/Private/MetasoundGenerator.cpp
@@ -329,19 +329,6 @@ PRAGMA_DISABLE_DEPRECATION_WARNINGS
 		ParameterPackSendAddress = UMetasoundParameterPack::CreateSendAddressFromEnvironment(InInitParams.Environment);
 PRAGMA_ENABLE_DEPRECATION_WARNINGS
 		ParameterPackReceiver = FDataTransmissionCenter::Get().RegisterNewReceiver<FMetasoundParameterStorageWrapper>(ParameterPackSendAddress, FReceiverInitParams{ OperatorSettings });
-
-		if (InInitParams.GraphRenderCost.IsValid())
-		{
-			// Use the provided render cost if it exists
-			GraphRenderCost = InInitParams.GraphRenderCost;
-		}
-		else
-		{
-			// If no render cost object is provided make one. 
-			GraphRenderCost = FGraphRenderCost::MakeGraphRenderCost();
-			// Add to the init params so that it is available to metasound graph being built.
-			InInitParams.GraphRenderCost = GraphRenderCost;
-		}
 	}
 
 	bool FMetasoundGenerator::RemoveGraphSetCallback(const FDelegateHandle& Handle)
@@ -918,6 +905,8 @@ PRAGMA_ENABLE_DEPRECATION_WARNINGS
 
 		if (!bDidUseCachedOperator)
 		{
+			GraphRenderCost = FGraphRenderCost::MakeGraphRenderCost();
+			InParams.GraphRenderCost = GraphRenderCost;
 			BuildGraph(MoveTemp(InParams));
 		}
 	}
@@ -960,11 +949,19 @@ PRAGMA_ENABLE_DEPRECATION_WARNINGS
 				// Apply and default inputs to the operator.
 				GeneratorBuilder::ResetGraphOperatorInputs(OperatorSettings, MoveTemp(InInitParams.DefaultParameters), GraphOperatorAndInputs.Inputs);
 
+				GraphRenderCost = GraphOperatorAndInputs.RenderCost;
+				if (!GraphRenderCost)
+				{
+					GraphRenderCost = FGraphRenderCost::MakeGraphRenderCost();
+				}
+				
+				InInitParams.GraphRenderCost = GraphRenderCost;
+
 				// Reset operator internal state before playing it.
 				if (IOperator::FResetFunction Reset = GraphOperatorAndInputs.Operator->GetResetFunction())
 				{
 					EnvironmentPtr = MakeUnique<FMetasoundEnvironment>(InInitParams.Environment);
-					IOperator::FResetParams ResetParams {OperatorSettings, *EnvironmentPtr};
+					IOperator::FResetParams ResetParams {OperatorSettings, *EnvironmentPtr, GraphRenderCost.Get()};
 					Reset(GraphOperatorAndInputs.Operator.Get(), ResetParams);
 				}
 				
@@ -1016,8 +1013,7 @@ PRAGMA_ENABLE_DEPRECATION_WARNINGS
 				// references may result in a race condition on deallocating references.
 				ClearGraph();
 
-				OperatorPool->AddOperator(*OperatorPoolID, MoveTemp(GraphOperator), MoveTemp(InputData));
-
+				OperatorPool->AddOperator(*OperatorPoolID, MoveTemp(GraphOperator), MoveTemp(InputData), MoveTemp(GraphRenderCost));
 			}
 		}
 	}
diff --git a/Engine/Plugins/Runtime/Metasound/Source/MetasoundGenerator/Private/MetasoundOperatorCache.cpp b/Engine/Plugins/Runtime/Metasound/Source/MetasoundGenerator/Private/MetasoundOperatorCache.cpp
--- a/Engine/Plugins/Runtime/Metasound/Source/MetasoundGenerator/Private/MetasoundOperatorCache.cpp
+++ b/Engine/Plugins/Runtime/Metasound/Source/MetasoundGenerator/Private/MetasoundOperatorCache.cpp
@@ -291,10 +291,10 @@ namespace Metasound
 		const FOperatorPoolEntryID EntryID{InOperatorID, FOperatorSettings{OperatorPoolPrivate::DefaultSampleRateForDeprecatedAPI, OperatorPoolPrivate::BlockRate}};
 		AddOperator(EntryID, MoveTemp(InOperator), MoveTemp(InputData));
 	}
-
-	void FOperatorPool::AddOperator(const FOperatorPoolEntryID& InOperatorID, TUniquePtr<IOperator>&& InOperator, FInputVertexInterfaceData&& InInputData)
+	
+	void FOperatorPool::AddOperator(const FOperatorPoolEntryID& InOperatorID, TUniquePtr<IOperator>&& InOperator, FInputVertexInterfaceData&& InInputData, TSharedPtr<FGraphRenderCost>&& InRenderCost)
 	{
-		AddOperator(InOperatorID, { MoveTemp(InOperator), MoveTemp(InInputData) });
+		AddOperator(InOperatorID, { MoveTemp(InOperator), MoveTemp(InInputData), MoveTemp(InRenderCost) });
 	}
 
 	bool FOperatorPool::ExecuteTaskAsync(FOperatorPool::FTaskFunction&& InFunction)
diff --git a/Engine/Plugins/Runtime/Metasound/Source/MetasoundGenerator/Public/MetasoundGenerator.h b/Engine/Plugins/Runtime/Metasound/Source/MetasoundGenerator/Public/MetasoundGenerator.h
--- a/Engine/Plugins/Runtime/Metasound/Source/MetasoundGenerator/Public/MetasoundGenerator.h
+++ b/Engine/Plugins/Runtime/Metasound/Source/MetasoundGenerator/Public/MetasoundGenerator.h
@@ -385,8 +385,10 @@ namespace Metasound
 		// This map provides setters for all of the input nodes in the metasound graph. 
 		// It is used when processing named parameters in a parameter pack.
 		TMap<FName, MetasoundGeneratorPrivate::FParameterPackSetter> ParameterPackSetters;
-	private:
 
+		TSharedPtr<FGraphRenderCost> GraphRenderCost;
+		
+	private:
 		// While parameter packs may arrive via the IAudioParameterInterface system,
 		// a faster method of sending parameters is via the QueueParameterPack function 
 		// and this queue.
@@ -400,8 +402,7 @@ namespace Metasound
 		double RenderTime;
 		bool bDoRuntimeRenderTiming;
 		TUniquePtr<MetasoundGeneratorPrivate::FRenderTimer> RenderTimer;
-
-		TSharedPtr<FGraphRenderCost> GraphRenderCost;
+		
 		std::atomic<float> RelativeRenderCost;
 	};
 
diff --git a/Engine/Plugins/Runtime/Metasound/Source/MetasoundGenerator/Public/MetasoundOperatorCache.h b/Engine/Plugins/Runtime/Metasound/Source/MetasoundGenerator/Public/MetasoundOperatorCache.h
--- a/Engine/Plugins/Runtime/Metasound/Source/MetasoundGenerator/Public/MetasoundOperatorCache.h
+++ b/Engine/Plugins/Runtime/Metasound/Source/MetasoundGenerator/Public/MetasoundOperatorCache.h
@@ -116,7 +116,7 @@ namespace Metasound
 
 		UE_DEPRECATED(5.5, "Use AddOperator(const FOperatorPoolEntryID&, ...) instead")
 		void AddOperator(const FGuid& InOperatorID, TUniquePtr<IOperator>&& InOperator, FInputVertexInterfaceData&& InputData);
-		void AddOperator(const FOperatorPoolEntryID& InOperatorID, TUniquePtr<IOperator>&& InOperator, FInputVertexInterfaceData&& InputData);
+		void AddOperator(const FOperatorPoolEntryID& InOperatorID, TUniquePtr<IOperator>&& InOperator, FInputVertexInterfaceData&& InputData, TSharedPtr<FGraphRenderCost>&& InRenderCost = {});
 
 		UE_DEPRECATED(5.5, "Use AddOperator(const FOperatorPoolEntryID&, ...) instead")
 		void AddOperator(const FGuid& InOperatorID, FOperatorAndInputs&& OperatorAndInputs);
diff --git a/Engine/Plugins/Runtime/Metasound/Source/MetasoundGraphCore/Public/MetasoundOperatorInterface.h b/Engine/Plugins/Runtime/Metasound/Source/MetasoundGraphCore/Public/MetasoundOperatorInterface.h
--- a/Engine/Plugins/Runtime/Metasound/Source/MetasoundGraphCore/Public/MetasoundOperatorInterface.h
+++ b/Engine/Plugins/Runtime/Metasound/Source/MetasoundGraphCore/Public/MetasoundOperatorInterface.h
@@ -190,5 +190,6 @@ namespace Metasound
 	{
 		TUniquePtr<IOperator> Operator;
 		FInputVertexInterfaceData Inputs;
+		TSharedPtr<FGraphRenderCost> RenderCost;
 	};
 }
