[Truncation Warnings] Update EdGraphSchema, EdGraphActions, and SGraphEditor to use FVector2f

diff --git a/Engine/Plugins/FX/Niagara/Source/NiagaraEditor/Private/Widgets/SNiagaraScriptGraph.cpp b/Engine/Plugins/FX/Niagara/Source/NiagaraEditor/Private/Widgets/SNiagaraScriptGraph.cpp
--- a/Engine/Plugins/FX/Niagara/Source/NiagaraEditor/Private/Widgets/SNiagaraScriptGraph.cpp
+++ b/Engine/Plugins/FX/Niagara/Source/NiagaraEditor/Private/Widgets/SNiagaraScriptGraph.cpp
@@ -165,7 +165,7 @@ SNiagaraScriptGraph::~SNiagaraScriptGraph()
 void SNiagaraScriptGraph::UpdateViewModel(TSharedRef<FNiagaraScriptGraphViewModel> InNewModel)
 {
 	bool bKeepOldGraphView = false;
-	FVector2D Location;
+	FVector2f Location;
 	float ZoomAmount =1.0f;
 	FGuid BookmarkId = FGuid();
 
@@ -599,7 +599,7 @@ void SNiagaraScriptGraph::NodesPasted(const TSet<UEdGraphNode*>& PastedNodes)
 
 void SNiagaraScriptGraph::PositionPastedNodes(const TSet<UEdGraphNode*>& PastedNodes)
 {
-	FVector2D AvgNodePosition(0.0f, 0.0f);
+	FVector2f AvgNodePosition(0.0f, 0.0f);
 
 	for (UEdGraphNode* PastedNode : PastedNodes)
 	{
@@ -611,7 +611,7 @@ void SNiagaraScriptGraph::PositionPastedNodes(const TSet<UEdGraphNode*>& PastedN
 	AvgNodePosition.X *= InvNumNodes;
 	AvgNodePosition.Y *= InvNumNodes;
 
-	FVector2D PasteLocation = GraphEditor->GetPasteLocation();
+	FVector2f PasteLocation = GraphEditor->GetPasteLocation2f();
 	for (UEdGraphNode* PastedNode : PastedNodes)
 	{
 		PastedNode->NodePosX = (PastedNode->NodePosX - AvgNodePosition.X) + PasteLocation.X;
@@ -989,7 +989,7 @@ void SNiagaraScriptGraph::FocusGraphSearchBox()
 void SNiagaraScriptGraph::OnCreateComment()
 {
 	FNiagaraSchemaAction_NewComment CommentAction = FNiagaraSchemaAction_NewComment(GraphEditor);
-	if (CommentAction.PerformAction(ViewModel->GetGraph(), nullptr, GraphEditor->GetPasteLocation(), false))
+	if (CommentAction.PerformAction(ViewModel->GetGraph(), nullptr, GraphEditor->GetPasteLocation2f(), false))
 	{
 		ViewModel->GetGraph()->NotifyGraphNeedsRecompile();
 	}
diff --git a/Engine/Plugins/FX/Niagara/Source/NiagaraEditor/Public/EdGraphSchema_Niagara.h b/Engine/Plugins/FX/Niagara/Source/NiagaraEditor/Public/EdGraphSchema_Niagara.h
--- a/Engine/Plugins/FX/Niagara/Source/NiagaraEditor/Public/EdGraphSchema_Niagara.h
+++ b/Engine/Plugins/FX/Niagara/Source/NiagaraEditor/Public/EdGraphSchema_Niagara.h
@@ -40,13 +40,14 @@ struct FNiagaraSchemaAction_NewNode : public FEdGraphSchemaAction
 	{}
 
 	//~ Begin FEdGraphSchemaAction Interface
-	NIAGARAEDITOR_API virtual UEdGraphNode* PerformAction(class UEdGraph* ParentGraph, UEdGraphPin* FromPin, const FVector2D Location, bool bSelectNewNode = true) override;
-	NIAGARAEDITOR_API virtual UEdGraphNode* PerformAction(class UEdGraph* ParentGraph, TArray<UEdGraphPin*>& FromPins, const FVector2D Location, bool bSelectNewNode = true) override;
+	using FEdGraphSchemaAction::PerformAction; // Prevent hiding of deprecated base class function with FVector2D
+	NIAGARAEDITOR_API virtual UEdGraphNode* PerformAction(class UEdGraph* ParentGraph, UEdGraphPin* FromPin, const FVector2f& Location, bool bSelectNewNode = true) override;
+	NIAGARAEDITOR_API virtual UEdGraphNode* PerformAction(class UEdGraph* ParentGraph, TArray<UEdGraphPin*>& FromPins, const FVector2f& Location, bool bSelectNewNode = true) override;
 	NIAGARAEDITOR_API virtual void AddReferencedObjects( FReferenceCollector& Collector ) override;
 	//~ End FEdGraphSchemaAction Interface
 
 	template <typename NodeType>
-	static NodeType* SpawnNodeFromTemplate(class UEdGraph* ParentGraph, NodeType* InTemplateNode, const FVector2D Location, bool bSelectNewNode = true)
+	static NodeType* SpawnNodeFromTemplate(class UEdGraph* ParentGraph, NodeType* InTemplateNode, const UE::Slate::FDeprecateVector2DParameter& Location, bool bSelectNewNode = true)
 	{
 		FNiagaraSchemaAction_NewNode Action;
 		Action.NodeTemplate = InTemplateNode;
@@ -72,7 +73,8 @@ public:
 	{}
 
 	//~ Begin FEdGraphSchemaAction Interface
-	NIAGARAEDITOR_API virtual UEdGraphNode* PerformAction(class UEdGraph* ParentGraph, UEdGraphPin* FromPin, const FVector2D Location, bool bSelectNewNode = true) override;
+	using FEdGraphSchemaAction::PerformAction; // Prevent hiding of deprecated base class function with FVector2D
+	NIAGARAEDITOR_API virtual UEdGraphNode* PerformAction(class UEdGraph* ParentGraph, UEdGraphPin* FromPin, const FVector2f& Location, bool bSelectNewNode = true) override;
 	//~ End FEdGraphSchemaAction Interface
 
 private:
@@ -106,9 +108,9 @@ class UEdGraphSchema_Niagara : public UEdGraphSchema
 	NIAGARAEDITOR_API virtual void BreakPinLinks(UEdGraphPin& TargetPin, bool bSendsNodeNotification) const override;
 	NIAGARAEDITOR_API virtual FConnectionDrawingPolicy* CreateConnectionDrawingPolicy(int32 InBackLayerID, int32 InFrontLayerID, float InZoomFactor, const FSlateRect& InClippingRect, class FSlateWindowElementList& InDrawElements, class UEdGraph* InGraphObj) const override;
 	NIAGARAEDITOR_API virtual void ResetPinToAutogeneratedDefaultValue(UEdGraphPin* Pin, bool bCallModifyCallbacks = true) const override;
-	NIAGARAEDITOR_API virtual void OnPinConnectionDoubleCicked(UEdGraphPin* PinA, UEdGraphPin* PinB, const FVector2D& GraphPosition) const override; 
+	NIAGARAEDITOR_API virtual void OnPinConnectionDoubleCicked(UEdGraphPin* PinA, UEdGraphPin* PinB, const FVector2f& GraphPosition) const override; 
 	virtual bool ShouldAlwaysPurgeOnModification() const override { return false; }
-	NIAGARAEDITOR_API virtual void DroppedAssetsOnGraph(const TArray<FAssetData>& Assets, const FVector2D& GraphPosition, UEdGraph* Graph) const override;
+	NIAGARAEDITOR_API virtual void DroppedAssetsOnGraph(const TArray<FAssetData>& Assets, const FVector2f& GraphPosition, UEdGraph* Graph) const override;
 	NIAGARAEDITOR_API virtual void GetAssetsGraphHoverMessage(const TArray<FAssetData>& Assets, const UEdGraph* HoverGraph, FString& OutTooltipText, bool& OutOkIcon) const override;
 
 	/**
diff --git a/Engine/Plugins/FX/Niagara/Source/NiagaraEditor/Public/EdGraphSchema_NiagaraSystemOverview.h b/Engine/Plugins/FX/Niagara/Source/NiagaraEditor/Public/EdGraphSchema_NiagaraSystemOverview.h
--- a/Engine/Plugins/FX/Niagara/Source/NiagaraEditor/Public/EdGraphSchema_NiagaraSystemOverview.h
+++ b/Engine/Plugins/FX/Niagara/Source/NiagaraEditor/Public/EdGraphSchema_NiagaraSystemOverview.h
@@ -22,7 +22,7 @@ class UEdGraphSchema_NiagaraSystemOverview : public UEdGraphSchema
 	NIAGARAEDITOR_API virtual bool IsCacheVisualizationOutOfDate(int32 InVisualizationCacheID) const override;
 	NIAGARAEDITOR_API virtual int32 GetCurrentVisualizationCacheID() const override;
 	NIAGARAEDITOR_API virtual void ForceVisualizationCacheClear() const override;
-	NIAGARAEDITOR_API virtual void DroppedAssetsOnGraph(const TArray<FAssetData>& Assets, const FVector2D& GraphPosition, UEdGraph* Graph) const override;
+	NIAGARAEDITOR_API virtual void DroppedAssetsOnGraph(const TArray<FAssetData>& Assets, const FVector2f& GraphPosition, UEdGraph* Graph) const override;
 	NIAGARAEDITOR_API virtual void GetAssetsGraphHoverMessage(const TArray<FAssetData>& Assets, const UEdGraph* HoverGraph, FString& OutTooltipText, bool& OutOkIcon) const override;
 
 private:
diff --git a/Engine/Plugins/FX/Niagara/Source/NiagaraEditorWidgets/Private/SNiagaraOverviewGraph.cpp b/Engine/Plugins/FX/Niagara/Source/NiagaraEditorWidgets/Private/SNiagaraOverviewGraph.cpp
--- a/Engine/Plugins/FX/Niagara/Source/NiagaraEditorWidgets/Private/SNiagaraOverviewGraph.cpp
+++ b/Engine/Plugins/FX/Niagara/Source/NiagaraEditorWidgets/Private/SNiagaraOverviewGraph.cpp
@@ -210,7 +210,7 @@ void SNiagaraOverviewGraph::PreClose()
 {
 	if (ViewModel.IsValid() && GraphEditor.IsValid())
 	{
-		FVector2D Location;
+		FVector2f Location;
 		float Zoom;
 		GraphEditor->GetViewLocation(Location, Zoom);
 		ViewModel->SetViewSettings(FNiagaraGraphViewSettings(Location, Zoom));
@@ -323,7 +323,7 @@ void SNiagaraOverviewGraph::OnCreateComment()
 	if (ViewModel->GetSystemViewModel()->GetEditMode() != ENiagaraSystemViewModelEditMode::EmitterAsset)
 	{
 		FNiagaraSchemaAction_NewComment CommentAction = FNiagaraSchemaAction_NewComment(GraphEditor);
-		CommentAction.PerformAction(ViewModel->GetGraph(), nullptr, GraphEditor->GetPasteLocation(), false);
+		CommentAction.PerformAction(ViewModel->GetGraph(), nullptr, GraphEditor->GetPasteLocation2f(), false);
 	}
 }
 
@@ -393,7 +393,7 @@ void SNiagaraOverviewGraph::NodesPasted(const TSet<UEdGraphNode*>& PastedNodes)
 
 void SNiagaraOverviewGraph::PositionPastedNodes(const TSet<UEdGraphNode*>& PastedNodes)
 {
-	FVector2D AvgNodePosition(0.0f, 0.0f);
+	FVector2f AvgNodePosition(0.0f, 0.0f);
 
 	for (UEdGraphNode* PastedNode : PastedNodes)
 	{
@@ -405,7 +405,7 @@ void SNiagaraOverviewGraph::PositionPastedNodes(const TSet<UEdGraphNode*>& Paste
 	AvgNodePosition.X *= InvNumNodes;
 	AvgNodePosition.Y *= InvNumNodes;
 
-	FVector2D PasteLocation = GraphEditor->GetPasteLocation();
+	FVector2f PasteLocation = GraphEditor->GetPasteLocation2f();
 	for (UEdGraphNode* PastedNode : PastedNodes)
 	{
 		
@@ -471,7 +471,7 @@ void SNiagaraOverviewGraph::OnNodesCreated(const FEdGraphEditAction& Action)
 {
 	if((Action.Action & GRAPHACTION_AddNode) != 0)
 	{
-		FVector2D PasteLocation = GraphEditor->GetPasteLocation();
+		FVector2f PasteLocation = GraphEditor->GetPasteLocation2f();
 		int32 NodeCounter = 0;
 
 		for(const UEdGraphNode* NewNode : Action.Nodes)
diff --git a/Engine/Plugins/MovieScene/MovieRenderPipeline/Source/MovieRenderPipelineEditor/Private/Customizations/Graph/MovieGraphApplyCVarPresetNodeCustomization.h b/Engine/Plugins/MovieScene/MovieRenderPipeline/Source/MovieRenderPipelineEditor/Private/Customizations/Graph/MovieGraphApplyCVarPresetNodeCustomization.h
--- a/Engine/Plugins/MovieScene/MovieRenderPipeline/Source/MovieRenderPipelineEditor/Private/Customizations/Graph/MovieGraphApplyCVarPresetNodeCustomization.h
+++ b/Engine/Plugins/MovieScene/MovieRenderPipeline/Source/MovieRenderPipelineEditor/Private/Customizations/Graph/MovieGraphApplyCVarPresetNodeCustomization.h
@@ -50,7 +50,7 @@ protected:
 
 		// Put the new node to the right of this node
 		const TObjectPtr<UEdGraphNode> GraphNode = InPresetNode->GraphNode;
-		const FVector2d NewLocation(GraphNode->NodePosX + 300, GraphNode->NodePosY);
+		const FVector2f NewLocation(GraphNode->NodePosX + 300.0f, GraphNode->NodePosY);
 
 		// Create the new node (providing FromPin will trigger the action to connect the new node and the preset node)
 		UEdGraphPin* FromPin = GraphNode->FindPin(NAME_None, EGPD_Output);
diff --git a/Engine/Plugins/MovieScene/MovieRenderPipeline/Source/MovieRenderPipelineEditor/Private/Graph/MovieEdGraphNode.cpp b/Engine/Plugins/MovieScene/MovieRenderPipeline/Source/MovieRenderPipelineEditor/Private/Graph/MovieEdGraphNode.cpp
--- a/Engine/Plugins/MovieScene/MovieRenderPipeline/Source/MovieRenderPipelineEditor/Private/Graph/MovieEdGraphNode.cpp
+++ b/Engine/Plugins/MovieScene/MovieRenderPipeline/Source/MovieRenderPipelineEditor/Private/Graph/MovieEdGraphNode.cpp
@@ -439,7 +439,7 @@ void UMoviePipelineEdGraphNode::PromotePropertyToVariable(const FMovieGraphPrope
 		NewAction->NodeClass = UMovieGraphVariableNode::StaticClass();
 
 		// Put the new node in a roughly ok-ish position relative to this node
-		const FVector2d NewLocation(NodePosX - 200, NodePosY);
+		const FVector2f NewLocation(NodePosX - 200.0f, NodePosY);
 
 		// Note: Providing FromPin will trigger the action to connect the new node and this node
 		UEdGraphPin* FromPin = FindPin(TargetProperty.Name, EGPD_Input);
diff --git a/Engine/Plugins/MovieScene/MovieRenderPipeline/Source/MovieRenderPipelineEditor/Private/Graph/MovieGraphSchema.cpp b/Engine/Plugins/MovieScene/MovieRenderPipeline/Source/MovieRenderPipelineEditor/Private/Graph/MovieGraphSchema.cpp
--- a/Engine/Plugins/MovieScene/MovieRenderPipeline/Source/MovieRenderPipelineEditor/Private/Graph/MovieGraphSchema.cpp
+++ b/Engine/Plugins/MovieScene/MovieRenderPipeline/Source/MovieRenderPipelineEditor/Private/Graph/MovieGraphSchema.cpp
@@ -644,7 +644,7 @@ FMovieGraphSchemaAction_NewNode::FMovieGraphSchemaAction_NewNode(FText InNodeCat
 	
 }
 
-UEdGraphNode* FMovieGraphSchemaAction_NewNode::PerformAction(UEdGraph* ParentGraph, UEdGraphPin* FromPin, const FVector2D Location, bool bSelectNewNode)
+UEdGraphNode* FMovieGraphSchemaAction_NewNode::PerformAction(UEdGraph* ParentGraph, UEdGraphPin* FromPin, const FVector2f& Location, bool bSelectNewNode)
 {
 	UMovieGraphConfig* RuntimeGraph = CastChecked<UMoviePipelineEdGraph>(ParentGraph)->GetPipelineGraph();
 	const FScopedTransaction Transaction(LOCTEXT("GraphEditor_NewNode", "Create Pipeline Graph Node."));
@@ -693,7 +693,7 @@ FMovieGraphSchemaAction_NewVariableNode::FMovieGraphSchemaAction_NewVariableNode
 	
 }
 
-UEdGraphNode* FMovieGraphSchemaAction_NewVariableNode::PerformAction(UEdGraph* ParentGraph, UEdGraphPin* FromPin, const FVector2D Location, bool bSelectNewNode)
+UEdGraphNode* FMovieGraphSchemaAction_NewVariableNode::PerformAction(UEdGraph* ParentGraph, UEdGraphPin* FromPin, const FVector2f& Location, bool bSelectNewNode)
 {
 	UMovieGraphConfig* RuntimeGraph = CastChecked<UMoviePipelineEdGraph>(ParentGraph)->GetPipelineGraph();
 	const FScopedTransaction Transaction(LOCTEXT("GraphEditor_NewVariableNode", "Add New Variable Accessor Node"));
@@ -724,13 +724,13 @@ UEdGraphNode* FMovieGraphSchemaAction_NewVariableNode::PerformAction(UEdGraph* P
 	return GraphNode;
 }
 
-UEdGraphNode* FMovieGraphSchemaAction_NewComment::PerformAction(UEdGraph* ParentGraph, UEdGraphPin* FromPin, const FVector2D Location, bool bSelectNewNode)
+UEdGraphNode* FMovieGraphSchemaAction_NewComment::PerformAction(UEdGraph* ParentGraph, UEdGraphPin* FromPin, const FVector2f& Location, bool bSelectNewNode)
 {
 	UEdGraphNode_Comment* CommentTemplate = NewObject<UEdGraphNode_Comment>();
 
 	const TSharedPtr<SGraphEditor> GraphEditorPtr = SGraphEditor::FindGraphEditorForGraph(ParentGraph);
 
-	FVector2D SpawnLocation = Location;
+	FVector2f SpawnLocation = Location;
 	FSlateRect Bounds;
 	if (GraphEditorPtr && GraphEditorPtr->GetBoundsForSelectedNodes(Bounds, 50.0f))
 	{
diff --git a/Engine/Plugins/MovieScene/MovieRenderPipeline/Source/MovieRenderPipelineEditor/Private/Graph/MovieGraphSchema.h b/Engine/Plugins/MovieScene/MovieRenderPipeline/Source/MovieRenderPipelineEditor/Private/Graph/MovieGraphSchema.h
--- a/Engine/Plugins/MovieScene/MovieRenderPipeline/Source/MovieRenderPipelineEditor/Private/Graph/MovieGraphSchema.h
+++ b/Engine/Plugins/MovieScene/MovieRenderPipeline/Source/MovieRenderPipelineEditor/Private/Graph/MovieGraphSchema.h
@@ -135,7 +135,8 @@ struct FMovieGraphSchemaAction_NewNode : public FMovieGraphSchemaAction
 	}
 	
 	//~ Begin FEdGraphSchemaAction Interface
-	virtual UEdGraphNode* PerformAction(UEdGraph* ParentGraph, UEdGraphPin* FromPin, const FVector2D Location, bool bSelectNewNode = true) override;
+	using FEdGraphSchemaAction::PerformAction; // Prevent hiding of deprecated base class function with FVector2D
+	virtual UEdGraphNode* PerformAction(UEdGraph* ParentGraph, UEdGraphPin* FromPin, const FVector2f& Location, bool bSelectNewNode = true) override;
 	//~ End FEdGraphSchemaAction Interface
 };
 
@@ -160,7 +161,8 @@ struct FMovieGraphSchemaAction_NewVariableNode : public FMovieGraphSchemaAction
 	virtual ~FMovieGraphSchemaAction_NewVariableNode() override = default;
 	
 	//~ Begin FEdGraphSchemaAction Interface
-	virtual UEdGraphNode* PerformAction(UEdGraph* ParentGraph, UEdGraphPin* FromPin, const FVector2D Location, bool bSelectNewNode = true) override;
+	using FEdGraphSchemaAction::PerformAction; // Prevent hiding of deprecated base class function with FVector2D
+	virtual UEdGraphNode* PerformAction(UEdGraph* ParentGraph, UEdGraphPin* FromPin, const FVector2f& Location, bool bSelectNewNode = true) override;
 	//~ End FEdGraphSchemaAction Interface
 
 private:
@@ -190,6 +192,7 @@ struct FMovieGraphSchemaAction_NewComment : public FMovieGraphSchemaAction
 
 	//~ Begin FEdGraphSchemaAction Interface
 	virtual FName GetTypeId() const override { return StaticGetTypeId(); }
-	virtual UEdGraphNode* PerformAction(UEdGraph* ParentGraph, UEdGraphPin* FromPin, const FVector2D Location, bool bSelectNewNode = true) override;
+	using FEdGraphSchemaAction::PerformAction;
+	virtual UEdGraphNode* PerformAction(UEdGraph* ParentGraph, UEdGraphPin* FromPin, const FVector2f& Location, bool bSelectNewNode = true) override;
 	//~ End FEdGraphSchemaAction Interface
 };
\ No newline at end of file
diff --git a/Engine/Plugins/MovieScene/MovieRenderPipeline/Source/MovieRenderPipelineEditor/Private/Widgets/Graph/SMovieGraphConfigPanel.cpp b/Engine/Plugins/MovieScene/MovieRenderPipeline/Source/MovieRenderPipelineEditor/Private/Widgets/Graph/SMovieGraphConfigPanel.cpp
--- a/Engine/Plugins/MovieScene/MovieRenderPipeline/Source/MovieRenderPipelineEditor/Private/Widgets/Graph/SMovieGraphConfigPanel.cpp
+++ b/Engine/Plugins/MovieScene/MovieRenderPipeline/Source/MovieRenderPipelineEditor/Private/Widgets/Graph/SMovieGraphConfigPanel.cpp
@@ -239,10 +239,10 @@ void SMoviePipelineGraphPanel::OnCreateComment() const
 		FMovieGraphSchemaAction_NewComment CommentAction;
 		
 		const TSharedPtr<SGraphEditor> GraphEditorPtr = SGraphEditor::FindGraphEditorForGraph(CurrentGraph->PipelineEdGraph);
-		FVector2D Location = FVector2D::ZeroVector;
+		FVector2f Location = FVector2f::ZeroVector;
 		if (GraphEditorPtr)
 		{
-			Location = GraphEditorPtr->GetPasteLocation();
+			Location = GraphEditorPtr->GetPasteLocation2f();
 		}
 		
 		CommentAction.PerformAction(CurrentGraph->PipelineEdGraph, nullptr, Location);
@@ -400,7 +400,7 @@ void SMoviePipelineGraphPanel::PasteNodes()
 {
 	if (GraphEditorWidget.IsValid())
 	{
-		PasteNodesHere(GraphEditorWidget->GetPasteLocation());
+		PasteNodesHere(GraphEditorWidget->GetPasteLocation2f());
 	}
 }
 
diff --git a/Engine/Plugins/MovieScene/SequencerScripting/Source/SequencerScriptingEditor/Private/SequencerTools.cpp b/Engine/Plugins/MovieScene/SequencerScripting/Source/SequencerScriptingEditor/Private/SequencerTools.cpp
--- a/Engine/Plugins/MovieScene/SequencerScripting/Source/SequencerScriptingEditor/Private/SequencerTools.cpp
+++ b/Engine/Plugins/MovieScene/SequencerScripting/Source/SequencerScriptingEditor/Private/SequencerTools.cpp
@@ -949,7 +949,7 @@ FSequencerQuickBindingResult USequencerToolsFunctionLibrary::CreateQuickBinding(
 	UEdGraphPin* ThenPin = NewEventEndpoint->FindPin(UEdGraphSchema_K2::PN_Then, EGPD_Output);
 	UEdGraphPin* BoundObjectPin = FMovieSceneDirectorBlueprintUtils::FindCallTargetPin(NewEventEndpoint, EndpointDefinition.PossibleCallTargetClass);
 
-	FVector2D NodePosition(NewEventEndpoint->NodePosX + 400.f, NewEventEndpoint->NodePosY);
+	FVector2f NodePosition(NewEventEndpoint->NodePosX + 400.f, NewEventEndpoint->NodePosY);
 	UEdGraphNode* NewNode = Action.PerformAction(NewEventEndpoint->GetGraph(), BoundObjectPin ? BoundObjectPin : ThenPin, NodePosition);
 
 	if (NewNode == nullptr)
diff --git a/Engine/Plugins/Mutable/Source/CustomizableObjectEditor/Private/MuCOE/CustomizableObjectMacroLibrary/CustomizableObjectGraphEditorToolkit.cpp b/Engine/Plugins/Mutable/Source/CustomizableObjectEditor/Private/MuCOE/CustomizableObjectMacroLibrary/CustomizableObjectGraphEditorToolkit.cpp
--- a/Engine/Plugins/Mutable/Source/CustomizableObjectEditor/Private/MuCOE/CustomizableObjectMacroLibrary/CustomizableObjectGraphEditorToolkit.cpp
+++ b/Engine/Plugins/Mutable/Source/CustomizableObjectEditor/Private/MuCOE/CustomizableObjectMacroLibrary/CustomizableObjectGraphEditorToolkit.cpp
@@ -38,7 +38,7 @@ void FCustomizableObjectGraphEditorToolkit::PasteNodes()
 		return;
 	}
 
-	PasteNodesHere(GraphEditor->GetPasteLocation());
+	PasteNodesHere(GraphEditor->GetPasteLocation2f());
 }
 
 
@@ -325,7 +325,7 @@ bool FCustomizableObjectGraphEditorToolkit::CanRenameNodes() const
 
 void FCustomizableObjectGraphEditorToolkit::CreateCommentBoxFromKey()
 {
-	CreateCommentBox(GraphEditor->GetPasteLocation());
+	CreateCommentBox(GraphEditor->GetPasteLocation2f());
 }
 
 
diff --git a/Engine/Plugins/Mutable/Source/CustomizableObjectEditor/Private/MuCOE/CustomizableObjectSchemaActions.h b/Engine/Plugins/Mutable/Source/CustomizableObjectEditor/Private/MuCOE/CustomizableObjectSchemaActions.h
--- a/Engine/Plugins/Mutable/Source/CustomizableObjectEditor/Private/MuCOE/CustomizableObjectSchemaActions.h
+++ b/Engine/Plugins/Mutable/Source/CustomizableObjectEditor/Private/MuCOE/CustomizableObjectSchemaActions.h
@@ -3,6 +3,7 @@
 #pragma once
 
 #include "EdGraphSchema_K2_Actions.h"
+#include "Types/SlateVector2.h"
 
 #include "CustomizableObjectSchemaActions.generated.h"
 
@@ -29,8 +30,9 @@ struct CUSTOMIZABLEOBJECTEDITOR_API FCustomizableObjectSchemaAction_NewNode : pu
 	{}
 
 	// FEdGraphSchemaAction interface
-	virtual UEdGraphNode* PerformAction(class UEdGraph* ParentGraph, UEdGraphPin* FromPin, const FVector2D Location, bool bSelectNewNode = true) override;
-	virtual UEdGraphNode* PerformAction(class UEdGraph* ParentGraph, TArray<UEdGraphPin*>& FromPins, const FVector2D Location, bool bSelectNewNode = true) override;
+	using FEdGraphSchemaAction::PerformAction; // Prevent hiding of deprecated base class function with FVector2D
+	virtual UEdGraphNode* PerformAction(class UEdGraph* ParentGraph, UEdGraphPin* FromPin, const FVector2f& Location, bool bSelectNewNode = true) override;
+	virtual UEdGraphNode* PerformAction(class UEdGraph* ParentGraph, TArray<UEdGraphPin*>& FromPins, const FVector2f& Location, bool bSelectNewNode = true) override;
 	virtual void AddReferencedObjects( FReferenceCollector& Collector ) override;
 
 	/** Reimplementaiton of EdGraphSchema::CreateNode(...). Performs the overlap calculation before calling AutowireNewNode(...). AutowireNewNode can induce a call to a ReconstrucNode() which removes pins required for the calculation. */
@@ -38,7 +40,7 @@ struct CUSTOMIZABLEOBJECTEDITOR_API FCustomizableObjectSchemaAction_NewNode : pu
 	// End of FEdGraphSchemaAction interface
 
 	template <typename NodeType>
-	static NodeType* InstantSpawn(class UEdGraph* ParentGraph, NodeType* InTemplateNode, const FVector2D Location)
+	static NodeType* InstantSpawn(class UEdGraph* ParentGraph, NodeType* InTemplateNode, const UE::Slate::FDeprecateVector2DParameter& Location)
 	{
 		FEdGraphSchemaAction_K2NewNode Action;
 		Action.NodeTemplate = InTemplateNode;
@@ -67,6 +69,7 @@ struct CUSTOMIZABLEOBJECTEDITOR_API FCustomizableObjectSchemaAction_Paste : publ
 	{}
 
 	// FEdGraphSchemaAction interface
-	virtual UEdGraphNode* PerformAction(class UEdGraph* ParentGraph, UEdGraphPin* FromPin, const FVector2D Location, bool bSelectNewNode = true) override;
+	using FEdGraphSchemaAction::PerformAction; // Prevent hiding of deprecated base class function with FVector2D
+	virtual UEdGraphNode* PerformAction(class UEdGraph* ParentGraph, UEdGraphPin* FromPin, const FVector2f& Location, bool bSelectNewNode = true) override;
 	// End of FEdGraphSchemaAction interface
 };
diff --git a/Engine/Plugins/Mutable/Source/CustomizableObjectEditor/Private/MuCOE/EdGraphSchema_CustomizableObject.cpp b/Engine/Plugins/Mutable/Source/CustomizableObjectEditor/Private/MuCOE/EdGraphSchema_CustomizableObject.cpp
--- a/Engine/Plugins/Mutable/Source/CustomizableObjectEditor/Private/MuCOE/EdGraphSchema_CustomizableObject.cpp
+++ b/Engine/Plugins/Mutable/Source/CustomizableObjectEditor/Private/MuCOE/EdGraphSchema_CustomizableObject.cpp
@@ -144,7 +144,7 @@ TSharedPtr<FCustomizableObjectGraphEditorToolkit> GetCustomizableObjectEditor(co
 }
 
 
-UEdGraphNode* FCustomizableObjectSchemaAction_NewNode::PerformAction(class UEdGraph* ParentGraph, UEdGraphPin* FromPin, const FVector2D Location, bool bSelectNewNode)
+UEdGraphNode* FCustomizableObjectSchemaAction_NewNode::PerformAction(class UEdGraph* ParentGraph, UEdGraphPin* FromPin, const FVector2f& Location, bool bSelectNewNode)
 {
 	UEdGraphNode* ResultNode = nullptr;
 
@@ -159,14 +159,14 @@ UEdGraphNode* FCustomizableObjectSchemaAction_NewNode::PerformAction(class UEdGr
 			FromPin->Modify();
 		}
 
-		ResultNode = FCustomizableObjectSchemaAction_NewNode::CreateNode(ParentGraph, FromPin, Location, NodeTemplate);
+		ResultNode = FCustomizableObjectSchemaAction_NewNode::CreateNode(ParentGraph, FromPin, FDeprecateSlateVector2D(Location), NodeTemplate);
 	}
 
 	return ResultNode;
 }
 
 
-UEdGraphNode* FCustomizableObjectSchemaAction_NewNode::PerformAction(class UEdGraph* ParentGraph, TArray<UEdGraphPin*>& FromPins, const FVector2D Location, bool bSelectNewNode) 
+UEdGraphNode* FCustomizableObjectSchemaAction_NewNode::PerformAction(class UEdGraph* ParentGraph, TArray<UEdGraphPin*>& FromPins, const FVector2f& Location, bool bSelectNewNode) 
 {
 	UEdGraphNode* ResultNode = 0;
 
@@ -247,13 +247,13 @@ void FCustomizableObjectSchemaAction_NewNode::AddReferencedObjects( FReferenceCo
 ////////////////////////////////////////
 // FCustomizableObjectSchemaAction_Paste
 
-UEdGraphNode* FCustomizableObjectSchemaAction_Paste::PerformAction(class UEdGraph* ParentGraph, UEdGraphPin* FromPin, const FVector2D Location, bool bSelectNewNode/* = true*/)
+UEdGraphNode* FCustomizableObjectSchemaAction_Paste::PerformAction(class UEdGraph* ParentGraph, UEdGraphPin* FromPin, const FVector2f& Location, bool bSelectNewNode/* = true*/)
 {
 	TSharedPtr<FCustomizableObjectGraphEditorToolkit> CustomizableObjectEditor = GetCustomizableObjectEditor(ParentGraph);
 
 	if (CustomizableObjectEditor.IsValid() && CustomizableObjectEditor->CanPasteNodes())
 	{
-		CustomizableObjectEditor->PasteNodesHere(Location);
+		CustomizableObjectEditor->PasteNodesHere(FDeprecateSlateVector2D(Location));
 	}
 	
 	return nullptr;
@@ -1160,7 +1160,7 @@ bool IsSpawnableAsset(const FAssetData& InAsset, ESpawnableObjectType& OutObject
 }
 
 
-void UEdGraphSchema_CustomizableObject::DroppedAssetsOnGraph(const TArray<FAssetData>& Assets, const FVector2D& GraphPosition, UEdGraph* Graph) const
+void UEdGraphSchema_CustomizableObject::DroppedAssetsOnGraph(const TArray<FAssetData>& Assets, const FVector2f& GraphPosition, UEdGraph* Graph) const
 {
 	// To prevent overlapping when multiple assets are dropped at the same time on the graph
 	const int PixelOffset = 20;
@@ -1246,7 +1246,7 @@ void UEdGraphSchema_CustomizableObject::DroppedAssetsOnGraph(const TArray<FAsset
 }
 
 
-void UEdGraphSchema_CustomizableObject::OnPinConnectionDoubleCicked(UEdGraphPin* PinA, UEdGraphPin* PinB, const FVector2D& GraphPosition) const
+void UEdGraphSchema_CustomizableObject::OnPinConnectionDoubleCicked(UEdGraphPin* PinA, UEdGraphPin* PinB, const FVector2f& GraphPosition) const
 {
 	const FScopedTransaction Transaction(LOCTEXT("CreateRerouteNodeOnWire", "Create Reroute Node"));
 
@@ -1254,12 +1254,12 @@ void UEdGraphSchema_CustomizableObject::OnPinConnectionDoubleCicked(UEdGraphPin*
 	ParentGraph->Modify();
 
 	// This constant is duplicated from inside of SGraphNodeKnot
-	const FVector2D NodeSpacerSize(42.0f, 24.0f);
-	const FVector2D KnotTopLeft = GraphPosition - (NodeSpacerSize * 0.5f);
+	const FVector2f NodeSpacerSize(42.0f, 24.0f);
+	const FVector2f KnotTopLeft = GraphPosition - (NodeSpacerSize * 0.5f);
 
 	// Create a new knot
 	UCustomizableObjectNodeReroute* DefaultNodeReroute = CastChecked<UCustomizableObjectNodeReroute>(UCustomizableObjectNodeReroute::StaticClass()->GetDefaultObject());
-	UCustomizableObjectNodeReroute* NodeReroute = CastChecked<UCustomizableObjectNodeReroute>(FCustomizableObjectSchemaAction_NewNode::CreateNode(ParentGraph, nullptr, KnotTopLeft, DefaultNodeReroute));
+	UCustomizableObjectNodeReroute* NodeReroute = CastChecked<UCustomizableObjectNodeReroute>(FCustomizableObjectSchemaAction_NewNode::CreateNode(ParentGraph, nullptr, FDeprecateSlateVector2D(KnotTopLeft), DefaultNodeReroute));
 
 	PinA->BreakLinkTo(PinB);
 	PinA->MakeLinkTo((PinA->Direction == EGPD_Output) ? NodeReroute->GetInputPin() : NodeReroute->GetOutputPin());
diff --git a/Engine/Plugins/Mutable/Source/CustomizableObjectEditor/Public/MuCOE/EdGraphSchema_CustomizableObject.h b/Engine/Plugins/Mutable/Source/CustomizableObjectEditor/Public/MuCOE/EdGraphSchema_CustomizableObject.h
--- a/Engine/Plugins/Mutable/Source/CustomizableObjectEditor/Public/MuCOE/EdGraphSchema_CustomizableObject.h
+++ b/Engine/Plugins/Mutable/Source/CustomizableObjectEditor/Public/MuCOE/EdGraphSchema_CustomizableObject.h
@@ -57,8 +57,8 @@ public:
 	virtual const FPinConnectionResponse CanCreateConnection(const UEdGraphPin* A, const UEdGraphPin* B) const override;
 	virtual FLinearColor GetPinTypeColor(const FEdGraphPinType& PinType) const override;
 	virtual bool ShouldHidePinDefaultValue(UEdGraphPin* Pin) const override;
-	virtual void DroppedAssetsOnGraph(const TArray<FAssetData>& Assets, const FVector2D& GraphPosition, UEdGraph* Graph) const override;
-	virtual void OnPinConnectionDoubleCicked(UEdGraphPin* PinA, UEdGraphPin* PinB, const FVector2D& GraphPosition) const override;
+	virtual void DroppedAssetsOnGraph(const TArray<FAssetData>& Assets, const FVector2f& GraphPosition, UEdGraph* Graph) const override;
+	virtual void OnPinConnectionDoubleCicked(UEdGraphPin* PinA, UEdGraphPin* PinB, const FVector2f& GraphPosition) const override;
 	virtual void GetContextMenuActions(UToolMenu* Menu, UGraphNodeContextMenuContext* Context) const override;
 	virtual void BreakNodeLinks(UEdGraphNode& TargetNode) const override;
 	virtual void BreakPinLinks(UEdGraphPin& TargetPin, bool bSendsNodeNotification) const override;
diff --git a/Engine/Plugins/PCG/Source/PCGEditor/Private/PCGEditor.cpp b/Engine/Plugins/PCG/Source/PCGEditor/Private/PCGEditor.cpp
--- a/Engine/Plugins/PCG/Source/PCGEditor/Private/PCGEditor.cpp
+++ b/Engine/Plugins/PCG/Source/PCGEditor/Private/PCGEditor.cpp
@@ -1583,7 +1583,7 @@ FReply FPCGEditor::OnSpawnNodeByShortcut(FInputChord InChord, const FVector2D& I
 	if (Action.IsValid())
 	{
 		TArray<UEdGraphPin*> DummyPins;
-		Action->PerformAction(InGraph, DummyPins, InPosition);
+		Action->PerformAction(InGraph, DummyPins, UE::Slate::FDeprecateVector2DParameter(InPosition));
 		return FReply::Handled();
 	}
 	else
@@ -2473,7 +2473,7 @@ void FPCGEditor::PasteNodes()
 {
 	if (GraphEditorWidget.IsValid())
 	{
-		PasteNodesHere(GraphEditorWidget->GetPasteLocation());
+		PasteNodesHere(GraphEditorWidget->GetPasteLocation2f());
 	}
 }
 
@@ -2667,10 +2667,10 @@ void FPCGEditor::OnCreateComment()
 		FPCGEditorGraphSchemaAction_NewComment CommentAction;
 
 		TSharedPtr<SGraphEditor> GraphEditorPtr = SGraphEditor::FindGraphEditorForGraph(PCGEditorGraph);
-		FVector2D Location = FVector2D::ZeroVector;
+		FVector2f Location = FVector2f::ZeroVector;
 		if (GraphEditorPtr)
 		{
-			Location = GraphEditorPtr->GetPasteLocation();
+			Location = GraphEditorPtr->GetPasteLocation2f();
 		}
 
 		CommentAction.PerformAction(PCGEditorGraph, nullptr, Location);
diff --git a/Engine/Plugins/PCG/Source/PCGEditor/Private/PCGEditorGraphSchema.cpp b/Engine/Plugins/PCG/Source/PCGEditor/Private/PCGEditorGraphSchema.cpp
--- a/Engine/Plugins/PCG/Source/PCGEditor/Private/PCGEditorGraphSchema.cpp
+++ b/Engine/Plugins/PCG/Source/PCGEditor/Private/PCGEditorGraphSchema.cpp
@@ -753,9 +753,9 @@ void UPCGEditorGraphSchema::GetDataAssetActions(FGraphActionMenuBuilder& ActionM
 	});
 }
 
-void UPCGEditorGraphSchema::DroppedAssetsOnGraph(const TArray<FAssetData>& Assets, const FVector2D& GraphPosition, UEdGraph* Graph) const
+void UPCGEditorGraphSchema::DroppedAssetsOnGraph(const TArray<FAssetData>& Assets, const FVector2f& GraphPosition, UEdGraph* Graph) const
 {
-	FVector2D GraphPositionOffset = GraphPosition;
+	FVector2f GraphPositionOffset = GraphPosition;
 	constexpr float PositionOffsetIncrementY = 50.f;
 	UEdGraphPin* NullFromPin = nullptr;
 
@@ -772,7 +772,7 @@ void UPCGEditorGraphSchema::DroppedAssetsOnGraph(const TArray<FAssetData>& Asset
 			if(AssetClass->IsChildOf(UPCGGraphInterface::StaticClass()))
 			{
 				SubgraphPaths.Add(AssetData.GetSoftObjectPath());
-				SubgraphPositions.Add(GraphPositionOffset);
+				SubgraphPositions.Add(FDeprecateSlateVector2D(GraphPositionOffset));
 				GraphPositionOffset.Y += PositionOffsetIncrementY;
 			}
 			else if (PCGEditorUtils::IsAssetPCGBlueprint(AssetData))
@@ -796,7 +796,7 @@ void UPCGEditorGraphSchema::DroppedAssetsOnGraph(const TArray<FAssetData>& Asset
 			{
 				// Delay creation so we can open a menu, once, if needed.
 				SettingsPaths.Add(AssetData.GetSoftObjectPath());
-				SettingsGraphPositions.Add(GraphPositionOffset);
+				SettingsGraphPositions.Add(FDeprecateSlateVector2D(GraphPositionOffset));
 				GraphPositionOffset.Y += PositionOffsetIncrementY;
 			}
 		}
@@ -805,7 +805,7 @@ void UPCGEditorGraphSchema::DroppedAssetsOnGraph(const TArray<FAssetData>& Asset
 	UPCGEditorGraph* EditorGraph = CastChecked<UPCGEditorGraph>(Graph);
 
 	TSharedPtr<SGraphEditor> GraphEditor = SGraphEditor::FindGraphEditorForGraph(EditorGraph);
-	const FVector2D MouseCursorLocation = FSlateApplication::Get().GetCursorPos();
+	const FDeprecateSlateVector2D MouseCursorLocation = FSlateApplication::Get().GetCursorPos();
 
 	// If we've dragged settings assets or a graph, we might want to open a menu (ergo this call)
 	if (!SettingsPaths.IsEmpty())
@@ -846,12 +846,12 @@ void UPCGEditorGraphSchema::GetAssetsGraphHoverMessage(const TArray<FAssetData>&
 	OutOkIcon = false;
 }
 
-void UPCGEditorGraphSchema::OnPinConnectionDoubleCicked(UEdGraphPin* PinA, UEdGraphPin* PinB, const FVector2D& GraphPosition) const
+void UPCGEditorGraphSchema::OnPinConnectionDoubleCicked(UEdGraphPin* PinA, UEdGraphPin* PinB, const FVector2f& GraphPosition) const
 {
 	const FScopedTransaction Transaction(*FPCGEditorCommon::ContextIdentifier, LOCTEXT("PCGCreateRerouteNodeOnWire", "Create Reroute Node"), nullptr);
 
-	const FVector2D NodeSpacerSize(42.0f, 24.0f);
-	const FVector2D KnotTopLeft = GraphPosition - (NodeSpacerSize * 0.5f);
+	const FVector2f NodeSpacerSize(42.0f, 24.0f);
+	const FVector2f KnotTopLeft = GraphPosition - (NodeSpacerSize * 0.5f);
 
 	UEdGraph* EditorGraph = PinA->GetOwningNode()->GetGraph();
 	EditorGraph->Modify();
diff --git a/Engine/Plugins/PCG/Source/PCGEditor/Private/PCGEditorGraphSchema.h b/Engine/Plugins/PCG/Source/PCGEditor/Private/PCGEditorGraphSchema.h
--- a/Engine/Plugins/PCG/Source/PCGEditor/Private/PCGEditorGraphSchema.h
+++ b/Engine/Plugins/PCG/Source/PCGEditor/Private/PCGEditorGraphSchema.h
@@ -42,9 +42,9 @@ public:
 	virtual void BreakPinLinks(UEdGraphPin& TargetPin, bool bSendsNodeNotification) const override;
 	virtual void BreakSinglePinLink(UEdGraphPin* SourcePin, UEdGraphPin* TargetPin) const override;
 	virtual FConnectionDrawingPolicy* CreateConnectionDrawingPolicy(int32 InBackLayerID, int32 InFrontLayerID, float InZoomFactor, const FSlateRect& InClippingRect, class FSlateWindowElementList& InDrawElements, class UEdGraph* InGraphObj) const override;
-	virtual void DroppedAssetsOnGraph(const TArray<FAssetData>& Assets, const FVector2D& GraphPosition, UEdGraph* Graph) const override;
+	virtual void DroppedAssetsOnGraph(const TArray<FAssetData>& Assets, const FVector2f& GraphPosition, UEdGraph* Graph) const override;
 	virtual void GetAssetsGraphHoverMessage(const TArray<FAssetData>& Assets, const UEdGraph* HoverGraph, FString& OutTooltipText, bool& OutOkIcon) const override;
-	virtual void OnPinConnectionDoubleCicked(UEdGraphPin* PinA, UEdGraphPin* PinB, const FVector2D& GraphPosition) const override;
+	virtual void OnPinConnectionDoubleCicked(UEdGraphPin* PinA, UEdGraphPin* PinB, const FVector2f& GraphPosition) const override;
 	//~ End EdGraphSchema Interface
 
 private:
diff --git a/Engine/Plugins/PCG/Source/PCGEditor/Private/PCGEditorGraphSchemaActions.cpp b/Engine/Plugins/PCG/Source/PCGEditor/Private/PCGEditorGraphSchemaActions.cpp
--- a/Engine/Plugins/PCG/Source/PCGEditor/Private/PCGEditorGraphSchemaActions.cpp
+++ b/Engine/Plugins/PCG/Source/PCGEditor/Private/PCGEditorGraphSchemaActions.cpp
@@ -27,7 +27,7 @@
 
 #define LOCTEXT_NAMESPACE "PCGEditorGraphSchemaActions"
 
-UEdGraphNode* FPCGEditorGraphSchemaAction_NewNativeElement::PerformAction(UEdGraph* ParentGraph, UEdGraphPin* FromPin, const FVector2D Location, bool bSelectNewNode)
+UEdGraphNode* FPCGEditorGraphSchemaAction_NewNativeElement::PerformAction(UEdGraph* ParentGraph, UEdGraphPin* FromPin, const FVector2f& Location, bool bSelectNewNode)
 {
 	UPCGEditorGraph* EditorGraph = Cast<UPCGEditorGraph>(ParentGraph);
 	if (!EditorGraph)
@@ -102,7 +102,7 @@ void FPCGEditorGraphSchemaAction_NewGetParameterElement::PostCreation(UPCGNode*
 	NewNode->SetSettingsInterface(Settings);
 }
 
-UEdGraphNode* FPCGEditorGraphSchemaAction_NewSettingsElement::PerformAction(UEdGraph* ParentGraph, UEdGraphPin* FromPin, const FVector2D Location, bool bSelectNewNode)
+UEdGraphNode* FPCGEditorGraphSchemaAction_NewSettingsElement::PerformAction(UEdGraph* ParentGraph, UEdGraphPin* FromPin, const FVector2f& Location, bool bSelectNewNode)
 {
 	UPCGEditorGraph* EditorGraph = Cast<UPCGEditorGraph>(ParentGraph);
 	if (!EditorGraph)
@@ -143,10 +143,10 @@ UEdGraphNode* FPCGEditorGraphSchemaAction_NewSettingsElement::PerformAction(UEdG
 		bCreateInstance = ModifierKeys.IsAltDown();
 	}
 
-	return MakeSettingsNode(EditorGraph, FromPin, Settings, Location, bSelectNewNode, bCreateInstance);
+	return MakeSettingsNode(EditorGraph, FromPin, Settings, FDeprecateSlateVector2D(Location), bSelectNewNode, bCreateInstance);
 }
 
-void FPCGEditorGraphSchemaAction_NewSettingsElement::MakeSettingsNodesOrContextualMenu(const TSharedRef<class SWidget>& InPanel, FVector2D InScreenPosition, UEdGraph* InGraph, const TArray<FSoftObjectPath>& InSettingsPaths, const TArray<FVector2D>& InLocations, bool bInSelectNewNodes)
+void FPCGEditorGraphSchemaAction_NewSettingsElement::MakeSettingsNodesOrContextualMenu(const TSharedRef<class SWidget>& InPanel, const FVector2D& InScreenPosition, UEdGraph* InGraph, const TArray<FSoftObjectPath>& InSettingsPaths, const TArray<FVector2D>& InLocations, bool bInSelectNewNodes)
 {
 	UPCGEditorGraph* EditorGraph = Cast<UPCGEditorGraph>(InGraph);
 
@@ -281,7 +281,7 @@ UEdGraphNode* FPCGEditorGraphSchemaAction_NewSettingsElement::MakeSettingsNode(U
 	return NewNode;
 }
 
-UEdGraphNode* FPCGEditorGraphSchemaAction_NewLoadAssetElement::PerformAction(UEdGraph* ParentGraph, UEdGraphPin* FromPin, const FVector2D Location, bool bSelectNewNode)
+UEdGraphNode* FPCGEditorGraphSchemaAction_NewLoadAssetElement::PerformAction(UEdGraph* ParentGraph, UEdGraphPin* FromPin, const FVector2f& Location, bool bSelectNewNode)
 {
 	UPCGEditorGraph* EditorGraph = Cast<UPCGEditorGraph>(ParentGraph);
 	if (!EditorGraph)
@@ -333,7 +333,7 @@ UEdGraphNode* FPCGEditorGraphSchemaAction_NewLoadAssetElement::PerformAction(UEd
 	return NewNode;
 }
 
-UEdGraphNode* FPCGEditorGraphSchemaAction_NewBlueprintElement::PerformAction(UEdGraph* ParentGraph, UEdGraphPin* FromPin, const FVector2D Location, bool bSelectNewNode)
+UEdGraphNode* FPCGEditorGraphSchemaAction_NewBlueprintElement::PerformAction(UEdGraph* ParentGraph, UEdGraphPin* FromPin, const FVector2f& Location, bool bSelectNewNode)
 {
 	UPCGEditorGraph* EditorGraph = Cast<UPCGEditorGraph>(ParentGraph);
 	if (!EditorGraph)
@@ -387,7 +387,7 @@ UEdGraphNode* FPCGEditorGraphSchemaAction_NewBlueprintElement::PerformAction(UEd
 	return NewNode;
 }
 
-UEdGraphNode* FPCGEditorGraphSchemaAction_NewSubgraphElement::PerformAction(UEdGraph* ParentGraph, UEdGraphPin* FromPin, const FVector2D Location, bool bSelectNewNode)
+UEdGraphNode* FPCGEditorGraphSchemaAction_NewSubgraphElement::PerformAction(UEdGraph* ParentGraph, UEdGraphPin* FromPin, const FVector2f& Location, bool bSelectNewNode)
 {
 	UPCGEditorGraph* EditorGraph = Cast<UPCGEditorGraph>(ParentGraph);
 	if (!EditorGraph)
@@ -427,7 +427,7 @@ UEdGraphNode* FPCGEditorGraphSchemaAction_NewSubgraphElement::PerformAction(UEdG
 		bCreateLoop = ModifierKeys.IsAltDown();
 	}
 
-	return MakeGraphNode(EditorGraph, FromPin, Graph, Location, bSelectNewNode, bCreateLoop);
+	return MakeGraphNode(EditorGraph, FromPin, Graph, FDeprecateSlateVector2D(Location), bSelectNewNode, bCreateLoop);
 }
 
 void FPCGEditorGraphSchemaAction_NewSubgraphElement::MakeGraphNodesOrContextualMenu(const TSharedRef<class SWidget>& InPanel, const FVector2D& InScreenPosition, UEdGraph* InGraph, const TArray<FSoftObjectPath>& InGraphPaths, const TArray<FVector2D>& InLocations, bool bInSelectNewNodes)
@@ -569,7 +569,7 @@ UEdGraphNode* FPCGEditorGraphSchemaAction_NewSubgraphElement::MakeGraphNode(UPCG
 	return NewNode;
 }
 
-UEdGraphNode* FPCGEditorGraphSchemaAction_NewComment::PerformAction(class UEdGraph* ParentGraph, UEdGraphPin* FromPin, const FVector2D Location, bool bSelectNewNode /*= true*/)
+UEdGraphNode* FPCGEditorGraphSchemaAction_NewComment::PerformAction(class UEdGraph* ParentGraph, UEdGraphPin* FromPin, const FVector2f& Location, bool bSelectNewNode /*= true*/)
 {
 	UPCGEditorGraph* EditorGraph = Cast<UPCGEditorGraph>(ParentGraph);
 	if (!EditorGraph)
@@ -582,7 +582,7 @@ UEdGraphNode* FPCGEditorGraphSchemaAction_NewComment::PerformAction(class UEdGra
 
 	TSharedPtr<SGraphEditor> GraphEditorPtr = SGraphEditor::FindGraphEditorForGraph(ParentGraph);
 
-	FVector2D SpawnLocation = Location;
+	FVector2f SpawnLocation = Location;
 	FSlateRect Bounds;
 	if (GraphEditorPtr && GraphEditorPtr->GetBoundsForSelectedNodes(Bounds, 50.0f))
 	{
@@ -599,7 +599,7 @@ UEdGraphNode* FPCGEditorGraphSchemaAction_NewComment::PerformAction(class UEdGra
 	return NewNode;
 }
 
-UEdGraphNode* FPCGEditorGraphSchemaAction_NewReroute::PerformAction(class UEdGraph* ParentGraph, UEdGraphPin* FromPin, const FVector2D Location, bool bSelectNewNode)
+UEdGraphNode* FPCGEditorGraphSchemaAction_NewReroute::PerformAction(class UEdGraph* ParentGraph, UEdGraphPin* FromPin, const FVector2f& Location, bool bSelectNewNode)
 {
 	UPCGEditorGraph* EditorGraph = Cast<UPCGEditorGraph>(ParentGraph);
 	if (!EditorGraph)
@@ -645,7 +645,7 @@ UEdGraphNode* FPCGEditorGraphSchemaAction_NewReroute::PerformAction(class UEdGra
 	return NewNode;
 }
 
-UEdGraphNode* FPCGEditorGraphSchemaAction_NewNamedRerouteUsage::PerformAction(class UEdGraph* ParentGraph, UEdGraphPin* FromPin, const FVector2D Location, bool bSelectNewNode)
+UEdGraphNode* FPCGEditorGraphSchemaAction_NewNamedRerouteUsage::PerformAction(class UEdGraph* ParentGraph, UEdGraphPin* FromPin, const FVector2f& Location, bool bSelectNewNode)
 {
 	check(DeclarationNode);
 
@@ -709,7 +709,7 @@ UEdGraphNode* FPCGEditorGraphSchemaAction_NewNamedRerouteUsage::PerformAction(cl
 	return NewNode;
 }
 
-UEdGraphNode* FPCGEditorGraphSchemaAction_NewNamedRerouteDeclaration::PerformAction(class UEdGraph* ParentGraph, UEdGraphPin* FromPin, const FVector2D Location, bool bSelectNewNode)
+UEdGraphNode* FPCGEditorGraphSchemaAction_NewNamedRerouteDeclaration::PerformAction(class UEdGraph* ParentGraph, UEdGraphPin* FromPin, const FVector2f& Location, bool bSelectNewNode)
 {
 	UPCGEditorGraph* EditorGraph = Cast<UPCGEditorGraph>(ParentGraph);
 	if (!EditorGraph)
diff --git a/Engine/Plugins/PCG/Source/PCGEditor/Private/PCGEditorGraphSchemaActions.h b/Engine/Plugins/PCG/Source/PCGEditor/Private/PCGEditorGraphSchemaActions.h
--- a/Engine/Plugins/PCG/Source/PCGEditor/Private/PCGEditorGraphSchemaActions.h
+++ b/Engine/Plugins/PCG/Source/PCGEditor/Private/PCGEditorGraphSchemaActions.h
@@ -55,7 +55,8 @@ struct FPCGEditorGraphSchemaAction_NewNativeElement : public FEdGraphSchemaActio
 
 	//~ Begin FEdGraphSchemaAction Interface
 	virtual FName GetTypeId() const override { return StaticGetTypeId(); }
-	virtual UEdGraphNode* PerformAction(UEdGraph* ParentGraph, UEdGraphPin* FromPin, const FVector2D Location, bool bSelectNewNode = true) override;
+	using FEdGraphSchemaAction::PerformAction; // Prevent hiding of deprecated base class function with FVector2D
+	virtual UEdGraphNode* PerformAction(UEdGraph* ParentGraph, UEdGraphPin* FromPin, const FVector2f& Location, bool bSelectNewNode = true) override;
 	//~ End FEdGraphSchemaAction Interface
 
 protected:
@@ -103,10 +104,11 @@ struct FPCGEditorGraphSchemaAction_NewSettingsElement : public FEdGraphSchemaAct
 
 	//~ Begin FEdGraphSchemaAction Interface
 	virtual FName GetTypeId() const override { return StaticGetTypeId(); }
-	virtual UEdGraphNode* PerformAction(UEdGraph* ParentGraph, UEdGraphPin* FromPin, const FVector2D Location, bool bSelectNewNode = true) override;
+	using FEdGraphSchemaAction::PerformAction; // Prevent hiding of deprecated base class function with FVector2D
+	virtual UEdGraphNode* PerformAction(UEdGraph* ParentGraph, UEdGraphPin* FromPin, const FVector2f& Location, bool bSelectNewNode = true) override;
 	//~ End FEdGraphSchemaAction Interface
 
-	static void MakeSettingsNodesOrContextualMenu(const TSharedRef<class SWidget>& InPanel, FVector2D InScreenPosition, UEdGraph* InGraph, const TArray<FSoftObjectPath>& InSettingsPaths, const TArray<FVector2D>& InLocations, bool bInSelectNewNodes);
+	static void MakeSettingsNodesOrContextualMenu(const TSharedRef<class SWidget>& InPanel, const FVector2D& InScreenPosition, UEdGraph* InGraph, const TArray<FSoftObjectPath>& InSettingsPaths, const TArray<FVector2D>& InLocations, bool bInSelectNewNodes);
 	static void MakeSettingsNodes(UPCGEditorGraph* InEditorGraph, TArray<UPCGSettings*> InSettings, TArray<FVector2D> InSettingsLocations, bool bInSelectNewNodes, bool bInCreateInstances);
 	static UEdGraphNode* MakeSettingsNode(UPCGEditorGraph* InEditorGraph, UEdGraphPin* InFromPin, UPCGSettings* InSettings, FVector2D InLocation, bool bInSelectNewNode, bool bInCreateInstance);
 };
@@ -134,7 +136,8 @@ struct FPCGEditorGraphSchemaAction_NewBlueprintElement : public FEdGraphSchemaAc
 
 	//~ Begin FEdGraphSchemaAction Interface
 	virtual FName GetTypeId() const override { return StaticGetTypeId(); }
-	virtual UEdGraphNode* PerformAction(UEdGraph* ParentGraph, UEdGraphPin* FromPin, const FVector2D Location, bool bSelectNewNode = true) override;
+	using FEdGraphSchemaAction::PerformAction; // Prevent hiding of deprecated base class function with FVector2D
+	virtual UEdGraphNode* PerformAction(UEdGraph* ParentGraph, UEdGraphPin* FromPin, const FVector2f& Location, bool bSelectNewNode = true) override;
 	//~ End FEdGraphSchemaAction Interface
 };
 
@@ -161,7 +164,8 @@ struct FPCGEditorGraphSchemaAction_NewLoadAssetElement : public FEdGraphSchemaAc
 
 	//~ Begin FEdGraphSchemaAction Interface
 	virtual FName GetTypeId() const override { return StaticGetTypeId(); }
-	virtual UEdGraphNode* PerformAction(UEdGraph* ParentGraph, UEdGraphPin* FromPin, const FVector2D Location, bool bSelectNewNode = true) override;
+	using FEdGraphSchemaAction::PerformAction; // Prevent hiding of deprecated base class function with FVector2D
+	virtual UEdGraphNode* PerformAction(UEdGraph* ParentGraph, UEdGraphPin* FromPin, const FVector2f& Location, bool bSelectNewNode = true) override;
 	//~ End FEdGraphSchemaAction Interface
 };
 
@@ -189,7 +193,8 @@ struct FPCGEditorGraphSchemaAction_NewSubgraphElement : public FEdGraphSchemaAct
 
 	//~ Begin FEdGraphSchemaAction Interface
 	virtual FName GetTypeId() const override { return StaticGetTypeId(); }
-	virtual UEdGraphNode* PerformAction(UEdGraph* ParentGraph, UEdGraphPin* FromPin, const FVector2D Location, bool bSelectNewNode = true) override;
+	using FEdGraphSchemaAction::PerformAction; // Prevent hiding of deprecated base class function with FVector2D
+	virtual UEdGraphNode* PerformAction(UEdGraph* ParentGraph, UEdGraphPin* FromPin, const FVector2f& Location, bool bSelectNewNode = true) override;
 	//~ End FEdGraphSchemaAction Interface
 
 	static void MakeGraphNodesOrContextualMenu(const TSharedRef<class SWidget>& InPanel, const FVector2D& InScreenPosition, UEdGraph* InGraph, const TArray<FSoftObjectPath>& InGraphPaths, const TArray<FVector2D>& InLocations, bool bInSelectNewNodes);
@@ -215,7 +220,8 @@ struct FPCGEditorGraphSchemaAction_NewComment : public FEdGraphSchemaAction
 
 	// FEdGraphSchemaAction interface
 	virtual FName GetTypeId() const override { return StaticGetTypeId(); }
-	virtual UEdGraphNode* PerformAction(UEdGraph* ParentGraph, UEdGraphPin* FromPin, const FVector2D Location, bool bSelectNewNode = true) override;
+	using FEdGraphSchemaAction::PerformAction;
+	virtual UEdGraphNode* PerformAction(UEdGraph* ParentGraph, UEdGraphPin* FromPin, const FVector2f& Location, bool bSelectNewNode = true) override;
 	// End of FEdGraphSchemaAction interface
 };
 
@@ -236,7 +242,8 @@ struct FPCGEditorGraphSchemaAction_NewReroute : public FEdGraphSchemaAction
 
 	// FEdGraphSchemaAction interface
 	virtual FName GetTypeId() const override { return StaticGetTypeId(); }
-	virtual UEdGraphNode* PerformAction(UEdGraph* ParentGraph, UEdGraphPin* FromPin, const FVector2D Location, bool bSelectNewNode = true) override;
+	using FEdGraphSchemaAction::PerformAction; // Prevent hiding of deprecated base class function with FVector2D
+	virtual UEdGraphNode* PerformAction(UEdGraph* ParentGraph, UEdGraphPin* FromPin, const FVector2f& Location, bool bSelectNewNode = true) override;
 	// End of FEdGraphSchemaAction interface
 };
 
@@ -261,7 +268,8 @@ struct FPCGEditorGraphSchemaAction_NewNamedRerouteUsage : public FEdGraphSchemaA
 
 	// FEdGraphSchemaAction interface
 	virtual FName GetTypeId() const override { return StaticGetTypeId(); }
-	virtual UEdGraphNode* PerformAction(UEdGraph* ParentGraph, UEdGraphPin* FromPin, const FVector2D Location, bool bSelectNewNode = true) override;
+	using FEdGraphSchemaAction::PerformAction; // Prevent hiding of deprecated base class function with FVector2D
+	virtual UEdGraphNode* PerformAction(UEdGraph* ParentGraph, UEdGraphPin* FromPin, const FVector2f& Location, bool bSelectNewNode = true) override;
 	// End of FEdGraphSchemaAction interface
 };
 
@@ -282,6 +290,7 @@ struct FPCGEditorGraphSchemaAction_NewNamedRerouteDeclaration : public FEdGraphS
 
 	// FEdGraphSchemaAction interface
 	virtual FName GetTypeId() const override { return StaticGetTypeId(); }
-	virtual UEdGraphNode* PerformAction(UEdGraph* ParentGraph, UEdGraphPin* FromPin, const FVector2D Location, bool bSelectNewNode = true) override;
+	using FEdGraphSchemaAction::PerformAction; // Prevent hiding of deprecated base class function with FVector2D
+	virtual UEdGraphNode* PerformAction(UEdGraph* ParentGraph, UEdGraphPin* FromPin, const FVector2f& Location, bool bSelectNewNode = true) override;
 	// End of FEdGraphSchemaAction interface
 };
diff --git a/Engine/Plugins/Runtime/Metasound/Source/MetasoundEditor/Private/MetasoundEditor.cpp b/Engine/Plugins/Runtime/Metasound/Source/MetasoundEditor/Private/MetasoundEditor.cpp
--- a/Engine/Plugins/Runtime/Metasound/Source/MetasoundEditor/Private/MetasoundEditor.cpp
+++ b/Engine/Plugins/Runtime/Metasound/Source/MetasoundEditor/Private/MetasoundEditor.cpp
@@ -2854,7 +2854,7 @@ namespace Metasound
 				if (Graph && IsGraphEditable())
 				{
 					FMetasoundGraphSchemaAction_NewComment CommentAction;
-					CommentAction.PerformAction(Graph, nullptr, MetasoundGraphEditor->GetPasteLocation());
+					CommentAction.PerformAction(Graph, nullptr, MetasoundGraphEditor->GetPasteLocation2f());
 				}
 			}
 		}
@@ -3408,7 +3408,7 @@ namespace Metasound
 			else
 			{
 				check(MetasoundGraphEditor);
-				Location = MetasoundGraphEditor->GetPasteLocation();
+				Location = MetasoundGraphEditor->GetPasteLocation2f();
 			}
 
 			FDocumentPasteNotifications Notifications;
diff --git a/Engine/Plugins/Runtime/Metasound/Source/MetasoundEditor/Private/MetasoundEditorGraphSchema.cpp b/Engine/Plugins/Runtime/Metasound/Source/MetasoundEditor/Private/MetasoundEditorGraphSchema.cpp
--- a/Engine/Plugins/Runtime/Metasound/Source/MetasoundEditor/Private/MetasoundEditorGraphSchema.cpp
+++ b/Engine/Plugins/Runtime/Metasound/Source/MetasoundEditor/Private/MetasoundEditorGraphSchema.cpp
@@ -274,7 +274,7 @@ PRAGMA_ENABLE_DEPRECATION_WARNINGS
 			}
 		}
 
-		UEdGraphNode* PromoteToVariable(const FName BaseName, UEdGraphPin& FromPin, const FName DataType, const FMetasoundFrontendClass& InVariableClass, const FVector2D& InLocation, bool bSelectNode)
+		UEdGraphNode* PromoteToVariable(const FName BaseName, UEdGraphPin& FromPin, const FName DataType, const FMetasoundFrontendClass& InVariableClass, const FVector2f& InLocation, bool bSelectNode)
 		{
 			using namespace Frontend;
 
@@ -318,7 +318,7 @@ PRAGMA_ENABLE_DEPRECATION_WARNINGS
 						UMetasoundEditorGraphVariableNode* NewGraphNode = FGraphBuilder::AddVariableNode(ParentMetaSound, NewVariableNode->GetID());
 						if (ensure(NewGraphNode))
 						{
-							NewGraphNode->UpdateFrontendNodeLocation(InLocation);
+							NewGraphNode->UpdateFrontendNodeLocation(FDeprecateSlateVector2D(InLocation));
 							NewGraphNode->SyncLocationFromFrontendNode();
 
 							// Set the literal using the new value if connecting to an input
@@ -675,7 +675,7 @@ PRAGMA_ENABLE_DEPRECATION_WARNINGS
 
 	namespace SchemaUtils
 	{
-		UEdGraphNode* PromoteToInput(UEdGraph* ParentGraph, UEdGraphPin* FromPin, const FVector2D InLocation, bool bSelectNewNode)
+		UEdGraphNode* PromoteToInput(UEdGraph* ParentGraph, UEdGraphPin* FromPin, const UE::Slate::FDeprecateVector2DParameter& InLocation, bool bSelectNewNode)
 		{
 			using namespace Metasound::Editor;
 			using namespace Metasound::Frontend;
@@ -763,7 +763,7 @@ PRAGMA_ENABLE_DEPRECATION_WARNINGS
 					{
 						if (UMetasoundEditorGraphNode* NewGraphNode = FGraphBuilder::AddInputNode(ParentMetasound, NewTemplateNode->GetID()))
 						{
-							NewGraphNode->UpdateFrontendNodeLocation(InLocation);
+							NewGraphNode->UpdateFrontendNodeLocation(FDeprecateSlateVector2D(InLocation));
 							NewGraphNode->SyncLocationFromFrontendNode();
 							UEdGraphNode* EdGraphNode = CastChecked<UEdGraphNode>(NewGraphNode);
 
@@ -789,7 +789,7 @@ PRAGMA_ENABLE_DEPRECATION_WARNINGS
 			return nullptr;
 		}
 	
-		UEdGraphNode* PromoteToOutput(UEdGraph* ParentGraph, UEdGraphPin* FromPin, const FVector2D InLocation, bool bSelectNewNode)
+		UEdGraphNode* PromoteToOutput(UEdGraph* ParentGraph, UEdGraphPin* FromPin, const UE::Slate::FDeprecateVector2DParameter& InLocation, bool bSelectNewNode)
 		{
 			using namespace Metasound::Editor;
 			using namespace Metasound::Frontend;
@@ -828,7 +828,7 @@ PRAGMA_ENABLE_DEPRECATION_WARNINGS
 				{
 					if (UMetasoundEditorGraphOutputNode* NewGraphNode = FGraphBuilder::AddOutputNode(ParentMetasound, OutputNodeID))
 					{
-						NewGraphNode->UpdateFrontendNodeLocation(InLocation);
+						NewGraphNode->UpdateFrontendNodeLocation(FDeprecateSlateVector2D(InLocation));
 						NewGraphNode->SyncLocationFromFrontendNode();
 
 						UEdGraphNode* EdGraphNode = CastChecked<UEdGraphNode>(NewGraphNode);
@@ -854,7 +854,7 @@ PRAGMA_ENABLE_DEPRECATION_WARNINGS
 			return nullptr;
 		}
 
-		UEdGraphNode* PromoteToVariable(UEdGraph* ParentGraph, UEdGraphPin* FromPin, const FVector2D InLocation, bool bSelectNewNode)
+		UEdGraphNode* PromoteToVariable(UEdGraph* ParentGraph, UEdGraphPin* FromPin, const UE::Slate::FDeprecateVector2DParameter& InLocation, bool bSelectNewNode)
 		{
 			using namespace Metasound;
 			using namespace Metasound::Editor;
@@ -877,7 +877,7 @@ PRAGMA_ENABLE_DEPRECATION_WARNINGS
 			return nullptr;
 		}
 
-		UEdGraphNode* PromoteToDeferredVariable(UEdGraph* ParentGraph, UEdGraphPin* FromPin, const FVector2D InLocation, bool bSelectNewNode)
+		UEdGraphNode* PromoteToDeferredVariable(UEdGraph* ParentGraph, UEdGraphPin* FromPin, const UE::Slate::FDeprecateVector2DParameter& InLocation, bool bSelectNewNode)
 		{
 			using namespace Metasound;
 			using namespace Metasound::Editor;
@@ -900,7 +900,7 @@ PRAGMA_ENABLE_DEPRECATION_WARNINGS
 			return nullptr;
 		}
 
-		UEdGraphNode* PromoteToMutatorVariable(UEdGraph* ParentGraph, UEdGraphPin* FromPin, const FVector2D InLocation, bool bSelectNewNode)
+		UEdGraphNode* PromoteToMutatorVariable(UEdGraph* ParentGraph, UEdGraphPin* FromPin, const UE::Slate::FDeprecateVector2DParameter& InLocation, bool bSelectNewNode)
 		{
 			using namespace Metasound;
 			using namespace Metasound::Editor;
@@ -925,7 +925,7 @@ PRAGMA_ENABLE_DEPRECATION_WARNINGS
 	} // namespace SchemaUtils
 } // namespace Metasound::Editor
 
-UEdGraphNode* FMetasoundGraphSchemaAction_NodeWithMultipleOutputs::PerformAction(class UEdGraph* ParentGraph, TArray<UEdGraphPin*>& FromPins, const FVector2D Location, bool bSelectNewNode/* = true*/)
+UEdGraphNode* FMetasoundGraphSchemaAction_NodeWithMultipleOutputs::PerformAction(class UEdGraph* ParentGraph, TArray<UEdGraphPin*>& FromPins, const FVector2f& Location, bool bSelectNewNode/* = true*/)
 {
 	using namespace Metasound::Editor;
 	
@@ -989,7 +989,7 @@ const FLinearColor& FMetasoundGraphSchemaAction_NewNode::GetIconColor() const
 	return Super::GetIconColor();
 }
 
-UEdGraphNode* FMetasoundGraphSchemaAction_NewNode::PerformAction(UEdGraph* ParentGraph, UEdGraphPin* FromPin, const FVector2D Location, bool bSelectNewNode /* = true */)
+UEdGraphNode* FMetasoundGraphSchemaAction_NewNode::PerformAction(UEdGraph* ParentGraph, UEdGraphPin* FromPin, const FVector2f& Location, bool bSelectNewNode /* = true */)
 {
 	using namespace Metasound::Editor;
 	using namespace Metasound::Frontend;
@@ -1011,7 +1011,7 @@ UEdGraphNode* FMetasoundGraphSchemaAction_NewNode::PerformAction(UEdGraph* Paren
 	if (UMetasoundEditorGraphExternalNode* NewGraphNode = FGraphBuilder::AddExternalNode(ParentMetasound, Metadata, bSelectNewNode))
 	{
 		NewGraphNode->Modify();
-		NewGraphNode->UpdateFrontendNodeLocation(Location);
+		NewGraphNode->UpdateFrontendNodeLocation(FDeprecateSlateVector2D(Location));
 		NewGraphNode->SyncLocationFromFrontendNode();
 		SchemaPrivate::TryConnectNewNodeToMatchingDataTypePin(*NewGraphNode, FromPin);
 		SchemaPrivate::SelectNodeInEditorForRename(*MetaSoundGraph, *NewGraphNode);
@@ -1042,7 +1042,7 @@ const FLinearColor& FMetasoundGraphSchemaAction_NewInput::GetIconColor() const
 	return Super::GetIconColor();
 }
 
-UEdGraphNode* FMetasoundGraphSchemaAction_NewInput::PerformAction(UEdGraph* ParentGraph, UEdGraphPin* FromPin, const FVector2D InLocation, bool bSelectNewNode /* = true */)
+UEdGraphNode* FMetasoundGraphSchemaAction_NewInput::PerformAction(UEdGraph* ParentGraph, UEdGraphPin* FromPin, const FVector2f& InLocation, bool bSelectNewNode /* = true */)
 {
 	using namespace Metasound::Editor;
 	using namespace Metasound::Engine;
@@ -1076,7 +1076,7 @@ PRAGMA_ENABLE_DEPRECATION_WARNINGS
 	if (UMetasoundEditorGraphNode* NewGraphNode = FGraphBuilder::AddInputNode(ParentMetasound, TemplateNode->GetID()))
 	{
 		NewGraphNode->Modify();
-		NewGraphNode->UpdateFrontendNodeLocation(InLocation);
+		NewGraphNode->UpdateFrontendNodeLocation(FDeprecateSlateVector2D(InLocation));
 		NewGraphNode->SyncLocationFromFrontendNode();
 		SchemaPrivate::TryConnectNewNodeToMatchingDataTypePin(*NewGraphNode, FromPin);
 		FGraphBuilder::RegisterGraphWithFrontend(ParentMetasound);
@@ -1095,7 +1095,7 @@ FMetasoundGraphSchemaAction_PromoteToInput::FMetasoundGraphSchemaAction_PromoteT
 {
 }
 
-UEdGraphNode* FMetasoundGraphSchemaAction_PromoteToInput::PerformAction(UEdGraph* ParentGraph, UEdGraphPin* FromPin, const FVector2D InLocation, bool bSelectNewNode /* = true */)
+UEdGraphNode* FMetasoundGraphSchemaAction_PromoteToInput::PerformAction(UEdGraph* ParentGraph, UEdGraphPin* FromPin, const FVector2f& InLocation, bool bSelectNewNode /* = true */)
 {
 	using namespace Metasound::Editor;
 	return SchemaUtils::PromoteToInput(ParentGraph, FromPin, InLocation, bSelectNewNode);
@@ -1110,7 +1110,7 @@ FMetasoundGraphSchemaAction_PromoteToVariable_AccessorNode::FMetasoundGraphSchem
 {
 }
 
-UEdGraphNode* FMetasoundGraphSchemaAction_PromoteToVariable_AccessorNode::PerformAction(UEdGraph* ParentGraph, UEdGraphPin* FromPin, const FVector2D InLocation, bool bSelectNewNode /* = true */)
+UEdGraphNode* FMetasoundGraphSchemaAction_PromoteToVariable_AccessorNode::PerformAction(UEdGraph* ParentGraph, UEdGraphPin* FromPin, const FVector2f& InLocation, bool bSelectNewNode /* = true */)
 {
 	using namespace Metasound::Editor;
 	return SchemaUtils::PromoteToVariable(ParentGraph, FromPin, InLocation, bSelectNewNode);
@@ -1125,7 +1125,7 @@ FMetasoundGraphSchemaAction_PromoteToVariable_DeferredAccessorNode::FMetasoundGr
 {
 }
 
-UEdGraphNode* FMetasoundGraphSchemaAction_PromoteToVariable_DeferredAccessorNode::PerformAction(UEdGraph* ParentGraph, UEdGraphPin* FromPin, const FVector2D InLocation, bool bSelectNewNode /* = true */)
+UEdGraphNode* FMetasoundGraphSchemaAction_PromoteToVariable_DeferredAccessorNode::PerformAction(UEdGraph* ParentGraph, UEdGraphPin* FromPin, const FVector2f& InLocation, bool bSelectNewNode /* = true */)
 {
 	using namespace Metasound::Editor;
 	return SchemaUtils::PromoteToDeferredVariable(ParentGraph, FromPin, InLocation, bSelectNewNode);
@@ -1140,7 +1140,7 @@ FMetasoundGraphSchemaAction_PromoteToVariable_MutatorNode::FMetasoundGraphSchema
 {
 }
 
-UEdGraphNode* FMetasoundGraphSchemaAction_PromoteToVariable_MutatorNode::PerformAction(UEdGraph* ParentGraph, UEdGraphPin* FromPin, const FVector2D InLocation, bool bSelectNewNode /* = true */)
+UEdGraphNode* FMetasoundGraphSchemaAction_PromoteToVariable_MutatorNode::PerformAction(UEdGraph* ParentGraph, UEdGraphPin* FromPin, const FVector2f& InLocation, bool bSelectNewNode /* = true */)
 {
 	using namespace Metasound::Editor;
 	return SchemaUtils::PromoteToMutatorVariable(ParentGraph, FromPin, InLocation, bSelectNewNode);
@@ -1167,7 +1167,7 @@ const FLinearColor& FMetasoundGraphSchemaAction_NewOutput::GetIconColor() const
 	return Super::GetIconColor();
 }
 
-UEdGraphNode* FMetasoundGraphSchemaAction_NewOutput::PerformAction(UEdGraph* ParentGraph, UEdGraphPin* FromPin, const FVector2D Location, bool bSelectNewNode /* = true */)
+UEdGraphNode* FMetasoundGraphSchemaAction_NewOutput::PerformAction(UEdGraph* ParentGraph, UEdGraphPin* FromPin, const FVector2f& Location, bool bSelectNewNode /* = true */)
 {
 	using namespace Metasound::Editor;
 	using namespace Metasound::Frontend;
@@ -1187,7 +1187,7 @@ UEdGraphNode* FMetasoundGraphSchemaAction_NewOutput::PerformAction(UEdGraph* Par
 
 	if (UMetasoundEditorGraphOutputNode* NewGraphNode = FGraphBuilder::AddOutputNode(ParentMetasound, Output->NodeID, bSelectNewNode))
 	{
-		NewGraphNode->UpdateFrontendNodeLocation(Location);
+		NewGraphNode->UpdateFrontendNodeLocation(FDeprecateSlateVector2D(Location));
 		NewGraphNode->SyncLocationFromFrontendNode();
 		SchemaPrivate::TryConnectNewNodeToMatchingDataTypePin(*NewGraphNode, FromPin);
 		FGraphBuilder::RegisterGraphWithFrontend(ParentMetasound);
@@ -1206,7 +1206,7 @@ FMetasoundGraphSchemaAction_PromoteToOutput::FMetasoundGraphSchemaAction_Promote
 {
 }
 
-UEdGraphNode* FMetasoundGraphSchemaAction_PromoteToOutput::PerformAction(UEdGraph* ParentGraph, UEdGraphPin* FromPin, const FVector2D InLocation, bool bSelectNewNode /* = true */)
+UEdGraphNode* FMetasoundGraphSchemaAction_PromoteToOutput::PerformAction(UEdGraph* ParentGraph, UEdGraphPin* FromPin, const FVector2f& InLocation, bool bSelectNewNode /* = true */)
 {
 	using namespace Metasound::Editor;
 	return SchemaUtils::PromoteToOutput(ParentGraph, FromPin, InLocation, bSelectNewNode);
@@ -1218,7 +1218,7 @@ FMetasoundGraphSchemaAction_NewVariableNode::FMetasoundGraphSchemaAction_NewVari
 {
 }
 
-UEdGraphNode* FMetasoundGraphSchemaAction_NewVariableNode::PerformAction(UEdGraph* ParentGraph, UEdGraphPin* FromPin, const FVector2D Location, bool bSelectNewNode)
+UEdGraphNode* FMetasoundGraphSchemaAction_NewVariableNode::PerformAction(UEdGraph* ParentGraph, UEdGraphPin* FromPin, const FVector2f& Location, bool bSelectNewNode)
 {
 	using namespace Metasound::Editor;
 	using namespace Metasound::Frontend;
@@ -1240,7 +1240,7 @@ UEdGraphNode* FMetasoundGraphSchemaAction_NewVariableNode::PerformAction(UEdGrap
 					if (UMetasoundEditorGraphVariableNode* NewGraphNode = FGraphBuilder::AddVariableNode(*ParentMetasound, FrontendNode->GetID(), bSelectNewNode))
 					{
 						NewGraphNode->Modify();
-						NewGraphNode->UpdateFrontendNodeLocation(Location);
+						NewGraphNode->UpdateFrontendNodeLocation(FDeprecateSlateVector2D(Location));
 						NewGraphNode->SyncLocationFromFrontendNode();
 						SchemaPrivate::TryConnectNewNodeToMatchingDataTypePin(*NewGraphNode, FromPin);
 						return NewGraphNode;
@@ -1313,7 +1313,7 @@ const FMetasoundFrontendNode* FMetasoundGraphSchemaAction_NewVariableMutatorNode
 	return nullptr;
 }
 
-UEdGraphNode* FMetasoundGraphSchemaAction_NewFromSelected::PerformAction(UEdGraph* ParentGraph, UEdGraphPin* FromPin, const FVector2D Location, bool bSelectNewNode /* = true*/)
+UEdGraphNode* FMetasoundGraphSchemaAction_NewFromSelected::PerformAction(UEdGraph* ParentGraph, UEdGraphPin* FromPin, const FVector2f& Location, bool bSelectNewNode /* = true*/)
 {
 	// TODO: Implement
 	return nullptr;
@@ -1329,7 +1329,7 @@ FMetasoundGraphSchemaAction_NewAudioAnalyzer::FMetasoundGraphSchemaAction_NewAud
 	//
 }
 
-UEdGraphNode* FMetasoundGraphSchemaAction_NewAudioAnalyzer::PerformAction(UEdGraph* ParentGraph, UEdGraphPin* FromPin, const FVector2D Location, bool bSelectNewNode/* = true*/)
+UEdGraphNode* FMetasoundGraphSchemaAction_NewAudioAnalyzer::PerformAction(UEdGraph* ParentGraph, UEdGraphPin* FromPin, const FVector2f& Location, bool bSelectNewNode/* = true*/)
 {
 	using namespace Metasound::Editor;
 	using namespace Metasound::Frontend;
@@ -1359,7 +1359,7 @@ UEdGraphNode* FMetasoundGraphSchemaAction_NewAudioAnalyzer::PerformAction(UEdGra
 	Params.InputsToConnect.Add(FromVertexDataType);
 
 	const FMetasoundFrontendNode* TemplateNode = DocBuilder.AddNodeByTemplate(*AudioAnalyzerTemplate, MoveTemp(Params));
-	DocBuilder.SetNodeLocation(TemplateNode->GetID(), Location);
+	DocBuilder.SetNodeLocation(TemplateNode->GetID(), FDeprecateSlateVector2D(Location));
 
 	const FMetasoundFrontendVertexHandle FromVertexHandle = FGraphBuilder::GetPinVertexHandle(DocBuilder, FromPin);
 	const auto VertexIsMatchingDataType = [FromVertexDataType](const FMetasoundFrontendVertex& Vertex) { return Vertex.TypeName == FromVertexDataType; };
@@ -1408,7 +1408,7 @@ FMetasoundGraphSchemaAction_NewReroute::FMetasoundGraphSchemaAction_NewReroute(c
 {
 }
 
-UEdGraphNode* FMetasoundGraphSchemaAction_NewReroute::PerformAction(UEdGraph* ParentGraph, UEdGraphPin* FromPin, const FVector2D Location, bool bSelectNewNode /* = true*/)
+UEdGraphNode* FMetasoundGraphSchemaAction_NewReroute::PerformAction(UEdGraph* ParentGraph, UEdGraphPin* FromPin, const FVector2f& Location, bool bSelectNewNode /* = true*/)
 {
 	using namespace Metasound;
 	using namespace Metasound::Editor;
@@ -1448,7 +1448,7 @@ UEdGraphNode* FMetasoundGraphSchemaAction_NewReroute::PerformAction(UEdGraph* Pa
 	if (UMetasoundEditorGraphExternalNode* NewGraphNode = FGraphBuilder::AddTemplateNode(ParentMetasound, NodeHandle->GetID(), NodeHandle->GetClassMetadata(), bSelectNewNode))
 	{
 		NewGraphNode->Modify();
-		NewGraphNode->UpdateFrontendNodeLocation(Location);
+		NewGraphNode->UpdateFrontendNodeLocation(FDeprecateSlateVector2D(Location));
 		NewGraphNode->SyncLocationFromFrontendNode();
 
 		SchemaPrivate::TryConnectNewNodeToMatchingDataTypePin(*NewGraphNode, FromPin);
@@ -1476,7 +1476,7 @@ const FLinearColor& FMetasoundGraphSchemaAction_NewReroute::GetIconColor() const
 	return IconColor;
 }
 
-UEdGraphNode* FMetasoundGraphSchemaAction_NewComment::PerformAction(UEdGraph* ParentGraph, UEdGraphPin* FromPin, const FVector2D Location, bool bSelectNewNode/* = true*/)
+UEdGraphNode* FMetasoundGraphSchemaAction_NewComment::PerformAction(UEdGraph* ParentGraph, UEdGraphPin* FromPin, const FVector2f& Location, bool bSelectNewNode/* = true*/)
 {
 	using namespace Metasound::Editor;
 
@@ -1527,14 +1527,15 @@ const FLinearColor& FMetasoundGraphSchemaAction_NewComment::GetIconColor() const
 	return Super::GetIconColor();
 }
 
-UEdGraphNode* FMetasoundGraphSchemaAction_Paste::PerformAction(UEdGraph* ParentGraph, UEdGraphPin* FromPin, const FVector2D Location, bool bSelectNewNode/* = true*/)
+UEdGraphNode* FMetasoundGraphSchemaAction_Paste::PerformAction(UEdGraph* ParentGraph, UEdGraphPin* FromPin, const FVector2f& Location, bool bSelectNewNode/* = true*/)
 {
 	using namespace Metasound::Editor;
 
 	TSharedPtr<FEditor> MetasoundEditor = FGraphBuilder::GetEditorForGraph(*ParentGraph);
 	if (MetasoundEditor.IsValid())
 	{
-		MetasoundEditor->PasteNodes(&Location);
+		FDeprecateSlateVector2D LocationToPaste(Location);
+		MetasoundEditor->PasteNodes(&LocationToPaste);
 	}
 
 	return nullptr;
@@ -1918,7 +1919,7 @@ const FPinConnectionResponse UMetasoundEditorGraphSchema::CanCreateConnection(co
 	}
 }
 
-void UMetasoundEditorGraphSchema::OnPinConnectionDoubleCicked(UEdGraphPin* PinA, UEdGraphPin* PinB, const FVector2D& GraphPosition) const
+void UMetasoundEditorGraphSchema::OnPinConnectionDoubleCicked(UEdGraphPin* PinA, UEdGraphPin* PinB, const FVector2f& GraphPosition) const
 {
 	using namespace Metasound;
 	using namespace Metasound::Editor;
@@ -1930,8 +1931,8 @@ void UMetasoundEditorGraphSchema::OnPinConnectionDoubleCicked(UEdGraphPin* PinA,
 	}
 
 	//@TODO: This constant is duplicated from inside of SGraphNodeKnot
-	const FVector2D NodeSpacerSize(42.0f, 24.0f);
-	const FVector2D KnotTopLeft = GraphPosition - (NodeSpacerSize * 0.5f);
+	const FVector2f NodeSpacerSize(42.0f, 24.0f);
+	const FVector2f KnotTopLeft = GraphPosition - (NodeSpacerSize * 0.5f);
 
 	UMetasoundEditorGraph* ParentGraph = Cast<UMetasoundEditorGraph>(PinA->GetOwningNode()->GetGraph());
 	if (ParentGraph->IsEditable())
@@ -2027,13 +2028,13 @@ bool UMetasoundEditorGraphSchema::TryCreateConnection(UEdGraphPin* PinA, UEdGrap
 
 			OutputPin->Modify();
 			
-			FVector2D Location = FVector2D();
-			Location += FVector2D(InputNode->NodePosX, InputNode->NodePosY);
-			Location += FVector2D(OutputNode->NodePosX, OutputNode->NodePosY);
+			FVector2f Location = FVector2f::ZeroVector;
+			Location += FVector2f(InputNode->NodePosX, InputNode->NodePosY);
+			Location += FVector2f(OutputNode->NodePosX, OutputNode->NodePosY);
 			Location *= 0.5f;
 
 			NewGraphNode->Modify();
-			NewGraphNode->UpdateFrontendNodeLocation(Location);
+			NewGraphNode->UpdateFrontendNodeLocation(FDeprecateSlateVector2D(Location));
 			NewGraphNode->SyncLocationFromFrontendNode();
 
 			SchemaPrivate::TryConnectNewNodeToMatchingDataTypePin(*NewGraphNode, InputPin);
@@ -2456,7 +2457,7 @@ void UMetasoundEditorGraphSchema::GetAssetsPinHoverMessage(const TArray<FAssetDa
 	OutOkIcon = false;
 }
 
-void UMetasoundEditorGraphSchema::DroppedAssetsOnGraph(const TArray<FAssetData>& Assets, const FVector2D& GraphPosition, UEdGraph* Graph) const
+void UMetasoundEditorGraphSchema::DroppedAssetsOnGraph(const TArray<FAssetData>& Assets, const FVector2f& GraphPosition, UEdGraph* Graph) const
 {
 	using namespace Metasound;
 	using namespace Metasound::Engine;
@@ -2500,7 +2501,7 @@ void UMetasoundEditorGraphSchema::DroppedAssetsOnGraph(const TArray<FAssetData>&
 			FMetaSoundNodeHandle NodeHandle = Builder.AddNode(DroppedDocInterface, Result);
 			if (ensure(Result == EMetaSoundBuilderResult::Succeeded))
 			{
-				Builder.SetNodeLocation(NodeHandle.NodeID, GraphPosition, Result);
+				Builder.SetNodeLocation(NodeHandle.NodeID, FDeprecateSlateVector2D(GraphPosition), Result);
 				bTransactionSucceeded = ensure(Result == EMetaSoundBuilderResult::Succeeded);
 			}
 		}
@@ -2512,12 +2513,12 @@ void UMetasoundEditorGraphSchema::DroppedAssetsOnGraph(const TArray<FAssetData>&
 	}
 }
 
-void UMetasoundEditorGraphSchema::DroppedAssetsOnNode(const TArray<FAssetData>& Assets, const FVector2D& GraphPosition, UEdGraphNode* Node) const
+void UMetasoundEditorGraphSchema::DroppedAssetsOnNode(const TArray<FAssetData>& Assets, const FVector2f& GraphPosition, UEdGraphNode* Node) const
 {
 	// Still needed?
 }
 
-void UMetasoundEditorGraphSchema::DroppedAssetsOnPin(const TArray<FAssetData>& Assets, const FVector2D& GraphPosition, UEdGraphPin* Pin) const
+void UMetasoundEditorGraphSchema::DroppedAssetsOnPin(const TArray<FAssetData>& Assets, const FVector2f& GraphPosition, UEdGraphPin* Pin) const
 {
 	using namespace Metasound::Editor;
 	using namespace Metasound::Frontend;
@@ -2846,13 +2847,13 @@ TSharedPtr<FEdGraphSchemaAction> UMetasoundEditorGraphSchema::GetCreateCommentAc
 	return StaticCastSharedPtr<FEdGraphSchemaAction, FMetasoundGraphSchemaAction_NewComment>(Comment);
 }
 
-void UMetasoundEditorGraphSchema::SetNodePosition(UEdGraphNode* Node, const FVector2D& Position) const
+void UMetasoundEditorGraphSchema::SetNodePosition(UEdGraphNode* Node, const FVector2f& Position) const
 {
 	if (UMetasoundEditorGraphNode* MetasoundGraphNode = Cast<UMetasoundEditorGraphNode>(Node))
 	{
 		MetasoundGraphNode->GetMetasoundChecked().Modify();
 		UEdGraphSchema::SetNodePosition(Node, Position);
-		MetasoundGraphNode->UpdateFrontendNodeLocation(Position);
+		MetasoundGraphNode->UpdateFrontendNodeLocation(FDeprecateSlateVector2D(Position));
 	}
 	else
 	{
diff --git a/Engine/Plugins/Runtime/Metasound/Source/MetasoundEditor/Private/SMetasoundActionMenu.h b/Engine/Plugins/Runtime/Metasound/Source/MetasoundEditor/Private/SMetasoundActionMenu.h
--- a/Engine/Plugins/Runtime/Metasound/Source/MetasoundEditor/Private/SMetasoundActionMenu.h
+++ b/Engine/Plugins/Runtime/Metasound/Source/MetasoundEditor/Private/SMetasoundActionMenu.h
@@ -82,7 +82,7 @@ namespace Metasound
 			bool bActionExecuted = false;
 
 			TArray<UEdGraphPin*> DraggedFromPins;
-			FVector2D NewNodePosition = FVector2D::ZeroVector;
+			FDeprecateSlateVector2D NewNodePosition = FVector2f::ZeroVector;
 
 			SGraphEditor::FActionMenuClosed OnClosedCallback;
 			FClosedReason OnCloseReasonCallback;
diff --git a/Engine/Plugins/Runtime/Metasound/Source/MetasoundEditor/Public/MetasoundEditorGraphSchema.h b/Engine/Plugins/Runtime/Metasound/Source/MetasoundEditor/Public/MetasoundEditorGraphSchema.h
--- a/Engine/Plugins/Runtime/Metasound/Source/MetasoundEditor/Public/MetasoundEditorGraphSchema.h
+++ b/Engine/Plugins/Runtime/Metasound/Source/MetasoundEditor/Public/MetasoundEditorGraphSchema.h
@@ -80,11 +80,11 @@ namespace Metasound::Editor
 
 	namespace SchemaUtils
 	{
-		UEdGraphNode* PromoteToInput(UEdGraph* ParentGraph, UEdGraphPin* FromPin, const FVector2D InLocation, bool bSelectNewNode);
-		UEdGraphNode* PromoteToOutput(UEdGraph* ParentGraph, UEdGraphPin* FromPin, const FVector2D InLocation, bool bSelectNewNode);
-		UEdGraphNode* PromoteToVariable(UEdGraph* ParentGraph, UEdGraphPin* FromPin, const FVector2D InLocation, bool bSelectNewNode);
-		UEdGraphNode* PromoteToDeferredVariable(UEdGraph* ParentGraph, UEdGraphPin* FromPin, const FVector2D InLocation, bool bSelectNewNode);
-		UEdGraphNode* PromoteToMutatorVariable(UEdGraph* ParentGraph, UEdGraphPin* FromPin, const FVector2D InLocation, bool bSelectNewNode);
+		UEdGraphNode* PromoteToInput(UEdGraph* ParentGraph, UEdGraphPin* FromPin, const UE::Slate::FDeprecateVector2DParameter& InLocation, bool bSelectNewNode);
+		UEdGraphNode* PromoteToOutput(UEdGraph* ParentGraph, UEdGraphPin* FromPin, const UE::Slate::FDeprecateVector2DParameter& InLocation, bool bSelectNewNode);
+		UEdGraphNode* PromoteToVariable(UEdGraph* ParentGraph, UEdGraphPin* FromPin, const UE::Slate::FDeprecateVector2DParameter& InLocation, bool bSelectNewNode);
+		UEdGraphNode* PromoteToDeferredVariable(UEdGraph* ParentGraph, UEdGraphPin* FromPin, const UE::Slate::FDeprecateVector2DParameter& InLocation, bool bSelectNewNode);
+		UEdGraphNode* PromoteToMutatorVariable(UEdGraph* ParentGraph, UEdGraphPin* FromPin, const UE::Slate::FDeprecateVector2DParameter& InLocation, bool bSelectNewNode);
 	} // namespace SchemaUtils
 } // namespace Metasound::Editor
 
@@ -140,8 +140,9 @@ struct METASOUNDEDITOR_API FMetasoundGraphSchemaAction_NodeWithMultipleOutputs :
 	virtual ~FMetasoundGraphSchemaAction_NodeWithMultipleOutputs() = default;
 
 	//~ Begin FEdGraphSchemaAction Interface
-	virtual UEdGraphNode* PerformAction(UEdGraph* ParentGraph, TArray<UEdGraphPin*>& FromPins, const FVector2D Location, bool bSelectNewNode = true) override;
-	virtual UEdGraphNode* PerformAction(UEdGraph* ParentGraph, UEdGraphPin* FromPin, const FVector2D Location, bool bSelectNewNode = true) { return nullptr; }
+	using FEdGraphSchemaAction::PerformAction; // Prevent hiding of deprecated base class function with FVector2D
+	virtual UEdGraphNode* PerformAction(UEdGraph* ParentGraph, TArray<UEdGraphPin*>& FromPins, const FVector2f& Location, bool bSelectNewNode = true) override;
+	virtual UEdGraphNode* PerformAction(UEdGraph* ParentGraph, UEdGraphPin* FromPin, const FVector2f& Location, bool bSelectNewNode = true) { return nullptr; }
 	//~ End FEdGraphSchemaAction Interface
 };
 
@@ -169,7 +170,8 @@ struct METASOUNDEDITOR_API FMetasoundGraphSchemaAction_NewInput : public FMetaso
 	//~ End FMetasoundGraphSchemaAction Interface
 
 	//~ Begin FEdGraphSchemaAction Interface
-	virtual UEdGraphNode* PerformAction(UEdGraph* ParentGraph, UEdGraphPin* FromPin, const FVector2D Location, bool bSelectNewNode = true) override;
+	using FEdGraphSchemaAction::PerformAction; // Prevent hiding of deprecated base class function with FVector2D
+	virtual UEdGraphNode* PerformAction(UEdGraph* ParentGraph, UEdGraphPin* FromPin, const FVector2f& Location, bool bSelectNewNode = true) override;
 	//~ End FEdGraphSchemaAction Interface
 };
 
@@ -184,7 +186,8 @@ struct METASOUNDEDITOR_API FMetasoundGraphSchemaAction_PromoteToInput : public F
 	virtual ~FMetasoundGraphSchemaAction_PromoteToInput() = default;
 
 	//~ Begin FEdGraphSchemaAction Interface
-	virtual UEdGraphNode* PerformAction(UEdGraph* ParentGraph, UEdGraphPin* FromPin, const FVector2D Location, bool bSelectNewNode = true) override;
+	using FEdGraphSchemaAction::PerformAction; // Prevent hiding of deprecated base class function with FVector2D
+	virtual UEdGraphNode* PerformAction(UEdGraph* ParentGraph, UEdGraphPin* FromPin, const FVector2f& Location, bool bSelectNewNode = true) override;
 	//~ End FEdGraphSchemaAction Interface
 };
 
@@ -212,7 +215,8 @@ struct METASOUNDEDITOR_API FMetasoundGraphSchemaAction_NewOutput : public FMetas
 	//~ End FMetasoundGraphSchemaAction Interface
 
 	//~ Begin FEdGraphSchemaAction Interface
-	virtual UEdGraphNode* PerformAction(UEdGraph* ParentGraph, UEdGraphPin* FromPin, const FVector2D Location, bool bSelectNewNode = true) override;
+	using FEdGraphSchemaAction::PerformAction; // Prevent hiding of deprecated base class function with FVector2D
+	virtual UEdGraphNode* PerformAction(UEdGraph* ParentGraph, UEdGraphPin* FromPin, const FVector2f& Location, bool bSelectNewNode = true) override;
 	//~ End FEdGraphSchemaAction Interface
 };
 
@@ -227,7 +231,8 @@ struct METASOUNDEDITOR_API FMetasoundGraphSchemaAction_PromoteToOutput : public
 	virtual ~FMetasoundGraphSchemaAction_PromoteToOutput() = default;
 
 	//~ Begin FEdGraphSchemaAction Interface
-	virtual UEdGraphNode* PerformAction(UEdGraph* ParentGraph, UEdGraphPin* FromPin, const FVector2D Location, bool bSelectNewNode = true) override;
+	using FEdGraphSchemaAction::PerformAction; // Prevent hiding of deprecated base class function with FVector2D
+	virtual UEdGraphNode* PerformAction(UEdGraph* ParentGraph, UEdGraphPin* FromPin, const FVector2f& Location, bool bSelectNewNode = true) override;
 	//~ End FEdGraphSchemaAction Interface
 };
 
@@ -249,7 +254,8 @@ struct METASOUNDEDITOR_API FMetasoundGraphSchemaAction_NewVariableNode : public
 	virtual ~FMetasoundGraphSchemaAction_NewVariableNode() = default;
 
 	//~ Begin FEdGraphSchemaAction Interface
-	virtual UEdGraphNode* PerformAction(UEdGraph* ParentGraph, UEdGraphPin* FromPin, const FVector2D Location, bool bSelectNewNode = true) override;
+	using FEdGraphSchemaAction::PerformAction; // Prevent hiding of deprecated base class function with FVector2D
+	virtual UEdGraphNode* PerformAction(UEdGraph* ParentGraph, UEdGraphPin* FromPin, const FVector2f& Location, bool bSelectNewNode = true) override;
 	//~ End FEdGraphSchemaAction Interface
 
 	//~ Begin FMetasoundGraphSchemaAction Interface
@@ -337,7 +343,8 @@ struct METASOUNDEDITOR_API FMetasoundGraphSchemaAction_PromoteToVariable_Accesso
 	virtual ~FMetasoundGraphSchemaAction_PromoteToVariable_AccessorNode() = default;
 
 	//~ Begin FEdGraphSchemaAction Interface
-	virtual UEdGraphNode* PerformAction(UEdGraph* ParentGraph, UEdGraphPin* FromPin, const FVector2D Location, bool bSelectNewNode = true) override;
+	using FEdGraphSchemaAction::PerformAction; // Prevent hiding of deprecated base class function with FVector2D
+	virtual UEdGraphNode* PerformAction(UEdGraph* ParentGraph, UEdGraphPin* FromPin, const FVector2f& Location, bool bSelectNewNode = true) override;
 	//~ End FEdGraphSchemaAction Interface
 };
 
@@ -352,7 +359,8 @@ struct METASOUNDEDITOR_API FMetasoundGraphSchemaAction_PromoteToVariable_Mutator
 	virtual ~FMetasoundGraphSchemaAction_PromoteToVariable_MutatorNode() = default;
 
 	//~ Begin FEdGraphSchemaAction Interface
-	virtual UEdGraphNode* PerformAction(UEdGraph* ParentGraph, UEdGraphPin* FromPin, const FVector2D Location, bool bSelectNewNode = true) override;
+	using FEdGraphSchemaAction::PerformAction; // Prevent hiding of deprecated base class function with FVector2D
+	virtual UEdGraphNode* PerformAction(UEdGraph* ParentGraph, UEdGraphPin* FromPin, const FVector2f& Location, bool bSelectNewNode = true) override;
 	//~ End FEdGraphSchemaAction Interface
 };
 
@@ -367,7 +375,8 @@ struct METASOUNDEDITOR_API FMetasoundGraphSchemaAction_PromoteToVariable_Deferre
 	virtual ~FMetasoundGraphSchemaAction_PromoteToVariable_DeferredAccessorNode() = default;
 
 	//~ Begin FEdGraphSchemaAction Interface
-	virtual UEdGraphNode* PerformAction(UEdGraph* ParentGraph, UEdGraphPin* FromPin, const FVector2D Location, bool bSelectNewNode = true) override;
+	using FEdGraphSchemaAction::PerformAction; // Prevent hiding of deprecated base class function with FVector2D
+	virtual UEdGraphNode* PerformAction(UEdGraph* ParentGraph, UEdGraphPin* FromPin, const FVector2f& Location, bool bSelectNewNode = true) override;
 	//~ End FEdGraphSchemaAction Interface
 };
 
@@ -397,7 +406,8 @@ struct METASOUNDEDITOR_API FMetasoundGraphSchemaAction_NewNode : public FMetasou
 	//~ End FMetasoundGraphSchemaAction Interface
 
 	//~ Begin FEdGraphSchemaAction Interface
-	virtual UEdGraphNode* PerformAction(UEdGraph* ParentGraph, UEdGraphPin* FromPin, const FVector2D Location, bool bSelectNewNode = true) override;
+	using FEdGraphSchemaAction::PerformAction; // Prevent hiding of deprecated base class function with FVector2D
+	virtual UEdGraphNode* PerformAction(UEdGraph* ParentGraph, UEdGraphPin* FromPin, const FVector2f& Location, bool bSelectNewNode = true) override;
 	//~ End FEdGraphSchemaAction Interface
 };
 
@@ -418,7 +428,8 @@ struct METASOUNDEDITOR_API FMetasoundGraphSchemaAction_NewFromSelected : public
 	virtual ~FMetasoundGraphSchemaAction_NewFromSelected() = default;
 
 	//~ Begin FEdGraphSchemaAction Interface
-	virtual UEdGraphNode* PerformAction(UEdGraph* ParentGraph, UEdGraphPin* FromPin, const FVector2D Location, bool bSelectNewNode = true) override;
+	using FEdGraphSchemaAction::PerformAction; // Prevent hiding of deprecated base class function with FVector2D
+	virtual UEdGraphNode* PerformAction(UEdGraph* ParentGraph, UEdGraphPin* FromPin, const FVector2f& Location, bool bSelectNewNode = true) override;
 	//~ End FEdGraphSchemaAction Interface
 };
 
@@ -435,7 +446,8 @@ struct METASOUNDEDITOR_API FMetasoundGraphSchemaAction_NewAudioAnalyzer : public
 	//~ End FMetasoundGraphSchemaAction Interface
 
 	//~ Begin FEdGraphSchemaAction Interface
-	virtual UEdGraphNode* PerformAction(UEdGraph* ParentGraph, UEdGraphPin* FromPin, const FVector2D Location, bool bSelectNewNode = true) override;
+	using FEdGraphSchemaAction::PerformAction; // Prevent hiding of deprecated base class function with FVector2D
+	virtual UEdGraphNode* PerformAction(UEdGraph* ParentGraph, UEdGraphPin* FromPin, const FVector2f& Location, bool bSelectNewNode = true) override;
 	//~ End FEdGraphSchemaAction Interface
 };
 
@@ -457,7 +469,8 @@ struct METASOUNDEDITOR_API FMetasoundGraphSchemaAction_NewReroute : public FMeta
 	//~ End FMetasoundGraphSchemaAction Interface
 
 	//~ Begin FEdGraphSchemaAction Interface
-	virtual UEdGraphNode* PerformAction(UEdGraph* ParentGraph, UEdGraphPin* FromPin, const FVector2D Location, bool bSelectNewNode = true) override;
+	using FEdGraphSchemaAction::PerformAction; // Prevent hiding of deprecated base class function with FVector2D
+	virtual UEdGraphNode* PerformAction(UEdGraph* ParentGraph, UEdGraphPin* FromPin, const FVector2f& Location, bool bSelectNewNode = true) override;
 	//~ End FEdGraphSchemaAction Interface
 
 private:
@@ -486,7 +499,8 @@ struct METASOUNDEDITOR_API FMetasoundGraphSchemaAction_NewComment : public FMeta
 	//~ End FMetasoundGraphSchemaAction Interface
 
 	//~ Begin FEdGraphSchemaAction Interface
-	virtual UEdGraphNode* PerformAction(UEdGraph* ParentGraph, UEdGraphPin* FromPin, const FVector2D Location, bool bSelectNewNode = true) override;
+	using FEdGraphSchemaAction::PerformAction; // Prevent hiding of deprecated base class function with FVector2D
+	virtual UEdGraphNode* PerformAction(UEdGraph* ParentGraph, UEdGraphPin* FromPin, const FVector2f& Location, bool bSelectNewNode = true) override;
 	//~ End FEdGraphSchemaAction Interface
 };
 
@@ -505,7 +519,7 @@ struct METASOUNDEDITOR_API FMetasoundGraphSchemaAction_Paste : public FMetasound
 	{}
 
 	//~ Begin FEdGraphSchemaAction Interface
-	virtual UEdGraphNode* PerformAction(UEdGraph* ParentGraph, UEdGraphPin* FromPin, const FVector2D Location, bool bSelectNewNode = true) override;
+	virtual UEdGraphNode* PerformAction(UEdGraph* ParentGraph, UEdGraphPin* FromPin, const FVector2f& Location, bool bSelectNewNode = true) override;
 	//~ End FEdGraphSchemaAction Interface
 };
 
@@ -538,13 +552,13 @@ public:
 	virtual void BreakSinglePinLink(UEdGraphPin* SourcePin, UEdGraphPin* TargetPin) const override;
 	virtual void GetAssetsGraphHoverMessage(const TArray<FAssetData>& Assets, const UEdGraph* HoverGraph, FString& OutTooltipText, bool& OutOkIcon) const override;
 	virtual void GetAssetsPinHoverMessage(const TArray<FAssetData>& Assets, const UEdGraphPin* HoverPin, FString& OutTooltipText, bool& OutOkIcon) const override;
-	virtual void DroppedAssetsOnGraph(const TArray<FAssetData>& Assets, const FVector2D& GraphPosition, UEdGraph* Graph) const override;
-	virtual void DroppedAssetsOnNode(const TArray<FAssetData>& Assets, const FVector2D& GraphPosition, UEdGraphNode* Node) const override;
-	virtual void DroppedAssetsOnPin(const TArray<FAssetData>& Assets, const FVector2D& GraphPosition, UEdGraphPin* Pin) const override;
+	virtual void DroppedAssetsOnGraph(const TArray<FAssetData>& Assets, const FVector2f& GraphPosition, UEdGraph* Graph) const override;
+	virtual void DroppedAssetsOnNode(const TArray<FAssetData>& Assets, const FVector2f& GraphPosition, UEdGraphNode* Node) const override;
+	virtual void DroppedAssetsOnPin(const TArray<FAssetData>& Assets, const FVector2f& GraphPosition, UEdGraphPin* Pin) const override;
 	virtual int32 GetNodeSelectionCount(const UEdGraph* Graph) const override;
 	virtual TSharedPtr<FEdGraphSchemaAction> GetCreateCommentAction() const override;
-	virtual void OnPinConnectionDoubleCicked(UEdGraphPin* PinA, UEdGraphPin* PinB, const FVector2D& GraphPosition) const override;
-	virtual void SetNodePosition(UEdGraphNode* Node, const FVector2D& Position) const;
+	virtual void OnPinConnectionDoubleCicked(UEdGraphPin* PinA, UEdGraphPin* PinB, const FVector2f& GraphPosition) const override;
+	virtual void SetNodePosition(UEdGraphNode* Node, const FVector2f& Position) const;
 	//~ End EdGraphSchema Interface
 
 	void BreakNodeLinks(UEdGraphNode& TargetNode, bool bShouldActuallyTransact) const;
diff --git a/Engine/Plugins/Runtime/RigVM/Source/RigVMDeveloper/Private/EdGraph/RigVMEdGraphSchema.cpp b/Engine/Plugins/Runtime/RigVM/Source/RigVMDeveloper/Private/EdGraph/RigVMEdGraphSchema.cpp
--- a/Engine/Plugins/Runtime/RigVM/Source/RigVMDeveloper/Private/EdGraph/RigVMEdGraphSchema.cpp
+++ b/Engine/Plugins/Runtime/RigVM/Source/RigVMDeveloper/Private/EdGraph/RigVMEdGraphSchema.cpp
@@ -269,7 +269,7 @@ FRigVMEdGraphSchemaAction_PromoteToVariable::FRigVMEdGraphSchemaAction_PromoteTo
 }
 
 UEdGraphNode* FRigVMEdGraphSchemaAction_PromoteToVariable::PerformAction(UEdGraph* ParentGraph, UEdGraphPin* FromPin,
-	const FVector2D Location, bool bSelectNewNode)
+	const FVector2f& Location, bool bSelectNewNode)
 {
 	URigVMEdGraph* RigGraph = Cast<URigVMEdGraph>(ParentGraph);
 	if(RigGraph == nullptr)
@@ -379,7 +379,7 @@ UEdGraphNode* FRigVMEdGraphSchemaAction_PromoteToVariable::PerformAction(UEdGrap
 			ModelPin->GetCPPTypeObject(),
 			FromPin->Direction == EGPD_Input,
 			ModelPin->GetDefaultValue(),
-			Location,
+			FDeprecateSlateVector2D(Location),
 			FString(),
 			true,
 			true
@@ -411,7 +411,7 @@ FRigVMEdGraphSchemaAction_PromoteToExposedPin::FRigVMEdGraphSchemaAction_Promote
 {
 }
 
-UEdGraphNode* FRigVMEdGraphSchemaAction_PromoteToExposedPin::PerformAction(UEdGraph* ParentGraph, UEdGraphPin* FromPin, const FVector2D Location, bool bSelectNewNode)
+UEdGraphNode* FRigVMEdGraphSchemaAction_PromoteToExposedPin::PerformAction(UEdGraph* ParentGraph, UEdGraphPin* FromPin, const FVector2f& Location, bool bSelectNewNode)
 {
 	URigVMEdGraph* RigGraph = Cast<URigVMEdGraph>(ParentGraph);
 	if(RigGraph == nullptr)
@@ -1324,13 +1324,13 @@ void URigVMEdGraphSchema::ClearPinWatch(UEdGraphPin const* Pin) const
 	}
 }
 
-void URigVMEdGraphSchema::OnPinConnectionDoubleCicked(UEdGraphPin* PinA, UEdGraphPin* PinB, const FVector2D& GraphPosition) const
+void URigVMEdGraphSchema::OnPinConnectionDoubleCicked(UEdGraphPin* PinA, UEdGraphPin* PinB, const FVector2f& GraphPosition) const
 {
 	if (const URigVMEdGraphNode* Node = Cast< URigVMEdGraphNode>(PinA->GetOwningNode()))
 	{
 		if (URigVMLink* Link = Node->GetModel()->FindLink(FString::Printf(TEXT("%s -> %s"), *PinA->GetName(), *PinB->GetName())))
 		{
-			Node->GetController()->AddRerouteNodeOnLink(Link, GraphPosition, FString(), true, true);
+			Node->GetController()->AddRerouteNodeOnLink(Link, FDeprecateSlateVector2D(GraphPosition), FString(), true, true);
 		}
 	}
 }
@@ -1357,9 +1357,9 @@ bool URigVMEdGraphSchema::IsStructEditable(UStruct* InStruct) const
 	return false;
 }
 
-void URigVMEdGraphSchema::SetNodePosition(UEdGraphNode* Node, const FVector2D& Position) const
+void URigVMEdGraphSchema::SetNodePosition(UEdGraphNode* Node, const FVector2f& Position) const
 {
-	return SetNodePosition(Node, Position, true);
+	return SetNodePosition(Node, FDeprecateSlateVector2D(Position), true);
 }
 
 void URigVMEdGraphSchema::SetNodePosition(UEdGraphNode* Node, const FVector2D& Position, bool bSetupUndo) const
@@ -2127,7 +2127,7 @@ bool URigVMEdGraphSchema::RequestFunctionDropOnPanel(UEdGraph* InGraph, const FR
 	return false;
 }
 
-bool URigVMEdGraphSchema::RequestVariableDropOnPanel(UEdGraph* InGraph, FProperty* InVariableToDrop, const FVector2D& InDropPosition, const FVector2D& InScreenPosition) const
+bool URigVMEdGraphSchema::RequestVariableDropOnPanel(UEdGraph* InGraph, FProperty* InVariableToDrop, const FVector2f& InDropPosition, const FVector2f& InScreenPosition) const
 {
 #if WITH_EDITOR
 	if (CanVariableBeDropped(InGraph, InVariableToDrop))
@@ -2136,7 +2136,7 @@ bool URigVMEdGraphSchema::RequestVariableDropOnPanel(UEdGraph* InGraph, FPropert
 		URigVMBlueprint* RigBlueprint = Cast<URigVMBlueprint>(Blueprint);
 		if (RigBlueprint != nullptr)
 		{
-			RigBlueprint->OnVariableDropped().Broadcast(InGraph, InVariableToDrop, InDropPosition, InScreenPosition);
+			RigBlueprint->OnVariableDropped().Broadcast(InGraph, InVariableToDrop, FDeprecateSlateVector2D(InDropPosition), FDeprecateSlateVector2D(InScreenPosition));
 			return true;
 		}
 	}
@@ -2145,7 +2145,7 @@ bool URigVMEdGraphSchema::RequestVariableDropOnPanel(UEdGraph* InGraph, FPropert
 	return false;
 }
 
-bool URigVMEdGraphSchema::RequestVariableDropOnPin(UEdGraph* InGraph, FProperty* InVariableToDrop, UEdGraphPin* InPin, const FVector2D& InDropPosition, const FVector2D& InScreenPosition) const
+bool URigVMEdGraphSchema::RequestVariableDropOnPin(UEdGraph* InGraph, FProperty* InVariableToDrop, UEdGraphPin* InPin, const FVector2f& InDropPosition, const FVector2f& InScreenPosition) const
 {
 #if WITH_EDITOR
 	if (CanVariableBeDropped(InGraph, InVariableToDrop))
@@ -2155,7 +2155,7 @@ bool URigVMEdGraphSchema::RequestVariableDropOnPin(UEdGraph* InGraph, FProperty*
 			if (const URigVMPin* ModelPin = Graph->GetModel()->FindPin(InPin->GetName()))
 			{
 				const FRigVMExternalVariable ExternalVariable = FRigVMExternalVariable::Make(InVariableToDrop, nullptr);
-				return RequestVariableDropOnPin(InGraph, ExternalVariable, InPin, InDropPosition, InScreenPosition);
+				return RequestVariableDropOnPin(InGraph, ExternalVariable, InPin, FDeprecateSlateVector2D(InDropPosition), FDeprecateSlateVector2D(InScreenPosition));
 			}
 		}
 	}
diff --git a/Engine/Plugins/Runtime/RigVM/Source/RigVMDeveloper/Public/EdGraph/RigVMEdGraphSchema.h b/Engine/Plugins/Runtime/RigVM/Source/RigVMDeveloper/Public/EdGraph/RigVMEdGraphSchema.h
--- a/Engine/Plugins/Runtime/RigVM/Source/RigVMDeveloper/Public/EdGraph/RigVMEdGraphSchema.h
+++ b/Engine/Plugins/Runtime/RigVM/Source/RigVMDeveloper/Public/EdGraph/RigVMEdGraphSchema.h
@@ -177,7 +177,7 @@ public:
 		return InType == GetTypeId();
 	}
 
-	virtual UEdGraphNode* PerformAction(UEdGraph* ParentGraph, UEdGraphPin* FromPin, const FVector2D Location, bool bSelectNewNode) override;
+	virtual UEdGraphNode* PerformAction(UEdGraph* ParentGraph, UEdGraphPin* FromPin, const FVector2f& Location, bool bSelectNewNode) override;
 
 private:
 
@@ -212,7 +212,7 @@ public:
 		return InType == GetTypeId();
 	}
 
-	virtual UEdGraphNode* PerformAction(UEdGraph* ParentGraph, UEdGraphPin* FromPin, const FVector2D Location, bool bSelectNewNode) override;
+	virtual UEdGraphNode* PerformAction(UEdGraph* ParentGraph, UEdGraphPin* FromPin, const FVector2f& Location, bool bSelectNewNode) override;
 
 private:
 
@@ -329,16 +329,18 @@ public:
 	virtual bool DoesSupportPinWatching() const	override { return true; }
 	virtual bool IsPinBeingWatched(UEdGraphPin const* Pin) const override;
 	virtual void ClearPinWatch(UEdGraphPin const* Pin) const override;
-	virtual void OnPinConnectionDoubleCicked(UEdGraphPin* PinA, UEdGraphPin* PinB, const FVector2D& GraphPosition) const override;
+	virtual void OnPinConnectionDoubleCicked(UEdGraphPin* PinA, UEdGraphPin* PinB, const FVector2f& GraphPosition) const override;
 	virtual bool MarkBlueprintDirtyFromNewNode(UBlueprint* InBlueprint, UEdGraphNode* InEdGraphNode) const override;
 	virtual bool SafeDeleteNodeFromGraph(UEdGraph* Graph, UEdGraphNode* Node) const override;
 	virtual bool CanVariableBeDropped(UEdGraph* InGraph, FProperty* InVariableToDrop) const override;
 	virtual bool RequestFunctionDropOnPanel(UEdGraph* InGraph, const FRigVMGraphFunctionIdentifier& InFunction, const FVector2D& InDropPosition, const FVector2D& InScreenPosition) const;
-	virtual bool RequestVariableDropOnPanel(UEdGraph* InGraph, FProperty* InVariableToDrop, const FVector2D& InDropPosition, const FVector2D& InScreenPosition) const override;
-	virtual bool RequestVariableDropOnPin(UEdGraph* InGraph, FProperty* InVariableToDrop, UEdGraphPin* InPin, const FVector2D& InDropPosition, const FVector2D& InScreenPosition) const override;
+	virtual bool RequestVariableDropOnPanel(UEdGraph* InGraph, FProperty* InVariableToDrop, const FVector2f& InDropPosition, const FVector2f& InScreenPosition) const override;
+	virtual bool RequestVariableDropOnPin(UEdGraph* InGraph, FProperty* InVariableToDrop, UEdGraphPin* InPin, const FVector2f& InDropPosition, const FVector2f& InScreenPosition) const override;
+	using UEdGraphSchema::RequestVariableDropOnPin;
 	virtual bool RequestVariableDropOnPin(UEdGraph* InGraph, const FRigVMExternalVariable& InVariableToDrop, UEdGraphPin* InPin, const FVector2D& InDropPosition, const FVector2D& InScreenPosition) const;
 	virtual bool IsStructEditable(UStruct* InStruct) const;
-	virtual void SetNodePosition(UEdGraphNode* Node, const FVector2D& Position) const override;
+	virtual void SetNodePosition(UEdGraphNode* Node, const FVector2f& Position) const override;
+	using UEdGraphSchema::SetNodePosition;
 	void SetNodePosition(UEdGraphNode* Node, const FVector2D& Position, bool bSetupUndo) const;
 	virtual void GetGraphDisplayInformation(const UEdGraph& Graph, /*out*/ FGraphDisplayInfo& DisplayInfo) const override;
 	virtual bool GetLocalVariables(const UEdGraph* InGraph, TArray<FBPVariableDescription>& OutLocalVariables) const override;
diff --git a/Engine/Plugins/Runtime/RigVM/Source/RigVMEditor/Private/Editor/RigVMActionMenuItem.cpp b/Engine/Plugins/Runtime/RigVM/Source/RigVMEditor/Private/Editor/RigVMActionMenuItem.cpp
--- a/Engine/Plugins/Runtime/RigVM/Source/RigVMEditor/Private/Editor/RigVMActionMenuItem.cpp
+++ b/Engine/Plugins/Runtime/RigVM/Source/RigVMEditor/Private/Editor/RigVMActionMenuItem.cpp
@@ -225,12 +225,12 @@ FRigVMActionMenuItem::FRigVMActionMenuItem(UBlueprintNodeSpawner const* NodeSpaw
 }
 
 //------------------------------------------------------------------------------
-UEdGraphNode* FRigVMActionMenuItem::PerformAction(UEdGraph* ParentGraph, UEdGraphPin* FromPin, FVector2D const Location, bool bSelectNewNode/* = true*/)
+UEdGraphNode* FRigVMActionMenuItem::PerformAction(UEdGraph* ParentGraph, UEdGraphPin* FromPin, const FVector2f& Location, bool bSelectNewNode/* = true*/)
 {
 	using namespace FBlueprintMenuActionItemImpl;
 	FScopedTransaction Transaction(LOCTEXT("AddNodeTransaction", "Add Node"));
 	
-	FVector2D ModifiedLocation = Location;
+	FVector2f ModifiedLocation = Location;
 	if (FromPin != nullptr)
 	{
 		// for input pins, a new node will generally overlap the node being
@@ -271,7 +271,7 @@ UEdGraphNode* FRigVMActionMenuItem::PerformAction(UEdGraph* ParentGraph, UEdGrap
 		const TSet<UEdGraphNode*> OldNodes(ParentGraph->Nodes);
 
 		bool bNewNode = false;
-		LastSpawnedNode = InvokeAction(Action, ParentGraph, ModifiedLocation, BindingsSubset, /*out*/ bNewNode);
+		LastSpawnedNode = InvokeAction(Action, ParentGraph, FDeprecateSlateVector2D(ModifiedLocation), BindingsSubset, /*out*/ bNewNode);
 
 		TArray<UEdGraphNode*> NewNodes = TSet<UEdGraphNode*>(ParentGraph->Nodes).Difference(OldNodes).Array();
 		// could already be an existent node, so we have to add here (can't 
@@ -312,7 +312,7 @@ UEdGraphNode* FRigVMActionMenuItem::PerformAction(UEdGraph* ParentGraph, UEdGrap
 }
 
 //------------------------------------------------------------------------------
-UEdGraphNode* FRigVMActionMenuItem::PerformAction(UEdGraph* ParentGraph, TArray<UEdGraphPin*>& FromPins, FVector2D const Location, bool bSelectNewNode/* = true*/)
+UEdGraphNode* FRigVMActionMenuItem::PerformAction(UEdGraph* ParentGraph, TArray<UEdGraphPin*>& FromPins, const FVector2f& Location, bool bSelectNewNode/* = true*/)
 {
 	UEdGraphPin* FromPin = nullptr;
 	if (FromPins.Num() > 0)
