Fix EngineTest failures to metasound builder API (on behalf of phil.popp)

diff --git a/Engine/Plugins/Runtime/Metasound/Source/MetasoundFrontend/Private/MetasoundFrontendDocumentBuilder.cpp b/Engine/Plugins/Runtime/Metasound/Source/MetasoundFrontend/Private/MetasoundFrontendDocumentBuilder.cpp
--- a/Engine/Plugins/Runtime/Metasound/Source/MetasoundFrontend/Private/MetasoundFrontendDocumentBuilder.cpp
+++ b/Engine/Plugins/Runtime/Metasound/Source/MetasoundFrontend/Private/MetasoundFrontendDocumentBuilder.cpp
@@ -1117,29 +1117,33 @@ const FMetasoundFrontendNode* FMetaSoundFrontendDocumentBuilder::AddGraphInput(F
 	FMetasoundFrontendDocument& Document = GetDocumentChecked();
 	FMetasoundFrontendGraphClass& RootGraph = Document.RootGraph;
 
-	auto FinalizeNode = [this, &ClassInput](FMetasoundFrontendNode& InOutNode, const Metasound::Frontend::FNodeRegistryKey&)
+	const int32 NewIndex = RootGraph.GetDefaultInterface().Inputs.Num();
+	FMetasoundFrontendClassInput& NewInput = RootGraph.GetDefaultInterface().Inputs.Add_GetRef(MoveTemp(ClassInput));
+
+	auto FinalizeNode = [this, &NewInput](FMetasoundFrontendNode& InOutNode, const Metasound::Frontend::FNodeRegistryKey&)
 	{
 		// Sets the name of the node an vertices on the node to match the class vertex name
-		DocumentBuilderPrivate::SetNodeAndVertexNames(InOutNode, ClassInput);
+		DocumentBuilderPrivate::SetNodeAndVertexNames(InOutNode, NewInput);
 
 		// Set the default literal on the nodes inputs so that it gets passed to the instantiated TInputNode on a live
 		// auditioned MetaSound.
-		DocumentBuilderPrivate::SetDefaultLiteralOnInputNode(InOutNode, ClassInput);
+		DocumentBuilderPrivate::SetDefaultLiteralOnInputNode(InOutNode, NewInput);
 	};
 
 #if WITH_EDITORONLY_DATA
 	bool bIsRequired = false;
 	FMetasoundFrontendInterface Interface;
-	if (DocumentBuilderPrivate::IsInterfaceInput(ClassInput.Name, ClassInput.TypeName, &Interface))
+	if (DocumentBuilderPrivate::IsInterfaceInput(NewInput.Name, NewInput.TypeName, &Interface))
 	{
 		if (Document.Interfaces.Contains(Interface.Metadata.Version))
 		{
 			FText RequiredText;
-			bIsRequired = Interface.IsMemberInputRequired(ClassInput.Name, RequiredText);
+			bIsRequired = Interface.IsMemberInputRequired(NewInput.Name, RequiredText);
 		}
 	}
 #endif // WITH_EDITORONLY_DATA
 
+
 	// Must add input node to all paged graphs to maintain API parity for all page implementations
 	FMetasoundFrontendNode* NewNode = nullptr;
 	RootGraph.IterateGraphPages([&](const FMetasoundFrontendGraph& Graph)
@@ -1158,7 +1162,7 @@ const FMetasoundFrontendNode* FMetaSoundFrontendDocumentBuilder::AddGraphInput(F
 			// This is added if this is an interface member that is required to force page to create visual
 			// representation that can inform the user of its required state.
 			FGuid LocationGuid = FDocumentIDGenerator::Get().CreateVertexID(Document);
-			SetNodeLocation(ClassInput.NodeID, FVector2D::ZeroVector, &LocationGuid, &Graph.PageID);
+			SetNodeLocation(NewInput.NodeID, FVector2D::ZeroVector, &LocationGuid, &Graph.PageID);
 		}
 #endif // WITH_EDITORONLY_DATA
 
@@ -1169,8 +1173,6 @@ const FMetasoundFrontendNode* FMetaSoundFrontendDocumentBuilder::AddGraphInput(F
 
 	if (NewNode)
 	{
-		const int32 NewIndex = RootGraph.GetDefaultInterface().Inputs.Num();
-		FMetasoundFrontendClassInput& NewInput = RootGraph.GetDefaultInterface().Inputs.Add_GetRef(MoveTemp(ClassInput));
 		if (!NewInput.VertexID.IsValid())
 		{
 			NewInput.VertexID = FDocumentIDGenerator::Get().CreateVertexID(Document);
@@ -1183,6 +1185,11 @@ const FMetasoundFrontendNode* FMetaSoundFrontendDocumentBuilder::AddGraphInput(F
 
 		return NewNode;
 	}
+	else
+	{
+		// Undo addition of graph input on failure.
+		RootGraph.GetDefaultInterface().Inputs.RemoveAt(NewIndex);
+	}
 
 	return nullptr;
 }
@@ -1226,28 +1233,35 @@ const FMetasoundFrontendNode* FMetaSoundFrontendDocumentBuilder::AddGraphOutput(
 		}
 	}
 
-	auto FinalizeNode = [&ClassOutput](FMetasoundFrontendNode& InOutNode, const Metasound::Frontend::FNodeRegistryKey&)
+	// Add graph output
+	FMetasoundFrontendDocument& Document = GetDocumentChecked();
+	FMetasoundFrontendGraphClass& RootGraph = Document.RootGraph;
+	const int32 NewIndex = RootGraph.GetDefaultInterface().Outputs.Num();
+	FMetasoundFrontendClassOutput& NewOutput = RootGraph.GetDefaultInterface().Outputs.Add_GetRef(MoveTemp(ClassOutput));
+
+	auto FinalizeNode = [&NewOutput](FMetasoundFrontendNode& InOutNode, const Metasound::Frontend::FNodeRegistryKey&)
 	{
-		DocumentBuilderPrivate::SetNodeAndVertexNames(InOutNode, ClassOutput);
+		DocumentBuilderPrivate::SetNodeAndVertexNames(InOutNode, NewOutput);
 	};
 
 #if WITH_EDITORONLY_DATA
 	bool bIsRequired = false;
 	FMetasoundFrontendInterface Interface;
-	if (DocumentBuilderPrivate::IsInterfaceOutput(ClassOutput.Name, ClassOutput.TypeName, &Interface))
+	if (DocumentBuilderPrivate::IsInterfaceOutput(NewOutput.Name, NewOutput.TypeName, &Interface))
 	{
 		FText RequiredText;
-		bIsRequired = Interface.IsMemberOutputRequired(ClassOutput.Name, RequiredText);
+		bIsRequired = Interface.IsMemberOutputRequired(NewOutput.Name, RequiredText);
 	}
 #endif // WITH_EDITORONLY_DATA
+	
+
 
+	// Add output nodes
 	bool bAddedNodes = true;
 	FMetasoundFrontendNode* NewNodeToReturn = nullptr;
-	FMetasoundFrontendDocument& Document = GetDocumentChecked();
-	FMetasoundFrontendGraphClass& RootGraph = Document.RootGraph;
 	Document.RootGraph.IterateGraphPages([&](FMetasoundFrontendGraph& Graph)
 	{
-		FMetasoundFrontendNode* NewNode = AddNodeInternal(ClassKey, FinalizeNode, Graph.PageID, ClassOutput.NodeID);
+		FMetasoundFrontendNode* NewNode = AddNodeInternal(ClassKey, FinalizeNode, Graph.PageID, NewOutput.NodeID);
 		if (Graph.PageID == PageID)
 		{
 			NewNodeToReturn = NewNode;
@@ -1260,7 +1274,7 @@ const FMetasoundFrontendNode* FMetaSoundFrontendDocumentBuilder::AddGraphOutput(
 			// This is added if this is an interface member that is required to force page to create visual
 			// representation that can inform the user of its required state.
 			FGuid LocationGuid = FDocumentIDGenerator::Get().CreateVertexID(Document);
-			SetNodeLocation(ClassOutput.NodeID, FVector2D::ZeroVector, &LocationGuid, &Graph.PageID);
+			SetNodeLocation(NewOutput.NodeID, FVector2D::ZeroVector, &LocationGuid, &Graph.PageID);
 		}
 #endif // WITH_EDITORONLY_DATA
 
@@ -1269,8 +1283,6 @@ const FMetasoundFrontendNode* FMetaSoundFrontendDocumentBuilder::AddGraphOutput(
 
 	if (bAddedNodes)
 	{
-		const int32 NewIndex = RootGraph.GetDefaultInterface().Outputs.Num();
-		FMetasoundFrontendClassOutput& NewOutput = RootGraph.GetDefaultInterface().Outputs.Add_GetRef(MoveTemp(ClassOutput));
 		if (!NewOutput.VertexID.IsValid())
 		{
 			NewOutput.VertexID = FDocumentIDGenerator::Get().CreateVertexID(Document);
@@ -1281,6 +1293,11 @@ const FMetasoundFrontendNode* FMetaSoundFrontendDocumentBuilder::AddGraphOutput(
 		Document.Metadata.ModifyContext.AddMemberIDModified(NewOutput.NodeID);
 #endif // WITH_EDITORONLY_DATA
 	}
+	else
+	{
+		// Remove added output
+		RootGraph.GetDefaultInterface().Outputs.RemoveAt(NewIndex);
+	}
 
 	check(NewNodeToReturn);
 	return NewNodeToReturn;
