[MetaSounds] Allow partial copies of final output buffer based on On Finished trigger frame index
#rnx

diff --git a/Engine/Plugins/Runtime/Metasound/Source/MetasoundGenerator/Private/MetasoundGenerator.cpp b/Engine/Plugins/Runtime/Metasound/Source/MetasoundGenerator/Private/MetasoundGenerator.cpp
--- a/Engine/Plugins/Runtime/Metasound/Source/MetasoundGenerator/Private/MetasoundGenerator.cpp
+++ b/Engine/Plugins/Runtime/Metasound/Source/MetasoundGenerator/Private/MetasoundGenerator.cpp
@@ -752,27 +752,33 @@ PRAGMA_ENABLE_DEPRECATION_WARNINGS
 
 	int32 FMetasoundGenerator::FillWithBuffer(const Audio::FAlignedFloatBuffer& InBuffer, float* OutAudio, int32 MaxNumOutputSamples)
 	{
-		int32 InNum = InBuffer.Num();
-
-		if (InNum > 0)
+		const int32 InNum = InBuffer.Num();
+		if (!InNum)
 		{
-			int32 NumSamplesToCopy = FMath::Min(InNum, MaxNumOutputSamples);
-			FMemory::Memcpy(OutAudio, InBuffer.GetData(), NumSamplesToCopy * sizeof(float));
-
-			if (FinishSample != INDEX_NONE)
+			return 0;
+		}
+		
+		
+		int32 NumSamplesToCopy = FMath::Min(InNum, MaxNumOutputSamples);
+		
+		if (FinishSample != INDEX_NONE)
+		{
+			FinishSample = NumSamplesToCopy; // maxdebug, remove this. forcing edge case
+			if (FinishSample > NumSamplesToCopy)
 			{
 				FinishSample -= NumSamplesToCopy;
-				if (FinishSample <= 0)
-				{
-					bIsFinishTriggered = true;
-					FinishSample = INDEX_NONE;
-				}
 			}
-
-			return NumSamplesToCopy;
+			else
+			{
+				NumSamplesToCopy = FinishSample;
+				bIsFinishTriggered = true;
+				FinishSample = INDEX_NONE;
+			}
 		}
+		
+		FMemory::Memcpy(OutAudio, InBuffer.GetData(), NumSamplesToCopy * sizeof(float));
 
-		return 0;
+		return NumSamplesToCopy;
 	}
 
 	void FMetasoundGenerator::InterleaveGeneratedAudio()
