Fix crash when changing levels after using MetaSound user preset widget

#rnx

diff --git a/Engine/Plugins/Runtime/Metasound/Source/MetasoundEditor/Private/MetasoundEditor.cpp b/Engine/Plugins/Runtime/Metasound/Source/MetasoundEditor/Private/MetasoundEditor.cpp
--- a/Engine/Plugins/Runtime/Metasound/Source/MetasoundEditor/Private/MetasoundEditor.cpp
+++ b/Engine/Plugins/Runtime/Metasound/Source/MetasoundEditor/Private/MetasoundEditor.cpp
@@ -34,6 +34,7 @@
 #include "IMetasoundEngineModule.h"
 #include "Kismet2/BlueprintEditorUtils.h"
 #include "Kismet2/KismetEditorUtilities.h"
+#include "LevelEditor.h"
 #include "Logging/TokenizedMessage.h"
 #include "Metasound.h"
 #include "MetasoundAssetSubsystem.h"
@@ -1030,7 +1031,7 @@ namespace Metasound
 
 				InTabManager->RegisterTabSpawner(TabNamesPrivate::UserPresetWidget, FOnSpawnTab::CreateLambda([this, InUserPresetWidget = UserPresetWidget](const FSpawnTabArgs&)
 				{
-					IMetaSoundPresetWidgetInterface::Execute_OnConstructed(UserPresetWidget.Get(), Builder.Get());
+					IMetaSoundPresetWidgetInterface::Execute_OnConstructed(InUserPresetWidget.Get(), Builder.Get());
 					return SNew(SDockTab)
 						.Label(LOCTEXT("UserPresetWidgetTitle", "Preset Widget"))
 						[
@@ -1105,6 +1106,10 @@ namespace Metasound
 			PagesView.Reset();
 			InterfacesView.Reset();
 			DestroyAnalyzers();
+
+			FLevelEditorModule& LevelEditor = FModuleManager::LoadModuleChecked<FLevelEditorModule>("LevelEditor");
+			LevelEditor.OnMapChanged().RemoveAll(this);
+
 			check(GEditor);
 			GEditor->UnregisterForUndo(this);
 		}
@@ -1723,6 +1728,9 @@ namespace Metasound
 
 			if (Builder.IsValid() && Builder->IsPreset())
 			{
+				FLevelEditorModule& LevelEditor = FModuleManager::LoadModuleChecked<FLevelEditorModule>("LevelEditor");
+				LevelEditor.OnMapChanged().AddSP(this, &FEditor::ChangeUserPresetWidgetWorld);
+
 				UserPresetWidget = TStrongObjectPtr<UUserWidget>(CreateUserPresetWidget());
 			}
 
@@ -1856,6 +1864,25 @@ namespace Metasound
 			return nullptr;
 		}
 
+		void FEditor::ChangeUserPresetWidgetWorld(UWorld* World, EMapChangeType MapChangeType)
+		{
+			if (MapChangeType == EMapChangeType::TearDownWorld)
+			{
+				// We need to Delete the UMG widget if we are tearing down the World it was built with.
+				if (UserPresetWidget && World == UserPresetWidget->GetWorld())
+				{
+					UserPresetWidget->Rename(nullptr, GetTransientPackage());
+					UserPresetWidget = nullptr;
+				}
+			}
+			else if (MapChangeType != EMapChangeType::SaveMap)
+			{
+				// Recreate the widget if loading a map or opening a new map
+				// since the widget is parented to the world
+				UserPresetWidget = TStrongObjectPtr<UUserWidget>(CreateUserPresetWidget());
+			}
+		}
+
 		TSharedRef<SWidget> FEditor::CreateAuditionMenuOptions()
 		{
 			TSharedPtr<FUICommandList> Commands = MakeShared<FUICommandList>();
diff --git a/Engine/Plugins/Runtime/Metasound/Source/MetasoundEditor/Private/MetasoundEditor.h b/Engine/Plugins/Runtime/Metasound/Source/MetasoundEditor/Private/MetasoundEditor.h
--- a/Engine/Plugins/Runtime/Metasound/Source/MetasoundEditor/Private/MetasoundEditor.h
+++ b/Engine/Plugins/Runtime/Metasound/Source/MetasoundEditor/Private/MetasoundEditor.h
@@ -451,6 +451,10 @@ namespace Metasound::Editor
 
 		UUserWidget* CreateUserPresetWidget();
 
+		/** For teardown and regeneration of user preset widget on world change
+			Based on UEditorUtilityWidgetBlueprint::ChangeTabWorld */
+		void ChangeUserPresetWidgetWorld(UWorld* World, EMapChangeType MapChangeType);
+
 		/** Builds the toolbar widget for the Metasound editor */
 		void ExtendToolbarInternal();
 
