Fix crash by only creating MetaSound node configuration details when configuration set

diff --git a/Engine/Plugins/Runtime/Metasound/Source/MetasoundEditor/Private/MetasoundNodeDetailCustomization.cpp b/Engine/Plugins/Runtime/Metasound/Source/MetasoundEditor/Private/MetasoundNodeDetailCustomization.cpp
--- a/Engine/Plugins/Runtime/Metasound/Source/MetasoundEditor/Private/MetasoundNodeDetailCustomization.cpp
+++ b/Engine/Plugins/Runtime/Metasound/Source/MetasoundEditor/Private/MetasoundNodeDetailCustomization.cpp
@@ -1098,28 +1098,34 @@ PRAGMA_ENABLE_DEPRECATION_WARNINGS
 				const UEdGraph* EdGraph = Node->GetGraph();
 				const UMetasoundEditorGraph* MetasoundGraph = CastChecked<UMetasoundEditorGraph>(EdGraph);
 				const FMetaSoundFrontendDocumentBuilder& DocBuilder = MetasoundGraph->GetBuilderChecked().GetBuilder();
-				
-				// Walk the property path to find the node configuration handle
-				TSharedPtr<IPropertyHandle> DocumentHandle = InDetailLayout.AddObjectPropertyData({ MetasoundGraph->GetMetasound() }, FName(TEXT("RootMetasoundDocument")));
-				check(DocumentHandle);
-				TSharedPtr<IPropertyHandle> RootGraphHandle = DocumentHandle->GetChildHandle(GET_MEMBER_NAME_CHECKED(FMetasoundFrontendDocument, RootGraph));
-				TSharedPtr<IPropertyHandle> PagedGraphsHandle = RootGraphHandle->GetChildHandle(FName(TEXT("PagedGraphs")));
-
-				const FGuid& PageID = DocBuilder.GetBuildPageID();
-				const int32 PageIndex = DocBuilder.FindPageIndex(PageID);
-				TSharedRef<IPropertyHandle> GraphHandle = PagedGraphsHandle->AsArray()->GetElement(PageIndex);
-				TSharedPtr<IPropertyHandle> NodesHandle = GraphHandle->GetChildHandle(GET_MEMBER_NAME_CHECKED(FMetasoundFrontendGraph, Nodes));
-				
-				const int32* NodeIndex = DocBuilder.FindNodeIndex(Node->GetNodeID());
-				if (NodeIndex != nullptr)
+				// Only add configuration details customization if node has valid configuration
+				if (DocBuilder.FindNodeConfiguration(Node->GetNodeID()).IsValid())
 				{
-					TSharedRef<IPropertyHandle> NodeHandle = NodesHandle->AsArray()->GetElement(*NodeIndex);
-					TSharedPtr<IPropertyHandle> ConfigurationHandle = NodeHandle->GetChildHandle(GET_MEMBER_NAME_CHECKED(FMetasoundFrontendNode, Configuration));
+					// Walk the property path to find the node configuration handle
+					TSharedPtr<IPropertyHandle> DocumentHandle = InDetailLayout.AddObjectPropertyData({ MetasoundGraph->GetMetasound() }, FName(TEXT("RootMetasoundDocument")));
+					check(DocumentHandle);
+					TSharedPtr<IPropertyHandle> RootGraphHandle = DocumentHandle->GetChildHandle(GET_MEMBER_NAME_CHECKED(FMetasoundFrontendDocument, RootGraph));
+					TSharedPtr<IPropertyHandle> PagedGraphsHandle = RootGraphHandle->GetChildHandle(FName(TEXT("PagedGraphs")));
+
+					const FGuid& PageID = DocBuilder.GetBuildPageID();
+					const int32 PageIndex = DocBuilder.FindPageIndex(PageID);
+					TSharedRef<IPropertyHandle> GraphHandle = PagedGraphsHandle->AsArray()->GetElement(PageIndex);
+					TSharedPtr<IPropertyHandle> NodesHandle = GraphHandle->GetChildHandle(GET_MEMBER_NAME_CHECKED(FMetasoundFrontendGraph, Nodes));
+
+					const int32* NodeIndex = DocBuilder.FindNodeIndex(Node->GetNodeID());
+					if (NodeIndex != nullptr)
+					{
+						TSharedRef<IPropertyHandle> NodeHandle = NodesHandle->AsArray()->GetElement(*NodeIndex);
+						TSharedPtr<IPropertyHandle> ConfigurationHandle = NodeHandle->GetChildHandle(GET_MEMBER_NAME_CHECKED(FMetasoundFrontendNode, Configuration));
 
-					// Use a custom details builder to add custom child property update behavior
-					// and hide the struct picker since users should only be able edit child values, not the struct type 
-					TSharedRef<IDetailCustomNodeBuilder> DataDetailsBuilder = MakeShared<FMetaSoundNodeConfigurationDataDetails>(ConfigurationHandle, TWeakObjectPtr<UMetasoundEditorGraphNode>(Node));
-					InDetailLayout.EditCategory("General").AddCustomBuilder(DataDetailsBuilder);
+						// Use a custom details builder to add custom child property update behavior
+						// and hide the struct picker since users should only be able edit child values, not the struct type 
+						if (ConfigurationHandle->IsValidHandle())
+						{
+							TSharedRef<IDetailCustomNodeBuilder> DataDetailsBuilder = MakeShared<FMetaSoundNodeConfigurationDataDetails>(ConfigurationHandle, TWeakObjectPtr<UMetasoundEditorGraphNode>(Node));
+							InDetailLayout.EditCategory("General").AddCustomBuilder(DataDetailsBuilder);
+						}
+					}
 				}
 			}
 		}
diff --git a/Engine/Plugins/Runtime/Metasound/Source/MetasoundFrontend/Private/MetasoundFrontendDocumentBuilder.cpp b/Engine/Plugins/Runtime/Metasound/Source/MetasoundFrontend/Private/MetasoundFrontendDocumentBuilder.cpp
--- a/Engine/Plugins/Runtime/Metasound/Source/MetasoundFrontend/Private/MetasoundFrontendDocumentBuilder.cpp
+++ b/Engine/Plugins/Runtime/Metasound/Source/MetasoundFrontend/Private/MetasoundFrontendDocumentBuilder.cpp
@@ -2634,6 +2634,18 @@ const FMetasoundFrontendNode* FMetaSoundFrontendDocumentBuilder::FindNode(const
 	return NodeCache.FindNode(InNodeID);
 }
 
+const TInstancedStruct<FMetaSoundFrontendNodeConfiguration> FMetaSoundFrontendDocumentBuilder::FindNodeConfiguration(const FGuid& InNodeID, const FGuid* InPageID) const
+{
+	using namespace Metasound::Frontend;
+	const FGuid& PageID = InPageID ? *InPageID : BuildPageID;
+	const IDocumentGraphNodeCache& NodeCache = DocumentCache->GetNodeCache(PageID);
+	if (const FMetasoundFrontendNode* Node = NodeCache.FindNode(InNodeID))
+	{
+		return Node->Configuration;
+	}
+	return TInstancedStruct<FMetaSoundFrontendNodeConfiguration>();
+}
+
 const int32* FMetaSoundFrontendDocumentBuilder::FindNodeIndex(const FGuid& InNodeID, const FGuid* InPageID) const
 {
 	using namespace Metasound::Frontend;
diff --git a/Engine/Plugins/Runtime/Metasound/Source/MetasoundFrontend/Public/MetasoundFrontendDocumentBuilder.h b/Engine/Plugins/Runtime/Metasound/Source/MetasoundFrontend/Public/MetasoundFrontendDocumentBuilder.h
--- a/Engine/Plugins/Runtime/Metasound/Source/MetasoundFrontend/Public/MetasoundFrontendDocumentBuilder.h
+++ b/Engine/Plugins/Runtime/Metasound/Source/MetasoundFrontend/Public/MetasoundFrontendDocumentBuilder.h
@@ -261,6 +261,7 @@ public:
 	const FMetasoundFrontendGraph& FindConstBuildGraphChecked() const;
 
 	const FMetasoundFrontendNode* FindNode(const FGuid& InNodeID, const FGuid* InPageID = nullptr) const;
+	const TInstancedStruct<FMetaSoundFrontendNodeConfiguration> FindNodeConfiguration(const FGuid& InNodeID, const FGuid* InPageID = nullptr) const;
 	// Return the node's index in the document's specified paged graph's node array 
 	const int32* FindNodeIndex(const FGuid& InNodeID, const FGuid* InPageID = nullptr) const;
 
