Fixing settings categories

[RN] Fixed up settings category for MetaSounds and AudioModulation plugins to show up in Plugins.

diff --git a/Engine/Plugins/Runtime/AudioModulation/Source/AudioModulation/Public/AudioModulationSettings.h b/Engine/Plugins/Runtime/AudioModulation/Source/AudioModulation/Public/AudioModulationSettings.h
--- a/Engine/Plugins/Runtime/AudioModulation/Source/AudioModulation/Public/AudioModulationSettings.h
+++ b/Engine/Plugins/Runtime/AudioModulation/Source/AudioModulation/Public/AudioModulationSettings.h
@@ -25,6 +25,8 @@ public:
 
 	TObjectPtr<USoundModulationParameter> GetModulationParameter(const FString& InName) const;
 
+	virtual FName GetCategoryName() const override { return FName(TEXT("Plugins")); }
+
 #if WITH_EDITOR
 	virtual void PostEditChangeChainProperty(FPropertyChangedChainEvent& PropertyChangedEvent) override;
 #endif // WITH_EDITOR
diff --git a/Engine/Plugins/Runtime/Metasound/Source/MetasoundEngine/Public/MetasoundSettings.h b/Engine/Plugins/Runtime/Metasound/Source/MetasoundEngine/Public/MetasoundSettings.h
--- a/Engine/Plugins/Runtime/Metasound/Source/MetasoundEngine/Public/MetasoundSettings.h
+++ b/Engine/Plugins/Runtime/Metasound/Source/MetasoundEngine/Public/MetasoundSettings.h
@@ -137,6 +137,8 @@ class METASOUNDENGINE_API UMetaSoundSettings : public UDeveloperSettings
 	GENERATED_BODY()
 
 public:
+	virtual FName GetCategoryName() const override { return FName(TEXT("Plugins")); }
+
 	/** If true, AutoUpdate is enabled, increasing load times.  If false, skips AutoUpdate on load, but can result in MetaSounds failing to load, 
 	  * register, and execute if interface differences are present. */
 	UPROPERTY(EditAnywhere, config, Category = AutoUpdate)
