Fix MetaSound ensure when referencing a data type from a disabled plugin

diff --git a/Engine/Plugins/Runtime/Metasound/Source/MetasoundFrontend/Private/MetasoundFrontendDataTypeRegistry.cpp b/Engine/Plugins/Runtime/Metasound/Source/MetasoundFrontend/Private/MetasoundFrontendDataTypeRegistry.cpp
--- a/Engine/Plugins/Runtime/Metasound/Source/MetasoundFrontend/Private/MetasoundFrontendDataTypeRegistry.cpp
+++ b/Engine/Plugins/Runtime/Metasound/Source/MetasoundFrontend/Private/MetasoundFrontendDataTypeRegistry.cpp
@@ -1503,10 +1503,14 @@ PRAGMA_ENABLE_DEPRECATION_WARNINGS
 				{
 					const TSharedRef<IDataTypeRegistryEntry, ESPMode::ThreadSafe>* Entry = RegisteredDataTypes.Find(InDataTypeName);
 
-					if (ensureMsgf(nullptr != Entry, TEXT("Data type not registered [Name:%s]"), *InDataTypeName.ToString()))
+					if (nullptr != Entry)
 					{
 						return &Entry->Get();
 					}
+					else
+					{
+						UE_LOG(LogMetaSound, Error, TEXT("Data type is not registered [Name:%s]"), *InDataTypeName.ToString());
+					}
 				}
 
 				return nullptr;
