Add convenience function to see if a MetaSound document inherits a named interface.

diff --git a/Engine/Plugins/Runtime/Metasound/Source/MetasoundFrontend/Private/MetasoundFrontendQuery.cpp b/Engine/Plugins/Runtime/Metasound/Source/MetasoundFrontend/Private/MetasoundFrontendQuery.cpp
--- a/Engine/Plugins/Runtime/Metasound/Source/MetasoundFrontend/Private/MetasoundFrontendQuery.cpp
+++ b/Engine/Plugins/Runtime/Metasound/Source/MetasoundFrontend/Private/MetasoundFrontendQuery.cpp
@@ -228,6 +228,25 @@ namespace Metasound
 			return bSuccess;
 		}
 
+#if WITH_EDITORONLY_DATA
+		bool FMetaSoundClassInfo::InheritsInterface(FName InterfaceName, const  FMetasoundFrontendVersionNumber& VersionNumber) const
+		{
+			if (bIsValid==false)
+			{
+				return false;
+			}
+
+			if (!VersionNumber.IsValid())
+			{
+				// match name only...
+				return InterfaceInfo.InheritedInterfaces.ContainsByPredicate([InterfaceName](const FMetasoundFrontendVersion& VersionInfo) { return VersionInfo.Name == InterfaceName; });
+			}
+			
+			// require exact version match...
+			const FMetasoundFrontendVersion OfInterest = {InterfaceName, VersionNumber };
+			return InterfaceInfo.InheritedInterfaces.Contains(OfInterest);
+		}
+#endif
 	} // namespace Frontend
 
 	namespace FrontendQueryPrivate
diff --git a/Engine/Plugins/Runtime/Metasound/Source/MetasoundFrontend/Public/MetasoundFrontendQuery.h b/Engine/Plugins/Runtime/Metasound/Source/MetasoundFrontend/Public/MetasoundFrontendQuery.h
--- a/Engine/Plugins/Runtime/Metasound/Source/MetasoundFrontend/Public/MetasoundFrontendQuery.h
+++ b/Engine/Plugins/Runtime/Metasound/Source/MetasoundFrontend/Public/MetasoundFrontendQuery.h
@@ -184,6 +184,8 @@ namespace Metasound
 
 #if WITH_EDITORONLY_DATA
 			FMetaSoundClassInterfaceInfo InterfaceInfo;
+
+			bool InheritsInterface(FName InterfaceName, const  FMetasoundFrontendVersionNumber& VersionNumber = { 0, 0 }) const;
 #endif // WITH_EDITORONLY_DATA
 
 			// If true, class info is valid and accurately reflects
