Fix MetaSound comment nodes not being applied correctly per page

diff --git a/Engine/Plugins/Runtime/Metasound/Source/MetasoundEditor/Private/MetasoundEditorGraphBuilder.cpp b/Engine/Plugins/Runtime/Metasound/Source/MetasoundEditor/Private/MetasoundEditorGraphBuilder.cpp
--- a/Engine/Plugins/Runtime/Metasound/Source/MetasoundEditor/Private/MetasoundEditorGraphBuilder.cpp
+++ b/Engine/Plugins/Runtime/Metasound/Source/MetasoundEditor/Private/MetasoundEditorGraphBuilder.cpp
@@ -1876,6 +1876,11 @@ PRAGMA_ENABLE_DEPRECATION_WARNINGS
 					NewNode->CommentID = Pair.Key;
 					UMetasoundEditorGraphCommentNode::ConvertFromFrontendComment(Pair.Value, *NewNode);
 				}
+				else
+				{
+					const FMetaSoundFrontendGraphComment* Comment = InBuilder.FindGraphComment(Pair.Key);
+					UMetasoundEditorGraphCommentNode::ConvertFromFrontendComment(*Comment, *CommentNode);
+				}
 			}
 
 			// Remaining items are stale, so they are removed
