[Metasound]
* Moved dllstorage macro to after UE_INTERNAL

diff --git a/Engine/Plugins/Runtime/Metasound/Source/MetasoundFrontend/Public/MetasoundFrontendDocument.h b/Engine/Plugins/Runtime/Metasound/Source/MetasoundFrontend/Public/MetasoundFrontendDocument.h
--- a/Engine/Plugins/Runtime/Metasound/Source/MetasoundFrontend/Public/MetasoundFrontendDocument.h
+++ b/Engine/Plugins/Runtime/Metasound/Source/MetasoundFrontend/Public/MetasoundFrontendDocument.h
@@ -503,11 +503,11 @@ struct FMetasoundFrontendNodeInterface
 	// Create a node interface which satisfies an existing class interface.
 	UE_API FMetasoundFrontendNodeInterface(const FMetasoundFrontendClassInterface& InClassInterface);
 
-	UE_API UE_INTERNAL
-	void Update(const FMetasoundFrontendClassInterface& InClassInterface);
+	UE_INTERNAL
+	UE_API void Update(const FMetasoundFrontendClassInterface& InClassInterface);
 
-	UE_API UE_INTERNAL
-	void Update(const FMetasoundFrontendClassInterface& InClassInterface, TFunctionRef<void(const FMetasoundFrontendVertex&)> OnPreRemoveInput, TFunctionRef<void(const FMetasoundFrontendVertex&)> OnPreRemoveOutput);
+	UE_INTERNAL
+	UE_API void Update(const FMetasoundFrontendClassInterface& InClassInterface, TFunctionRef<void(const FMetasoundFrontendVertex&)> OnPreRemoveInput, TFunctionRef<void(const FMetasoundFrontendVertex&)> OnPreRemoveOutput);
 
 	// Input vertices to node.
 	UPROPERTY()
