Generically propagate all MetaSound and soundwave settings from referenced asset when creating MetaSound preset

diff --git a/Engine/Plugins/Runtime/Metasound/Source/MetasoundEditor/Private/MetasoundAssetDefinitions.cpp b/Engine/Plugins/Runtime/Metasound/Source/MetasoundEditor/Private/MetasoundAssetDefinitions.cpp
--- a/Engine/Plugins/Runtime/Metasound/Source/MetasoundEditor/Private/MetasoundAssetDefinitions.cpp
+++ b/Engine/Plugins/Runtime/Metasound/Source/MetasoundEditor/Private/MetasoundAssetDefinitions.cpp
@@ -400,19 +400,19 @@ namespace MenuExtension_MetaSoundSourceTemplate
 
 				IAssetTools::Get().CreateUniqueAssetName(ReferencedMetaSound->GetOutermost()->GetName(), TEXT("_Preset"), PackagePath, AssetName);
 
-				UMetaSoundBaseFactory* Factory = NewObject<TFactoryClass>();
-				Factory->ReferencedMetaSoundObject = ReferencedMetaSound;
-
-				UObject* NewMetaSound = IAssetTools::Get().CreateAssetWithDialog(AssetName, FPackageName::GetLongPackagePath(PackagePath), Factory->GetSupportedClass(), Factory);
-
+				// Duplicate asset to preserve properties from referenced asset (ex. quality settings, soundwave properties)
+				UObject* NewMetaSound = IAssetTools::Get().DuplicateAssetWithDialogAndTitle(AssetName, FPackageName::GetLongPackagePath(PackagePath), ReferencedMetaSound, LOCTEXT("CreateMetaSoundPresetTitle", "Create MetaSound Preset"));
+				
 				if (NewMetaSound)
 				{
+					UMetaSoundEditorSubsystem::GetChecked().InitAsset(*NewMetaSound, ReferencedMetaSound, /*bClearDocument=*/true);
+
 					FGraphBuilder::RegisterGraphWithFrontend(*NewMetaSound);
 					ObjectsToSync.Add(NewMetaSound);
 				}
 				else
 				{
-					UE_LOG(LogMetaSound, Display, TEXT("Error creating new asset with factory '%s' when creating preset '%s' or asset creation was canceled by user."), *Factory->GetName(), *AssetName);
+					UE_LOG(LogMetaSound, Display, TEXT("Error creating new asset when creating preset '%s' or asset creation was canceled by user."), *AssetName);
 				}
 			}
 
diff --git a/Engine/Plugins/Runtime/Metasound/Source/MetasoundEditor/Private/MetasoundEditorSubsystem.cpp b/Engine/Plugins/Runtime/Metasound/Source/MetasoundEditor/Private/MetasoundEditorSubsystem.cpp
--- a/Engine/Plugins/Runtime/Metasound/Source/MetasoundEditor/Private/MetasoundEditorSubsystem.cpp
+++ b/Engine/Plugins/Runtime/Metasound/Source/MetasoundEditor/Private/MetasoundEditorSubsystem.cpp
@@ -52,11 +52,28 @@ TScriptInterface<IMetaSoundDocumentInterface> UMetaSoundEditorSubsystem::BuildTo
 			InBuilder->AddToRoot();
 		}
 
-		constexpr UFactory* Factory = nullptr;
-
 		// Not about to follow this lack of const correctness down a multidecade in the works rabbit hole.
 		UClass& MetaSoundUClass = const_cast<UClass&>(InBuilder->GetBaseMetaSoundUClass());
-		if (UObject* NewMetaSound = IAssetTools::Get().CreateAsset(AssetName, PackagePath, &MetaSoundUClass, Factory))
+
+		UObject* NewMetaSound;
+		
+		// Duplicate referenced preset object to preserve object properties (ex. quality settings, soundwave properties)
+		if (InBuilder->IsPreset())
+		{
+			// Rebuild referenced classes to find referenced preset asset
+			FMetasoundAssetBase& InAsset = InBuilder->GetBuilder().GetMetasoundAsset();
+			InAsset.RebuildReferencedAssetClasses();
+
+			UObject* ReferencedObject = InBuilder->GetReferencedPresetAsset();
+			NewMetaSound = IAssetTools::Get().DuplicateAsset(AssetName, PackagePath, ReferencedObject);
+		}
+		else
+		{
+			constexpr UFactory* Factory = nullptr;
+			NewMetaSound = IAssetTools::Get().CreateAsset(AssetName, PackagePath, &MetaSoundUClass, Factory);
+		}
+
+		if (NewMetaSound)
 		{
 			InBuilder->InitNodeLocations();
 			InBuilder->SetAuthor(Author);
@@ -79,24 +96,12 @@ TScriptInterface<IMetaSoundDocumentInterface> UMetaSoundEditorSubsystem::BuildTo
 			Asset.RebuildReferencedAssetClasses();
 
 			// Apply template SoundWave settings 
-			// (must be post rebuilding referenced asset classes to find referenced asset to copy settings from)
 			{
-				const bool bIsSource = &MetaSoundUClass == UMetaSoundSource::StaticClass();
-				if (NewDocBuilder.IsPreset())
-				{
-					// Only use referenced UObject's SoundWave settings for sources if not overridden 
-					if (TemplateSoundWave == nullptr && bIsSource)
-					{
-						const UObject* ReferencedObject = NewDocBuilder.GetReferencedPresetAsset();
-						if (ensureMsgf(ReferencedObject, TEXT("Preset builder %s does not have a referenced asset to apply soundwave settings from."), *NewDocBuilder.GetName()))
-						{
-							TemplateSoundWave = CastChecked<USoundWave>(ReferencedObject);
-						}
-					}
-				}
-
 				// Template SoundWave settings only apply to sources
-				if (TemplateSoundWave != nullptr && bIsSource)
+				// and do not apply to presets, which were duplicated from the referenced preset asset 
+				// and therefore have those object properties 
+				const bool bIsSource = &MetaSoundUClass == UMetaSoundSource::StaticClass();
+				if (TemplateSoundWave != nullptr && bIsSource && !NewDocBuilder.IsPreset())
 				{
 					SetSoundWaveSettingsFromTemplate(*CastChecked<USoundWave>(NewMetaSound), *TemplateSoundWave);
 				}
@@ -316,7 +321,7 @@ const TArray<TSharedRef<FExtender>>& UMetaSoundEditorSubsystem::GetToolbarExtend
 	return EditorToolbarExtenders;
 }
 
-void UMetaSoundEditorSubsystem::InitAsset(UObject& InNewMetaSound, UObject* InReferencedMetaSound)
+void UMetaSoundEditorSubsystem::InitAsset(UObject& InNewMetaSound, UObject* InReferencedMetaSound, const bool bClearDocument)
 {
 	using namespace Metasound;
 	using namespace Metasound::Editor;
@@ -325,6 +330,12 @@ void UMetaSoundEditorSubsystem::InitAsset(UObject& InNewMetaSound, UObject* InRe
 	TScriptInterface<IMetaSoundDocumentInterface> DocInterface = &InNewMetaSound;
 	FMetaSoundFrontendDocumentBuilder Builder(DocInterface);
 
+	if (bClearDocument)
+	{
+		TSharedRef<FDocumentModifyDelegates> ModifyDelegates = MakeShared<FDocumentModifyDelegates>(Builder.GetConstDocumentChecked());
+		Builder.ClearDocument(ModifyDelegates);
+	}
+	
 	Builder.InitDocument();
 	Builder.InitNodeLocations();
 
@@ -343,12 +354,6 @@ void UMetaSoundEditorSubsystem::InitAsset(UObject& InNewMetaSound, UObject* InRe
 		// Initialize preset with referenced Metasound 
 		TScriptInterface<IMetaSoundDocumentInterface> ReferencedDocInterface = InReferencedMetaSound;
 		Builder.ConvertToPreset(ReferencedDocInterface->GetConstDocument());
-
-		// Copy sound wave settings to preset for sources
-		if (&ReferencedDocInterface->GetBaseMetaSoundUClass() == UMetaSoundSource::StaticClass())
-		{
-			SetSoundWaveSettingsFromTemplate(*CastChecked<USoundWave>(&InNewMetaSound), *CastChecked<USoundWave>(InReferencedMetaSound));
-		}
 	}
 }
 
diff --git a/Engine/Plugins/Runtime/Metasound/Source/MetasoundEditor/Public/MetasoundEditorSubsystem.h b/Engine/Plugins/Runtime/Metasound/Source/MetasoundEditor/Public/MetasoundEditorSubsystem.h
--- a/Engine/Plugins/Runtime/Metasound/Source/MetasoundEditor/Public/MetasoundEditorSubsystem.h
+++ b/Engine/Plugins/Runtime/Metasound/Source/MetasoundEditor/Public/MetasoundEditorSubsystem.h
@@ -78,8 +78,10 @@ public:
 		UPARAM(DisplayName = "Location") const FVector2D& InLocation,
 		EMetaSoundBuilderResult& OutResult);
 
-	// Initialize the UObject asset, with an optional MetaSound to be referenced if the asset is a preset
-	void InitAsset(UObject& InNewMetaSound, UObject* InReferencedMetaSound = nullptr);
+	// Initialize the UObject asset
+	// with an optional MetaSound to be referenced if the asset is a preset
+	// and optionally clearing the existing MetaSound document (for the case of duplicated assets) 
+	void InitAsset(UObject& InNewMetaSound, UObject* InReferencedMetaSound = nullptr, const bool bClearDocument = false);
 
 	UE_DEPRECATED(5.5, "EdGraph is now transiently generated and privately managed for asset editor use only.")
 	void InitEdGraph(UObject& InMetaSound);
diff --git a/Engine/Plugins/Runtime/Metasound/Source/MetasoundFrontend/Private/MetasoundFrontendDocumentBuilder.cpp b/Engine/Plugins/Runtime/Metasound/Source/MetasoundFrontend/Private/MetasoundFrontendDocumentBuilder.cpp
--- a/Engine/Plugins/Runtime/Metasound/Source/MetasoundFrontend/Private/MetasoundFrontendDocumentBuilder.cpp
+++ b/Engine/Plugins/Runtime/Metasound/Source/MetasoundFrontend/Private/MetasoundFrontendDocumentBuilder.cpp
@@ -3335,7 +3335,10 @@ void FMetaSoundFrontendDocumentBuilder::InitDocument(const FMetasoundFrontendDoc
 	else
 	{
 		// 1a. Initialize class using default data
-		Document.RootGraph.InitDefaultGraphPage();
+		if (Document.RootGraph.GetConstGraphPages().IsEmpty())
+		{
+			Document.RootGraph.InitDefaultGraphPage();
+		}
 		FMetasoundFrontendClassMetadata& ClassMetadata = Document.RootGraph.Metadata;
 		InitGraphClassMetadata(Document.RootGraph.Metadata, bResetVersion, InNewClassName);
 
