Fix for MetaSound editor not opening presets in RestrictedMode
#rnx

diff --git a/Engine/Plugins/Runtime/Metasound/Source/MetasoundEditor/Private/MetasoundAssetDefinitions.cpp b/Engine/Plugins/Runtime/Metasound/Source/MetasoundEditor/Private/MetasoundAssetDefinitions.cpp
--- a/Engine/Plugins/Runtime/Metasound/Source/MetasoundEditor/Private/MetasoundAssetDefinitions.cpp
+++ b/Engine/Plugins/Runtime/Metasound/Source/MetasoundEditor/Private/MetasoundAssetDefinitions.cpp
@@ -154,14 +154,40 @@ TConstArrayView<FAssetCategoryPath> UAssetDefinition_MetaSoundSource::GetAssetCa
 
 EAssetCommandResult UAssetDefinition_MetaSoundSource::OpenAssets(const FAssetOpenArgs& OpenArgs) const
 {
-	Metasound::Editor::IMetasoundEditorModule& MetaSoundEditorModule = FModuleManager::GetModuleChecked<Metasound::Editor::IMetasoundEditorModule>("MetaSoundEditor");
-	for (UMetaSoundSource* Metasound : OpenArgs.LoadObjects<UMetaSoundSource>())
+	using namespace Metasound;
+
+	for (const FAssetData& AssetData : OpenArgs.Assets)
 	{
-		// In restricted mode, we only want to open up the editor if we're preset otherwise, we don't want to open the editor
-		if (Metasound->bIsPreset || !MetaSoundEditorModule.IsRestrictedMode())
+		const UClass* AssetClass = AssetData.GetClass();
+		if (AssetClass && IMetasoundUObjectRegistry::Get().IsRegisteredClass(*AssetClass))
 		{
-			TSharedRef<Metasound::Editor::FEditor> NewEditor = MakeShared<Metasound::Editor::FEditor>();
-			NewEditor->InitMetasoundEditor(OpenArgs.GetToolkitMode(), OpenArgs.ToolkitHost, Metasound);
+			Engine::FMetaSoundAssetManager& AssetManager = Engine::FMetaSoundAssetManager::GetChecked();
+			TWeakPtr<IToolkitHost> ToolkitHost = OpenArgs.ToolkitHost;
+			const bool bHostNull = !OpenArgs.ToolkitHost.IsValid();
+			AssetManager.AddOrLoadAndUpdateFromObjectAsync(AssetData, [ToolkitMode = OpenArgs.GetToolkitMode(), bHostNull, ToolkitHost](FMetaSoundAssetKey, UObject& MetaSoundObject)
+			{
+				TSharedPtr<IToolkitHost> HostPtr = ToolkitHost.Pin();
+				if (bHostNull || HostPtr)
+				{
+					Editor::IMetasoundEditorModule* EditorModule = FModuleManager::GetModulePtr<Editor::IMetasoundEditorModule>("MetaSoundEditor");
+					if (EditorModule)
+					{
+						if (EditorModule->IsRestrictedMode())
+						{
+							TScriptInterface<const IMetaSoundDocumentInterface> DocInterface(&MetaSoundObject);
+							check(DocInterface.GetObject());
+							const Frontend::FMetaSoundAssetClassInfo ClassInfo(*DocInterface.GetInterface());
+							if (!ClassInfo.bIsValid || !ClassInfo.DocInfo.bIsPreset)
+							{
+								return;
+							}
+						}
+
+						TSharedRef<Editor::FEditor> NewEditor = MakeShared<Editor::FEditor>();
+						NewEditor->InitMetasoundEditor(ToolkitMode, HostPtr, &MetaSoundObject);
+					}
+				}
+			});
 		}
 	}
 
diff --git a/Engine/Plugins/Runtime/Metasound/Source/MetasoundEngine/Private/MetasoundAssetSubsystem.cpp b/Engine/Plugins/Runtime/Metasound/Source/MetasoundEngine/Private/MetasoundAssetSubsystem.cpp
--- a/Engine/Plugins/Runtime/Metasound/Source/MetasoundEngine/Private/MetasoundAssetSubsystem.cpp
+++ b/Engine/Plugins/Runtime/Metasound/Source/MetasoundEngine/Private/MetasoundAssetSubsystem.cpp
@@ -224,7 +224,7 @@ namespace Metasound::Engine
 		return AssetSubsystemPrivate::AddClassInfo(TagDataMapCriticalSection, ClassInfoMap, FMetaSoundAssetClassInfo (DocInterface));
 	}
 
-	void FMetaSoundAssetManager::AddOrLoadAndUpdateFromObjectAsync(const FAssetData& InAssetData, TUniqueFunction<void(FMetaSoundAssetKey, UObject&)>&& OnAssetRegistered)
+	void FMetaSoundAssetManager::AddOrLoadAndUpdateFromObjectAsync(const FAssetData& InAssetData, Frontend::IMetaSoundAssetManager::FOnUpdatedAssetLoaded&& OnUpdatedAssetLoaded)
 	{
 		using namespace Frontend;
 
@@ -233,7 +233,7 @@ namespace Metasound::Engine
 		// Don't add temporary assets used for diffing
 		if (InAssetData.HasAnyPackageFlags(PKG_ForDiffing))
 		{
-			OnAssetRegistered.Reset();
+			OnUpdatedAssetLoaded.Reset();
 			return;
 		}
 
@@ -242,8 +242,8 @@ namespace Metasound::Engine
 		{
 			UE_LOG(LogMetaSound, VeryVerbose, TEXT("FMetaSoundAssetManager::AddOrLoadAndUpdateFromObjectAsync called, but object ''%s' already loaded. Updating entry using in-memory object definition."), *Path.ToString());
 			const FMetaSoundAssetKey AssetKey = AddOrUpdateFromObject(*Object);
-			OnAssetRegistered(AssetKey, *Object);
-			OnAssetRegistered.Reset();
+			OnUpdatedAssetLoaded(AssetKey, *Object);
+			OnUpdatedAssetLoaded.Reset();
 			return;
 		}
 
@@ -252,16 +252,17 @@ namespace Metasound::Engine
 #endif // WITH_EDITORONLY_DATA
 
 		UE_LOG(LogMetaSound, VeryVerbose, TEXT("MetaSoundAssetManager requested aync loading asset '%s'..."), *InAssetData.GetSoftObjectPath().ToString());
-		LoadPackageAsync(Path.GetLongPackageName(), FLoadPackageAsyncDelegate::CreateLambda([](const FName& PackageName, UPackage* Package, EAsyncLoadingResult::Type Result)
+
+		LoadPackageAsync(Path.GetLongPackageName(), FLoadPackageAsyncDelegate::CreateLambda([OnAssetLoaded = MoveTemp(OnUpdatedAssetLoaded)](const FName& PackageName, UPackage* Package, EAsyncLoadingResult::Type Result) mutable
+		{
+			FMetaSoundAssetManager::OnPackageLoaded(FPackageLoadedArgs
 			{
-				FPackageLoadedArgs Args;
-				Args.bMarkAsVersioned = true;
-				Args.Package = Package;
-				Args.PackageName = PackageName;
-				Args.Result = Result;
-				FMetaSoundAssetManager::OnPackageLoaded(Args);
-			})
-		);
+				.PackageName = PackageName,
+				.Package = Package,
+				.Result = Result,
+				.OnUpdatedAssetLoaded = MoveTemp(OnAssetLoaded)
+			});
+		}));
 	}
 
 	void FMetaSoundAssetManager::AddOrUpdateFromAssetData(const FAssetData& InAssetData)
@@ -299,16 +300,11 @@ namespace Metasound::Engine
 #endif // WITH_EDITORONLY_DATA
 		
 				UE_LOG(LogMetaSound, Verbose, TEXT("MetaSound tags require updating: MetaSoundAssetManager aync loading asset '%s' to access tags/update entry..."), *InAssetData.GetSoftObjectPath().ToString());
+
 				LoadPackageAsync(Path.GetLongPackageName(), FLoadPackageAsyncDelegate::CreateLambda([](const FName& PackageName, UPackage* Package, EAsyncLoadingResult::Type Result)
-					{
-						FPackageLoadedArgs Args;
-						Args.bMarkAsVersioned = true;
-						Args.Package = Package;
-						Args.PackageName = PackageName;
-						Args.Result = Result;
-						FMetaSoundAssetManager::OnPackageLoaded(Args);
-					})
-				);
+				{
+					FMetaSoundAssetManager::OnPackageLoaded(FPackageLoadedArgs { .PackageName = PackageName, .Package = Package, .Result = Result });
+				}));
 			}
 		}
 	}
@@ -593,19 +589,16 @@ namespace Metasound::Engine
 					// This flags version scripts to properly resave even if the document version remained the same.
 					{
 #if WITH_EDITORONLY_DATA
-						if (PackageLoadedArgs.bMarkAsVersioned)
+						if (FMetasoundAssetBase* AssetBase = AssetManager->GetAsAsset(*MetaSoundObj))
 						{
-							if (FMetasoundAssetBase* AssetBase = AssetManager->GetAsAsset(*MetaSoundObj))
-							{
-								AssetBase->SetVersionedOnLoad();
-							}
+							AssetBase->SetVersionedOnLoad();
 						}
 #endif // WITH_EDITORONLY_DATA
 
 						const FMetaSoundAssetKey AssetKey = AssetManager->AddOrUpdateFromObject(*MetaSoundObj);
-						if (PackageLoadedArgs.OnAssetRegistered)
+						if (PackageLoadedArgs.OnUpdatedAssetLoaded)
 						{
-							PackageLoadedArgs.OnAssetRegistered(AssetKey, *MetaSoundObj);
+							PackageLoadedArgs.OnUpdatedAssetLoaded(AssetKey, *MetaSoundObj);
 						}
 					}
 				}
diff --git a/Engine/Plugins/Runtime/Metasound/Source/MetasoundEngine/Public/MetasoundAssetSubsystem.h b/Engine/Plugins/Runtime/Metasound/Source/MetasoundEngine/Public/MetasoundAssetSubsystem.h
--- a/Engine/Plugins/Runtime/Metasound/Source/MetasoundEngine/Public/MetasoundAssetSubsystem.h
+++ b/Engine/Plugins/Runtime/Metasound/Source/MetasoundEngine/Public/MetasoundAssetSubsystem.h
@@ -39,6 +39,7 @@ namespace Metasound::Engine
 		friend class UMetaSoundAssetSubsystem;
 
 		using FAssetClassInfoMap = TMap<FMetaSoundAssetKey, TArray<Frontend::FMetaSoundAssetClassInfo>>;
+		using FOnUpdatedAssetLoaded = Frontend::IMetaSoundAssetManager::FOnUpdatedAssetLoaded;
 
 PRAGMA_DISABLE_DEPRECATION_WARNINGS 
 		using FAssetInfo UE_DEPRECATED(5.6, "Use FAssetRef instead which is more compact and usage clear in context of other Info tag/query types") = Frontend::IMetaSoundAssetManager::FAssetInfo;
@@ -108,7 +109,7 @@ PRAGMA_ENABLE_DEPRECATION_WARNINGS
 		UE_API virtual bool AddAssetReferences(FMetasoundAssetBase& InAssetBase) override;
 #endif // WITH_EDITORONLY_DATA
 		UE_API virtual FMetaSoundAssetKey AddOrUpdateFromObject(const UObject& InObject) override;
-		UE_API virtual void AddOrLoadAndUpdateFromObjectAsync(const FAssetData& InAssetData, TUniqueFunction<void(FMetaSoundAssetKey, UObject&)>&& OnAssetRegistered) override;
+		UE_API virtual void AddOrLoadAndUpdateFromObjectAsync(const FAssetData& InAssetData, Frontend::IMetaSoundAssetManager::FOnUpdatedAssetLoaded&& OnUpdatedAssetLoaded) override;
 		UE_API virtual void AddOrUpdateFromAssetData(const FAssetData& InAssetData) override;
 		UE_API virtual bool CanAutoUpdate(const FMetasoundFrontendClassName& InClassName) const override;
 		UE_API virtual bool ContainsKey(const FMetaSoundAssetKey& InKey) const override;
@@ -177,8 +178,7 @@ PRAGMA_ENABLE_DEPRECATION_WARNINGS
 			FName PackageName;
 			UPackage* Package = nullptr;
 			EAsyncLoadingResult::Type Result = EAsyncLoadingResult::Failed;
-			TUniqueFunction<void(FMetaSoundAssetKey, UObject&)> OnAssetRegistered;
-			bool bMarkAsVersioned = false;
+			FOnUpdatedAssetLoaded OnUpdatedAssetLoaded;
 		};
 
 		static UE_API void OnPackageLoaded(const FPackageLoadedArgs& PackageLoadedArgs);
diff --git a/Engine/Plugins/Runtime/Metasound/Source/MetasoundFrontend/Public/MetasoundAssetManager.h b/Engine/Plugins/Runtime/Metasound/Source/MetasoundFrontend/Public/MetasoundAssetManager.h
--- a/Engine/Plugins/Runtime/Metasound/Source/MetasoundFrontend/Public/MetasoundAssetManager.h
+++ b/Engine/Plugins/Runtime/Metasound/Source/MetasoundFrontend/Public/MetasoundAssetManager.h
@@ -182,6 +182,11 @@ namespace Metasound::Frontend
 		UE_DEPRECATED(5.5, "Use Initialize/Deinitialize instead")
 		static void Set(IMetaSoundAssetManager& InInterface) { checkNoEntry(); }
 
+		// Passed template function pointer to execute in certain contexts when a MetaSound object is either already
+		// loaded or has completed asynchronous load.  Due to some implementation restrictions when constructing
+		// FLoadPackageAsyncDelegate, this has to be a copyable function (vs. TUniqueFunction).
+		using FOnUpdatedAssetLoaded = TFunction<void(FMetaSoundAssetKey, UObject&)>;
+
 		struct FAssetRef
 		{
 			FMetaSoundAssetKey Key;
@@ -246,7 +251,7 @@ namespace Metasound::Frontend
 		// Add or Update a MetaSound Asset's entry data from an object, loading it if it isn't already.  On initial call, requests object load
 		// asynchronously and runs provided function on successful completion. If asset is already loaded, runs provided function immediately
 		// on entry update (synchronously).
-		virtual void AddOrLoadAndUpdateFromObjectAsync(const FAssetData& InAssetData, TUniqueFunction<void(FMetaSoundAssetKey, UObject&)>&& OnAssetRegistered) = 0;
+		virtual void AddOrLoadAndUpdateFromObjectAsync(const FAssetData& InAssetData, FOnUpdatedAssetLoaded&& OnUpdatedAssetLoaded) = 0;
 
 		// Add or Update a MetaSound Asset's entry data from AssetData. Potentially loads asset
 		// and adds asynchronously if asset or associated tag schema is out-of-date.
