Fix MetaSound crash when clicking audition menu after closing asset editor

diff --git a/Engine/Plugins/Runtime/Metasound/Source/MetasoundEditor/Private/MetasoundEditor.cpp b/Engine/Plugins/Runtime/Metasound/Source/MetasoundEditor/Private/MetasoundEditor.cpp
--- a/Engine/Plugins/Runtime/Metasound/Source/MetasoundEditor/Private/MetasoundEditor.cpp
+++ b/Engine/Plugins/Runtime/Metasound/Source/MetasoundEditor/Private/MetasoundEditor.cpp
@@ -2233,34 +2233,27 @@ namespace Metasound
 								);
 							}
 							ToolbarBuilder.EndStyleOverride();
+
+							constexpr bool bHasProjectPageValues = true;
+							constexpr bool bPresetCanEditPageValues = true;
+							const bool bShowAuditionSettings = PageEditorEnabled(Builder->GetConstBuilder(), bHasProjectPageValues, bPresetCanEditPageValues);
+							if (bShowAuditionSettings)
+							{
+								ToolbarBuilder.AddComboButton(
+									FUIAction(),
+									FOnGetContent::CreateSP(this, &FEditor::CreateAuditionMenuOptions),
+									LOCTEXT("AuditionSettingsMenu", "Audition"),
+									LOCTEXT("AuditionSettingsMenu_Tooltip", "Settings related to auditioning MetaSound (Target page, platform etc.)"),
+									Style::CreateSlateIcon("MetasoundEditor.Audition"),
+									false
+								);
+							}
 						}
 						ToolbarBuilder.EndSection();
 					}
 				})
 			);
 
-			constexpr bool bHasProjectPageValues = true;
-			constexpr bool bPresetCanEditPageValues = true;
-			const bool bShowAuditionSettings = PageEditorEnabled(Builder->GetConstBuilder(), bHasProjectPageValues, bPresetCanEditPageValues);
-			if (bShowAuditionSettings)
-			{
-				if (UToolMenu* AssetToolbar = UToolMenus::Get()->ExtendMenu(GetToolMenuToolbarName()))
-				{
-					TSharedPtr<FUICommandList> CommandList = MakeShared<FUICommandList>();
-					FToolMenuSection& Section = AssetToolbar->FindOrAddSection("Asset.Utilities");
-					FToolMenuEntry Entry = FToolMenuEntry::InitComboButton(
-						"AuditionMenu",
-						FUIAction(),
-						FOnGetContent::CreateRaw(this, &FEditor::CreateAuditionMenuOptions),
-						LOCTEXT("AuditionSettingsMenu", "Audition"),
-						LOCTEXT("AuditionSettingsMenu_Tooltip", "Settings related to auditioning MetaSound (Target page, platform etc.)"),
-						Style::CreateSlateIcon("MetasoundEditor.Audition"),
-						false);
-					Entry.StyleNameOverride = "CalloutToolbar";
-					Section.AddEntry(Entry);
-				}
-			}
-
 			AddToolbarExtender(ToolbarExtender);
 
 			if (GEditor)
