MetaSound Editor Optimization Part 2:
- Add more controller logic to builder API (callsites will be updated in subsequent change)
- Misc Controller to Builder clean-up

#rnx

diff --git a/Engine/Plugins/Runtime/Metasound/Source/MetasoundEditor/Private/MetasoundEditorGraphBuilder.cpp b/Engine/Plugins/Runtime/Metasound/Source/MetasoundEditor/Private/MetasoundEditorGraphBuilder.cpp
--- a/Engine/Plugins/Runtime/Metasound/Source/MetasoundEditor/Private/MetasoundEditorGraphBuilder.cpp
+++ b/Engine/Plugins/Runtime/Metasound/Source/MetasoundEditor/Private/MetasoundEditorGraphBuilder.cpp
@@ -1608,21 +1608,17 @@ PRAGMA_ENABLE_DEPRECATION_WARNINGS
 			}
 		}
 
-		void FGraphBuilder::RefreshPinMetadata(UEdGraphPin& InPin, const FMetasoundFrontendVertexMetadata& InMetadata)
+		void FGraphBuilder::RefreshPinMetadata(UEdGraphPin& InPin, bool bAdvancedView)
 		{
+			UMetasoundEditorGraphNode* MetasoundGraphNode = Cast<UMetasoundEditorGraphNode>(InPin.GetOwningNode());
+			check(MetasoundGraphNode);
+
 			// Pin ToolTips are no longer cached on pins, and are instead dynamically generated via UMetasoundEditorGraphNode::GetPinHoverText
 			InPin.PinToolTip = { };
-			InPin.bAdvancedView = InMetadata.bIsAdvancedDisplay;
-
-			//Are the Unconnected hidden for the node
-			bool bUnconnectedHidden = false;
-			UMetasoundEditorGraphNode* MetasoundGraphNode = Cast<UMetasoundEditorGraphNode>(InPin.GetOwningNode());
-			{
-				FMetasoundFrontendNode FrontendNode = MetasoundGraphNode->GetFrontendNodeChecked();
-				bUnconnectedHidden = FrontendNode.Style.bUnconnectedPinsHidden;
-			}
+			InPin.bAdvancedView = bAdvancedView;
 
-			if (InPin.bAdvancedView || bUnconnectedHidden)
+			const FMetasoundFrontendNode& FrontendNode = MetasoundGraphNode->GetFrontendNodeChecked();
+			if (InPin.bAdvancedView || FrontendNode.Style.bUnconnectedPinsHidden)
 			{
 				UEdGraphNode* OwningNode = InPin.GetOwningNode();
 				check(OwningNode);
@@ -1755,7 +1751,7 @@ PRAGMA_ENABLE_DEPRECATION_WARNINGS
 			UEdGraphPin* NewPin = InEditorNode.CreatePin(EGPD_Input, PinType, PinName);
 			if (ensure(NewPin))
 			{
-				RefreshPinMetadata(*NewPin, InInputHandle->GetMetadata());
+				RefreshPinMetadata(*NewPin, InInputHandle->GetMetadata().bIsAdvancedDisplay);
 				const FMetaSoundFrontendDocumentBuilder& Builder = InEditorNode.GetBuilderChecked().GetConstBuilder();
 				SynchronizePinLiteral(Builder, *NewPin);
 			}
@@ -1778,7 +1774,7 @@ PRAGMA_ENABLE_DEPRECATION_WARNINGS
 			UEdGraphPin* NewPin = InEditorNode.CreatePin(EGPD_Output, PinType, PinName);
 			if (ensure(NewPin))
 			{
-				RefreshPinMetadata(*NewPin, InOutputHandle->GetMetadata());
+				RefreshPinMetadata(*NewPin, InOutputHandle->GetMetadata().bIsAdvancedDisplay);
 			}
 
 			return NewPin;
diff --git a/Engine/Plugins/Runtime/Metasound/Source/MetasoundEditor/Private/MetasoundEditorGraphBuilder.h b/Engine/Plugins/Runtime/Metasound/Source/MetasoundEditor/Private/MetasoundEditorGraphBuilder.h
--- a/Engine/Plugins/Runtime/Metasound/Source/MetasoundEditor/Private/MetasoundEditorGraphBuilder.h
+++ b/Engine/Plugins/Runtime/Metasound/Source/MetasoundEditor/Private/MetasoundEditorGraphBuilder.h
@@ -225,7 +225,7 @@ namespace Metasound
 			static UEdGraphPin* AddPinToNode(UMetasoundEditorGraphNode& InEditorNode, Frontend::FConstOutputHandle InOutputHandle);
 
 			// Refreshes pin state from class FrontendClassVertexMetadata
-			static void RefreshPinMetadata(UEdGraphPin& InPin, const FMetasoundFrontendVertexMetadata& InMetadata);
+			static void RefreshPinMetadata(UEdGraphPin& InPin, bool bAdvancedView = false);
 
 			// Adds and removes nodes, pins and connections so that the UEdGraph of the MetaSound matches the
 			// FMetasoundFrontendDocument model. Validates the graph (and those referenced recursively).
diff --git a/Engine/Plugins/Runtime/Metasound/Source/MetasoundEditor/Private/MetasoundEditorGraphNode.cpp b/Engine/Plugins/Runtime/Metasound/Source/MetasoundEditor/Private/MetasoundEditorGraphNode.cpp
--- a/Engine/Plugins/Runtime/Metasound/Source/MetasoundEditor/Private/MetasoundEditorGraphNode.cpp
+++ b/Engine/Plugins/Runtime/Metasound/Source/MetasoundEditor/Private/MetasoundEditorGraphNode.cpp
@@ -1277,10 +1277,12 @@ void UMetasoundEditorGraphExternalNode::Validate(Metasound::Editor::FGraphNodeVa
 		break;
 	}
 
+PRAGMA_DISABLE_DEPRECATION_WARNINGS
 	// 2. Check if node is invalid, version is missing and cache if interface changes exist between the document's records and the registry
 	FClassInterfaceUpdates InterfaceUpdates;
 	if (!NodeHandle->DiffAgainstRegistryInterface(InterfaceUpdates, false /* bUseHighestMinorVersion */))
 	{
+PRAGMA_ENABLE_DEPRECATION_WARNINGS
 		if (NodeHandle->IsValid())
 		{
 			const FText* PromptIfMissing = nullptr;
diff --git a/Engine/Plugins/Runtime/Metasound/Source/MetasoundEditor/Private/SMetasoundGraphNode.cpp b/Engine/Plugins/Runtime/Metasound/Source/MetasoundEditor/Private/SMetasoundGraphNode.cpp
--- a/Engine/Plugins/Runtime/Metasound/Source/MetasoundEditor/Private/SMetasoundGraphNode.cpp
+++ b/Engine/Plugins/Runtime/Metasound/Source/MetasoundEditor/Private/SMetasoundGraphNode.cpp
@@ -365,54 +365,69 @@ namespace Metasound
 
 		void SMetaSoundGraphNode::CreateStandardPinWidget(UEdGraphPin* InPin)
 		{
-			//Set pin hidden if the node has unconnected pins hidden
-			const FMetasoundFrontendNode* FrontendNode = GetMetaSoundNode().GetFrontendNode();
+			check(InPin);
+
+			// Set pin hidden if the node has unconnected pins hidden
+			const FGuid NodeID = GetMetaSoundNode().GetNodeID();
+			const FMetaSoundFrontendDocumentBuilder& Builder = GetMetaSoundNode().GetBuilderChecked().GetConstBuilder();
+			const FMetasoundFrontendNode* FrontendNode = Builder.FindNode(NodeID);
+
 			if (FrontendNode)
 			{
-				FMetasoundFrontendNodeStyle Style = FrontendNode->Style;
+				const FMetasoundFrontendNodeStyle& Style = FrontendNode->Style;
 				InPin->SafeSetHidden(Style.bUnconnectedPinsHidden);
 			}
 			
-			if (InPin->Direction == EGPD_Input)
-			{
-				const Frontend::FConstInputHandle InputHandle = FGraphBuilder::GetConstInputHandleFromPin(InPin);
-				if (InputHandle->GetMetadata().bIsAdvancedDisplay != InPin->bAdvancedView)
-				{
-					FGraphBuilder::RefreshPinMetadata(*InPin, InputHandle->GetMetadata());
-				}
-			}
-			else if (InPin->Direction == EGPD_Output)
-			{
-				Frontend::FConstOutputHandle OutputHandle = FGraphBuilder::GetConstOutputHandleFromPin(InPin);
-				if (OutputHandle->GetMetadata().bIsAdvancedDisplay != InPin->bAdvancedView)
-				{				
-					FGraphBuilder::RefreshPinMetadata(*InPin, OutputHandle->GetMetadata());
-				}
-			}
-
-			const bool bShowPin = ShouldPinBeHidden(InPin);
-			if (bShowPin)
+			if (const FMetasoundFrontendClass* FrontendClass = Builder.FindDependency(FrontendNode->ClassID))
 			{
-				TSharedPtr<SGraphPin> NewPin = CreatePinWidget(InPin);
-				check(NewPin.IsValid());
-
-				Frontend::FConstNodeHandle NodeHandle = GetMetaSoundNode().GetConstNodeHandle();
 				if (InPin->Direction == EGPD_Input)
 				{
-					if (!NodeHandle->GetClassStyle().Display.bShowInputNames)
+					const FMetasoundFrontendVertex* InputVertex = FGraphBuilder::GetPinVertex(Builder, InPin);
+					auto IsPinMetadata = [&InputVertex](const FMetasoundFrontendClassVertex& Vertex) { return Vertex.Name == InputVertex->Name; };
+					if (const FMetasoundFrontendClassVertex* ClassVertex = FrontendClass->GetInterfaceForNode(*FrontendNode).Inputs.FindByPredicate(IsPinMetadata))
 					{
-						NewPin->SetShowLabel(false);
+						if (ClassVertex->Metadata.bIsAdvancedDisplay != InPin->bAdvancedView)
+						{
+							FGraphBuilder::RefreshPinMetadata(*InPin, ClassVertex->Metadata.bIsAdvancedDisplay);
+						}
 					}
 				}
 				else if (InPin->Direction == EGPD_Output)
 				{
-					if (!NodeHandle->GetClassStyle().Display.bShowOutputNames)
+					const FMetasoundFrontendVertex* OutputVertex = FGraphBuilder::GetPinVertex(Builder, InPin);
+					auto IsPinMetadata = [&OutputVertex](const FMetasoundFrontendClassVertex& Vertex) { return Vertex.Name == OutputVertex->Name; };
+					if (const FMetasoundFrontendClassVertex* ClassVertex = FrontendClass->GetInterfaceForNode(*FrontendNode).Outputs.FindByPredicate(IsPinMetadata))
 					{
-						NewPin->SetShowLabel(false);
+						if (ClassVertex->Metadata.bIsAdvancedDisplay != InPin->bAdvancedView)
+						{
+							FGraphBuilder::RefreshPinMetadata(*InPin, ClassVertex->Metadata.bIsAdvancedDisplay);
+						}
 					}
 				}
 
-				AddPin(NewPin.ToSharedRef());
+				const bool bShowPin = ShouldPinBeHidden(InPin);
+				if (bShowPin)
+				{
+					TSharedPtr<SGraphPin> NewPin = CreatePinWidget(InPin);
+					check(NewPin.IsValid());
+
+					if (InPin->Direction == EGPD_Input)
+					{
+						if (!FrontendClass->Style.Display.bShowInputNames)
+						{
+							NewPin->SetShowLabel(false);
+						}
+					}
+					else if (InPin->Direction == EGPD_Output)
+					{
+						if (!FrontendClass->Style.Display.bShowOutputNames)
+						{
+							NewPin->SetShowLabel(false);
+						}
+					}
+
+					AddPin(NewPin.ToSharedRef());
+				}
 			}
 		}
 
diff --git a/Engine/Plugins/Runtime/Metasound/Source/MetasoundFrontend/Private/MetasoundFrontendDocument.cpp b/Engine/Plugins/Runtime/Metasound/Source/MetasoundFrontend/Private/MetasoundFrontendDocument.cpp
--- a/Engine/Plugins/Runtime/Metasound/Source/MetasoundFrontend/Private/MetasoundFrontendDocument.cpp
+++ b/Engine/Plugins/Runtime/Metasound/Source/MetasoundFrontend/Private/MetasoundFrontendDocument.cpp
@@ -4,6 +4,7 @@
 
 #include "Algo/AnyOf.h"
 #include "Algo/ForEach.h"
+#include "Algo/Sort.h"
 #include "Algo/Transform.h"
 #include "IAudioParameterInterfaceRegistry.h"
 #include "Logging/LogMacros.h"
@@ -1279,6 +1280,44 @@ FMetasoundFrontendClassEnvironmentVariable::FMetasoundFrontendClassEnvironmentVa
 {
 }
 
+#if WITH_EDITORONLY_DATA
+void FMetasoundFrontendInterfaceStyle::SortVertices(TArray<const FMetasoundFrontendVertex*>& OutVertices, TFunctionRef<FText(const FMetasoundFrontendVertex&)> InGetDisplayNamePredicate) const
+{
+	TMap<FGuid, int32> NodeIDToSortIndex;
+	int32 HighestSortOrder = TNumericLimits<int32>::Min();
+	for (int32 i = 0; i < OutVertices.Num(); ++i)
+	{
+		const FGuid& VertexID = OutVertices[i]->VertexID;
+		int32 SortIndex = 0;
+		if (DefaultSortOrder.IsValidIndex(i))
+		{
+			SortIndex = DefaultSortOrder[i];
+			HighestSortOrder = FMath::Max(SortIndex, HighestSortOrder);
+		}
+		else
+		{
+			SortIndex = ++HighestSortOrder;
+		}
+		NodeIDToSortIndex.Add(VertexID, SortIndex);
+	}
+
+	Algo::Sort(OutVertices, [&](const FMetasoundFrontendVertex* VertexA, const FMetasoundFrontendVertex* VertexB)
+	{
+		const FGuid VertexAID = VertexA->VertexID;
+		const FGuid VertexBID = VertexB->VertexID;
+		const int32 AID = NodeIDToSortIndex[VertexAID];
+		const int32 BID = NodeIDToSortIndex[VertexBID];
+
+		// If IDs are equal, sort alphabetically using provided name predicate
+		if (AID == BID)
+		{
+			return InGetDisplayNamePredicate(*VertexA).CompareTo(InGetDisplayNamePredicate(*VertexB)) < 0;
+		}
+		return AID < BID;
+	});
+}
+#endif // WITH_EDITORONLY_DATA
+
 FMetasoundFrontendGraphClass::FMetasoundFrontendGraphClass()
 {
 	Metadata.SetType(EMetasoundFrontendClassType::Graph);
diff --git a/Engine/Plugins/Runtime/Metasound/Source/MetasoundFrontend/Private/MetasoundFrontendDocumentBuilder.cpp b/Engine/Plugins/Runtime/Metasound/Source/MetasoundFrontend/Private/MetasoundFrontendDocumentBuilder.cpp
--- a/Engine/Plugins/Runtime/Metasound/Source/MetasoundFrontend/Private/MetasoundFrontendDocumentBuilder.cpp
+++ b/Engine/Plugins/Runtime/Metasound/Source/MetasoundFrontend/Private/MetasoundFrontendDocumentBuilder.cpp
@@ -1875,6 +1875,119 @@ const FMetasoundFrontendGraph& FMetaSoundFrontendDocumentBuilder::AddGraphPage(c
 }
 #endif // WITH_EDITORONLY_DATA
 
+#if WITH_EDITOR
+void FMetaSoundFrontendDocumentBuilder::CacheRegistryMetadata()
+{
+	using namespace Metasound::Frontend;
+	using FNameDataTypePair = TPair<FName, FName>;
+
+	FMetasoundFrontendDocument& Document = GetDocumentChecked();
+	FMetasoundFrontendClassInterface& RootGraphClassInterface = Document.RootGraph.GetDefaultInterface();
+
+	// 1. Gather inputs/outputs managed by interfaces
+	TMap<FNameDataTypePair, FMetasoundFrontendClassInput*> Inputs;
+	for (FMetasoundFrontendClassInput& Input : RootGraphClassInterface.Inputs)
+	{
+		Inputs.Add(FNameDataTypePair(Input.Name, Input.TypeName), &Input);
+	}
+
+	TMap<FNameDataTypePair, FMetasoundFrontendClassOutput*> Outputs;
+	for (FMetasoundFrontendClassOutput& Output : RootGraphClassInterface.Outputs)
+	{
+		Outputs.Add(FNameDataTypePair(Output.Name, Output.TypeName), &Output);
+	}
+
+	// 2. Copy metadata for inputs/outputs managed by interfaces, removing them from maps generated
+	auto CacheInterfaceMetadata = [](const FMetasoundFrontendVertexMetadata& InRegistryMetadata, FMetasoundFrontendVertexMetadata& OutMetadata)
+	{
+		const int32 CachedSortOrderIndex = OutMetadata.SortOrderIndex;
+		OutMetadata = InRegistryMetadata;
+		OutMetadata.SortOrderIndex = CachedSortOrderIndex;
+	};
+
+	const TSet<FMetasoundFrontendVersion>& InterfaceVersions = Document.Interfaces;
+	for (const FMetasoundFrontendVersion& Version : InterfaceVersions)
+	{
+		const FInterfaceRegistryKey InterfaceKey = GetInterfaceRegistryKey(Version);
+		const IInterfaceRegistryEntry* Entry = IInterfaceRegistry::Get().FindInterfaceRegistryEntry(InterfaceKey);
+
+		UE_CLOG(nullptr == Entry, LogMetaSound, Error,
+			TEXT("Failed to find interface (%s) when caching registry data for %s. "
+			"MetaSound inputs and outputs for asset may not function correctly."),
+			*Version.ToString(), *GetDebugName());
+
+		if (Entry)
+		{
+			for (const FMetasoundFrontendClassInput& InterfaceInput : Entry->GetInterface().Inputs)
+			{
+				const FNameDataTypePair NameDataTypePair = FNameDataTypePair(InterfaceInput.Name, InterfaceInput.TypeName);
+				if (FMetasoundFrontendClassInput* Input = Inputs.FindRef(NameDataTypePair))
+				{
+					CacheInterfaceMetadata(InterfaceInput.Metadata, Input->Metadata);
+					Inputs.Remove(NameDataTypePair);
+				}
+			}
+
+			for (const FMetasoundFrontendClassOutput& InterfaceOutput : Entry->GetInterface().Outputs)
+			{
+				const FNameDataTypePair NameDataTypePair = FNameDataTypePair(InterfaceOutput.Name, InterfaceOutput.TypeName);
+				if (FMetasoundFrontendClassOutput* Output = Outputs.FindRef(NameDataTypePair))
+				{
+					CacheInterfaceMetadata(InterfaceOutput.Metadata, Output->Metadata);
+					Outputs.Remove(NameDataTypePair);
+				}
+			}
+		}
+	}
+
+	// 3. Iterate remaining inputs/outputs not managed by interfaces and set to serialize text
+	// (in case they were orphaned by an interface no longer being implemented).
+	for (const TPair<FNameDataTypePair, FMetasoundFrontendClassInput*>& Pair : Inputs)
+	{
+		Pair.Value->Metadata.SetSerializeText(true);
+	}
+
+	for (const TPair<FNameDataTypePair, FMetasoundFrontendClassOutput*>& Pair : Outputs)
+	{
+		Pair.Value->Metadata.SetSerializeText(true);
+	}
+
+	// 4. Refresh style as order of members could've changed
+	{
+		FMetasoundFrontendInterfaceStyle InputStyle;
+		Algo::ForEach(RootGraphClassInterface.Inputs, [&InputStyle](const FMetasoundFrontendClassInput& Input)
+		{
+			InputStyle.DefaultSortOrder.Add(Input.Metadata.SortOrderIndex);
+		});
+		RootGraphClassInterface.SetInputStyle(MoveTemp(InputStyle));
+	}
+
+	{
+		FMetasoundFrontendInterfaceStyle OutputStyle;
+		Algo::ForEach(RootGraphClassInterface.Outputs, [&OutputStyle](const FMetasoundFrontendClassOutput& Output)
+		{
+			OutputStyle.DefaultSortOrder.Add(Output.Metadata.SortOrderIndex);
+		});
+		RootGraphClassInterface.SetOutputStyle(MoveTemp(OutputStyle));
+	}
+
+	// 5. Cache registry data on document dependencies
+	for (FMetasoundFrontendClass& Dependency : Document.Dependencies)
+	{
+		if (!FMetasoundFrontendClass::CacheGraphDependencyMetadataFromRegistry(Dependency))
+		{
+			UE_LOG(LogMetaSound, Warning,
+				TEXT("'%s' failed to cache dependency registry data: Registry missing class with key '%s'"),
+				*GetDebugName(),
+				*Dependency.Metadata.GetClassName().ToString());
+			UE_LOG(LogMetaSound, Warning,
+				TEXT("Asset '%s' may fail to build runtime graph unless re-registered after dependency with given key is loaded."),
+				*GetDebugName());
+		}
+	}
+}
+#endif // WITH_EDITOR
+
 bool FMetaSoundFrontendDocumentBuilder::CanAddEdge(const FMetasoundFrontendEdge& InEdge, const FGuid* InPageID) const
 {
 	using namespace Metasound::Frontend;
@@ -2134,6 +2247,100 @@ bool FMetaSoundFrontendDocumentBuilder::ConvertToPreset(const FMetasoundFrontend
 	return false;
 }
 
+#if WITH_EDITORONLY_DATA
+bool FMetaSoundFrontendDocumentBuilder::DiffAgainstRegistryInterface(const FGuid& InNodeID, const FGuid& InPageID, bool bInUseHighestMinorVersion, Metasound::Frontend::FClassInterfaceUpdates& OutInterfaceUpdates) const
+{
+	METASOUND_TRACE_CPUPROFILER_EVENT_SCOPE(FMetaSoundFrontendDocumentBuilder::DiffAgainstRegistryInterface);
+
+	using namespace Metasound::Frontend;
+
+	OutInterfaceUpdates = FClassInterfaceUpdates();
+
+	const FMetasoundFrontendNode* Node = FindNode(InNodeID, &InPageID);
+	if (!Node)
+	{
+		return false;
+	}
+
+	const FMetasoundFrontendClass* Class = FindDependency(Node->ClassID);
+	if (!Class)
+	{
+		return false;
+	}
+
+	const FMetasoundFrontendClassMetadata& NodeClassMetadata = Class->Metadata;
+	const FMetasoundFrontendClassInterface& NodeClassInterface = Class->GetInterfaceForNode(*Node);
+
+	Metasound::FNodeClassName NodeClassName = NodeClassMetadata.GetClassName().ToNodeClassName();
+
+	if (bInUseHighestMinorVersion)
+	{
+		if (!ISearchEngine::Get().FindClassWithHighestMinorVersion(NodeClassName, NodeClassMetadata.GetVersion().Major, OutInterfaceUpdates.RegistryClass))
+		{
+			Algo::Transform(NodeClassInterface.Inputs, OutInterfaceUpdates.RemovedInputs, [&](const FMetasoundFrontendClassInput& Input) { return &Input; });
+			Algo::Transform(NodeClassInterface.Outputs, OutInterfaceUpdates.RemovedOutputs, [&](const FMetasoundFrontendClassOutput& Output) { return &Output; });
+			return false;
+		}
+	}
+	else
+	{
+		// Find class with same metadata in the node registry.
+		FMetasoundFrontendRegistryContainer* Registry = FMetasoundFrontendRegistryContainer::Get();
+		checkf(nullptr != Registry, TEXT("The metasound node registry should always be available if the metasound plugin is loaded"));
+		bool bFoundRegisteredClass = Registry->FindFrontendClassFromRegistered(FNodeRegistryKey(NodeClassMetadata), OutInterfaceUpdates.RegistryClass);
+		if (!bFoundRegisteredClass)
+		{
+			// If the class was not found, mark all inputs and outputs as removed.
+			UE_LOG(LogMetaSound, Warning, TEXT("Could not find registered version of interface. %s %s is not registered."), *NodeClassMetadata.GetClassName().ToString(), *NodeClassMetadata.GetVersion().ToString());
+			Algo::Transform(NodeClassInterface.Inputs, OutInterfaceUpdates.RemovedInputs, [&](const FMetasoundFrontendClassInput& Input) { return &Input; });
+			Algo::Transform(NodeClassInterface.Outputs, OutInterfaceUpdates.RemovedOutputs, [&](const FMetasoundFrontendClassOutput& Output) { return &Output; });
+			return false;
+		}
+	}
+
+	Algo::Transform(OutInterfaceUpdates.RegistryClass.GetInterfaceForNode(*Node).Inputs, OutInterfaceUpdates.AddedInputs, [](const FMetasoundFrontendClassInput& Input) { return &Input; });
+	for (const FMetasoundFrontendClassInput& Input : NodeClassInterface.Inputs)
+	{
+		auto IsEquivalent = [&Input](const FMetasoundFrontendClassInput* RegistryInput)
+		{
+			return FMetasoundFrontendClassInput::IsFunctionalEquivalent(Input, *RegistryInput);
+		};
+
+		const int32 Index = OutInterfaceUpdates.AddedInputs.FindLastByPredicate(IsEquivalent);
+		if (Index == INDEX_NONE)
+		{
+			OutInterfaceUpdates.RemovedInputs.Add(&Input);
+		}
+		else
+		{
+			OutInterfaceUpdates.AddedInputs.RemoveAtSwap(Index, EAllowShrinking::No);
+		}
+	}
+
+
+	Algo::Transform(OutInterfaceUpdates.RegistryClass.GetInterfaceForNode(*Node).Outputs, OutInterfaceUpdates.AddedOutputs, [](const FMetasoundFrontendClassOutput& Output) { return &Output; });
+	for (const FMetasoundFrontendClassOutput& Output : NodeClassInterface.Outputs)
+	{
+		auto IsFunctionalEquivalent = [&Output](const FMetasoundFrontendClassOutput* Iter)
+		{
+			return FMetasoundFrontendClassVertex::IsFunctionalEquivalent(Output, *Iter);
+		};
+
+		const int32 Index = OutInterfaceUpdates.AddedOutputs.FindLastByPredicate(IsFunctionalEquivalent);
+		if (Index == INDEX_NONE)
+		{
+			OutInterfaceUpdates.RemovedOutputs.Add(&Output);
+		}
+		else
+		{
+			OutInterfaceUpdates.AddedOutputs.RemoveAtSwap(Index, EAllowShrinking::No);
+		}
+	}
+
+	return true;
+}
+#endif // WITH_EDITORONLY_DATA
+
 const FMetasoundFrontendNode* FMetaSoundFrontendDocumentBuilder::DuplicateGraphInput(const FMetasoundFrontendClassInput& InClassInput, const FName InName, const FGuid* InPageID)
 {
 	using namespace Metasound;
@@ -2626,6 +2833,46 @@ UMetaSoundFrontendMemberMetadata* FMetaSoundFrontendDocumentBuilder::FindMemberM
 }
 #endif // WITH_EDITOR
 
+TArray<const FMetasoundFrontendVertex*> FMetaSoundFrontendDocumentBuilder::FindUserModifiableNodeInputs(const FGuid& InNodeID, const FGuid* InPageID) const
+{
+	TArray<const FMetasoundFrontendVertex*> Vertices;
+	if (const FMetasoundFrontendNode* Node = FindNode(InNodeID, InPageID))
+	{
+		Algo::TransformIf(Node->Interface.Inputs, Vertices,
+			[this, &InNodeID](const FMetasoundFrontendVertex& Vertex)
+			{
+				return IsNodeInputConnectionUserModifiable(InNodeID, Vertex.VertexID);
+			},
+			[](const FMetasoundFrontendVertex& Vertex)
+			{
+				return &Vertex;
+			}
+		);
+	}
+
+	return Vertices;
+}
+
+TArray<const FMetasoundFrontendVertex*> FMetaSoundFrontendDocumentBuilder::FindUserModifiableNodeOutputs(const FGuid& InNodeID, const FGuid* InPageID) const
+{
+	TArray<const FMetasoundFrontendVertex*> Vertices;
+	if (const FMetasoundFrontendNode* Node = FindNode(InNodeID, InPageID))
+	{
+		Algo::TransformIf(Node->Interface.Outputs, Vertices,
+			[this, &InNodeID](const FMetasoundFrontendVertex& Vertex)
+			{
+				return IsNodeOutputConnectionUserModifiable(InNodeID, Vertex.VertexID);
+			},
+			[](const FMetasoundFrontendVertex& Vertex)
+			{
+				return &Vertex;
+			}
+		);
+	}
+
+	return Vertices;
+}
+
 const FMetasoundFrontendNode* FMetaSoundFrontendDocumentBuilder::FindNode(const FGuid& InNodeID, const FGuid* InPageID) const
 {
 	using namespace Metasound::Frontend;
@@ -2823,9 +3070,6 @@ FMetasoundFrontendNode* FMetaSoundFrontendDocumentBuilder::FindNodeInternal(cons
 {
 	using namespace Metasound::Frontend;
 
-	using namespace Metasound;
-	using namespace Metasound::Frontend;
-
 	const FGuid& PageID = InPageID ? *InPageID : BuildPageID;
 	const IDocumentGraphNodeCache& NodeCache = DocumentCache->GetNodeCache(PageID);
 	if (const int32* NodeIndex = NodeCache.FindNodeIndex(InNodeID))
@@ -3153,6 +3397,182 @@ EMetasoundFrontendVertexAccessType FMetaSoundFrontendDocumentBuilder::GetNodeInp
 	return EMetasoundFrontendVertexAccessType::Unset;
 }
 
+#if WITH_EDITORONLY_DATA
+FText FMetaSoundFrontendDocumentBuilder::GetNodeInputDisplayName(const FGuid& InNodeID, const FName VertexName, const FGuid* InPageID) const
+{
+	using namespace Metasound::Frontend;
+	using namespace Metasound::VariableNames;
+
+	const FGuid& PageID = InPageID ? *InPageID : BuildPageID;
+	const IDocumentGraphNodeCache& NodeCache = DocumentCache->GetNodeCache(PageID);
+	const FMetasoundFrontendNode* Node = NodeCache.FindNode(InNodeID);
+	if (!Node)
+	{
+		return { };
+	}
+
+	const FMetasoundFrontendClass* Class = DocumentCache->FindDependency(Node->ClassID);
+	if (!Class)
+	{
+		return { };
+	}
+
+	const EMetasoundFrontendClassType ClassType = Class->Metadata.GetType();
+	switch (ClassType)
+	{
+		case EMetasoundFrontendClassType::Input:
+		case EMetasoundFrontendClassType::Output:
+		{
+			return FText::FromName(Node->Name);
+		}
+
+		case EMetasoundFrontendClassType::Variable:
+		case EMetasoundFrontendClassType::VariableAccessor:
+		case EMetasoundFrontendClassType::VariableDeferredAccessor:
+		case EMetasoundFrontendClassType::VariableMutator:
+		{
+			return FText::FromName(METASOUND_GET_PARAM_NAME(InputData));
+		}
+
+		case EMetasoundFrontendClassType::Literal:
+		case EMetasoundFrontendClassType::External:
+		{
+			auto MatchesName = [&VertexName](const FMetasoundFrontendClassVertex& OtherVertex) { return OtherVertex.Name == VertexName; };
+			const FMetasoundFrontendVertex* Vertex = nullptr;
+			const FMetasoundFrontendClassVertex* ClassVertex = nullptr;
+			ClassVertex = Class->GetInterfaceForNode(*Node).Inputs.FindByPredicate(MatchesName);
+			if (Vertex && ClassVertex)
+			{
+				FName Namespace, ParamName;
+				ClassVertex->SplitName(Namespace, ParamName);
+				const FText DisplayName = ClassVertex->Metadata.GetDisplayName();
+				if (DisplayName.IsEmptyOrWhitespace())
+				{
+					if (Namespace.IsNone())
+					{
+						return FText::FromName(ParamName);
+					}
+					else
+					{
+						return FText::Format(NSLOCTEXT("MetaSoundFrontend", "ClassMetadataDisplayNameWithNamespaceFormat", "{0} ({1})"), FText::FromName(ParamName), FText::FromName(Namespace));
+					}
+				}
+
+				return DisplayName;
+			}
+		}
+		break;
+
+		case EMetasoundFrontendClassType::Template:
+		{
+			const INodeTemplate* Template = INodeTemplateRegistry::Get().FindTemplate(Class->Metadata.GetClassName());
+			if (ensure(Template))
+			{
+				return Template->GetInputVertexDisplayName(*this, PageID, InNodeID, VertexName);
+			}
+		}
+		break;
+
+		case EMetasoundFrontendClassType::Graph:
+		case EMetasoundFrontendClassType::Invalid:
+		default:
+		{
+			static_assert(static_cast<int32>(EMetasoundFrontendClassType::Invalid) == 10, "Possible missing EMetasoundFrontendClassType case coverage");
+		}
+		break;
+	}
+
+	return FText::FromName(VertexName);
+}
+
+FText FMetaSoundFrontendDocumentBuilder::GetNodeOutputDisplayName(const FGuid& InNodeID, const FName VertexName, const FGuid* InPageID) const
+{
+	using namespace Metasound::Frontend;
+	using namespace Metasound::VariableNames;
+
+	const FGuid& PageID = InPageID ? *InPageID : BuildPageID;
+	const IDocumentGraphNodeCache& NodeCache = DocumentCache->GetNodeCache(PageID);
+	const FMetasoundFrontendNode* Node = NodeCache.FindNode(InNodeID);
+	if (!Node)
+	{
+		return { };
+	}
+
+	const FMetasoundFrontendClass* Class = DocumentCache->FindDependency(Node->ClassID);
+	if (!Class)
+	{
+		return { };
+	}
+
+	const EMetasoundFrontendClassType ClassType = Class->Metadata.GetType();
+	switch (ClassType)
+	{
+		case EMetasoundFrontendClassType::Input:
+		case EMetasoundFrontendClassType::Output:
+		{
+			return FText::FromName(Node->Name);
+		}
+
+		case EMetasoundFrontendClassType::Variable:
+		case EMetasoundFrontendClassType::VariableAccessor:
+		case EMetasoundFrontendClassType::VariableDeferredAccessor:
+		case EMetasoundFrontendClassType::VariableMutator:
+		{
+			return FText::FromName(METASOUND_GET_PARAM_NAME(OutputData));
+		}
+
+		case EMetasoundFrontendClassType::Literal:
+		case EMetasoundFrontendClassType::External:
+		{
+			auto PinMatchesClassVertex = [&VertexName](const FMetasoundFrontendClassVertex& OtherVertex) { return OtherVertex.Name == VertexName; };
+			const FMetasoundFrontendVertex* Vertex = nullptr;
+			const FMetasoundFrontendClassVertex* ClassVertex = nullptr;
+			ClassVertex = Class->GetInterfaceForNode(*Node).Inputs.FindByPredicate(PinMatchesClassVertex);
+			if (Vertex && ClassVertex)
+			{
+				FName Namespace, ParamName;
+				ClassVertex->SplitName(Namespace, ParamName);
+				const FText DisplayName = ClassVertex->Metadata.GetDisplayName();
+				if (DisplayName.IsEmptyOrWhitespace())
+				{
+					if (Namespace.IsNone())
+					{
+						return FText::FromName(ParamName);
+					}
+					else
+					{
+						return FText::Format(NSLOCTEXT("MetaSoundFrontend", "ClassMetadataDisplayNameWithNamespaceFormat", "{0} ({1})"), FText::FromName(ParamName), FText::FromName(Namespace));
+					}
+				}
+
+				return DisplayName;
+			}
+		}
+		break;
+
+		case EMetasoundFrontendClassType::Template:
+		{
+			const INodeTemplate* Template = INodeTemplateRegistry::Get().FindTemplate(Class->Metadata.GetClassName());
+			if (ensure(Template))
+			{
+				return Template->GetOutputVertexDisplayName(*this, PageID, InNodeID, VertexName);
+			}
+		}
+		break;
+
+		case EMetasoundFrontendClassType::Graph:
+		case EMetasoundFrontendClassType::Invalid:
+		default:
+		{
+			static_assert(static_cast<int32>(EMetasoundFrontendClassType::Invalid) == 10, "Possible missing EMetasoundFrontendClassType case coverage");
+		}
+		break;
+	}
+
+	return FText::FromName(VertexName);
+}
+#endif // WITH_EDITORONLY_DATA
+
 const FMetasoundFrontendLiteral* FMetaSoundFrontendDocumentBuilder::GetNodeInputClassDefault(const FGuid& InNodeID, const FGuid& InVertexID, const FGuid* InPageID) const
 {
 	using namespace Metasound;
@@ -3477,12 +3897,112 @@ bool FMetaSoundFrontendDocumentBuilder::IsNodeInputConnected(const FGuid& InNode
 	return DocumentCache->GetEdgeCache(PageID).IsNodeInputConnected(InNodeID, InVertexID);
 }
 
+bool FMetaSoundFrontendDocumentBuilder::IsNodeInputConnectionUserModifiable(const FGuid& InNodeID, const FGuid& InVertexID, const FGuid* InPageID) const
+{
+	using namespace Metasound::Frontend;
+	using namespace Metasound::VariableNames;
+
+	const FGuid& PageID = InPageID ? *InPageID : BuildPageID;
+	const IDocumentGraphNodeCache& NodeCache = DocumentCache->GetNodeCache(PageID);
+	if (const FMetasoundFrontendNode* Node = NodeCache.FindNode(InNodeID))
+	{
+		if (const FMetasoundFrontendClass* Class = FindDependency(Node->ClassID))
+		{
+			switch (Class->Metadata.GetType())
+			{
+				case EMetasoundFrontendClassType::Template:
+				{
+					const FNodeClassRegistryKey Key(Class->Metadata);
+					const INodeTemplate* Template = INodeTemplateRegistry::Get().FindTemplate(Key);
+					if (ensure(Template))
+					{
+						return Template->IsInputConnectionUserModifiable();
+					}
+				}
+				break;
+
+				case EMetasoundFrontendClassType::Input:
+				case EMetasoundFrontendClassType::VariableMutator:
+				{
+					if (const FMetasoundFrontendVertex* Vertex = NodeCache.FindInputVertex(InNodeID, InVertexID))
+					{
+						return Vertex->Name == METASOUND_GET_PARAM_NAME(InputData);
+					}
+					return false;
+				}
+
+				case EMetasoundFrontendClassType::Variable:
+				case EMetasoundFrontendClassType::VariableAccessor:
+				case EMetasoundFrontendClassType::VariableDeferredAccessor:
+				{
+					return false;
+				}
+
+				default:
+				break;
+			}
+		}
+	}
+
+	return true;
+}
+
 bool FMetaSoundFrontendDocumentBuilder::IsNodeOutputConnected(const FGuid& InNodeID, const FGuid& InVertexID, const FGuid* InPageID) const
 {
 	const FGuid& PageID = InPageID ? *InPageID : BuildPageID;
 	return DocumentCache->GetEdgeCache(PageID).IsNodeOutputConnected(InNodeID, InVertexID);
 }
 
+bool FMetaSoundFrontendDocumentBuilder::IsNodeOutputConnectionUserModifiable(const FGuid& InNodeID, const FGuid& InVertexID, const FGuid* InPageID) const
+{
+	using namespace Metasound::Frontend;
+	using namespace Metasound::VariableNames;
+
+	const FGuid& PageID = InPageID ? *InPageID : BuildPageID;
+	const IDocumentGraphNodeCache& NodeCache = DocumentCache->GetNodeCache(PageID);
+	if (const FMetasoundFrontendNode* Node = NodeCache.FindNode(InNodeID))
+	{
+		if (const FMetasoundFrontendClass* Class = FindDependency(Node->ClassID))
+		{
+			switch (Class->Metadata.GetType())
+			{
+				case EMetasoundFrontendClassType::Template:
+				{
+					const FNodeClassRegistryKey Key(Class->Metadata);
+					const INodeTemplate* Template = INodeTemplateRegistry::Get().FindTemplate(Key);
+					if (ensure(Template))
+					{
+						return Template->IsOutputConnectionUserModifiable();
+					}
+				}
+				break;
+
+				case EMetasoundFrontendClassType::Output:
+				case EMetasoundFrontendClassType::Variable:
+				case EMetasoundFrontendClassType::VariableMutator:
+				{
+					return false;
+				}
+
+				case EMetasoundFrontendClassType::VariableAccessor:
+				case EMetasoundFrontendClassType::VariableDeferredAccessor:
+				{
+					if (const FMetasoundFrontendVertex* Vertex = NodeCache.FindOutputVertex(InNodeID, InVertexID))
+					{
+						return Vertex->Name == METASOUND_GET_PARAM_NAME(OutputData);
+					}
+					return false;
+				}
+
+				default:
+					break;
+			}
+		}
+	}
+
+	return true;
+}
+
 bool FMetaSoundFrontendDocumentBuilder::IsInterfaceDeclared(FName InInterfaceName) const
 {
 	using namespace Metasound::Frontend;
@@ -3563,6 +4083,22 @@ void FMetaSoundFrontendDocumentBuilder::IterateNodesConnectedWithVertex(const FM
 	}
 }
 
+void FMetaSoundFrontendDocumentBuilder::IterateNodes(Metasound::Frontend::FConstClassAndNodeFunctionRef Func, const FGuid* InPageID) const
+{
+	using namespace Metasound::Frontend;
+
+	const FGuid& PageID = InPageID ? *InPageID : BuildPageID;
+	const FMetasoundFrontendDocument& Doc = GetConstDocumentChecked();
+	const FMetasoundFrontendGraph& Graph = Doc.RootGraph.FindConstGraphChecked(PageID);
+	for (const FMetasoundFrontendNode& Node : Graph.Nodes)
+	{
+		if (const FMetasoundFrontendClass* Class = FindDependency(Node.ClassID))
+		{
+			Func(*Class, Node);
+		}
+	}
+}
+
 void FMetaSoundFrontendDocumentBuilder::IterateNodesByClassType(Metasound::Frontend::FConstClassAndNodeFunctionRef Func, EMetasoundFrontendClassType ClassType, const FGuid* InPageID) const
 {
 	using namespace Metasound::Frontend;
@@ -5432,9 +5968,9 @@ bool FMetaSoundFrontendDocumentBuilder::SetNodeUnconnectedPinsHidden(const FGuid
 	return false;
 }
 
-const FMetasoundFrontendNodeStyle* FMetaSoundFrontendDocumentBuilder::GetNodeStyle(const FGuid& InNodeID, const FGuid* InPageID)
+const FMetasoundFrontendNodeStyle* FMetaSoundFrontendDocumentBuilder::GetNodeStyle(const FGuid& InNodeID, const FGuid* InPageID) const
 {
-	if (const FMetasoundFrontendNode* Node = FindNodeInternal(InNodeID, InPageID))
+	if (const FMetasoundFrontendNode* Node = FindNode(InNodeID, InPageID))
 	{
 		return &Node->Style;
 	}
@@ -5442,6 +5978,69 @@ const FMetasoundFrontendNodeStyle* FMetaSoundFrontendDocumentBuilder::GetNodeSty
 	return nullptr;
 }
 
+FText FMetaSoundFrontendDocumentBuilder::GetNodeTitle(const FGuid& InNodeID, const FGuid* InPageID) const
+{
+	using namespace Metasound::Frontend;
+
+	if (const FMetasoundFrontendNode* Node = FindNode(InNodeID, InPageID))
+	{
+		if (const FMetasoundFrontendClass* Class = FindDependency(Node->ClassID))
+		{
+			const EMetasoundFrontendClassType ClassType = Class->Metadata.GetType();
+			switch (ClassType)
+			{
+				case EMetasoundFrontendClassType::External:
+				case EMetasoundFrontendClassType::Template:
+				{
+					const FText DisplayName = Class->Metadata.GetDisplayName();
+					if (!DisplayName.IsEmptyOrWhitespace())
+					{
+						return DisplayName;
+					}
+
+					const FTopLevelAssetPath Path = IMetaSoundAssetManager::GetChecked().FindAssetPath(FMetaSoundAssetKey(Class->Metadata));
+					if (Path.IsValid())
+					{
+						return FText::FromName(Path.GetAssetName());
+					}
+
+					return FText::FromName(Class->Metadata.GetClassName().Name);
+				}
+				break;
+
+				case EMetasoundFrontendClassType::Input:
+				{
+					return NSLOCTEXT("MetasoundFrontend", "InputNode_Title", "Input");
+				}
+
+				case EMetasoundFrontendClassType::Output:
+				{
+					return NSLOCTEXT("MetasoundFrontend", "OutputNode_Title", "Output");
+				}
+
+				case EMetasoundFrontendClassType::Variable:
+				case EMetasoundFrontendClassType::VariableDeferredAccessor:
+				case EMetasoundFrontendClassType::VariableAccessor:
+				case EMetasoundFrontendClassType::VariableMutator:
+				{
+					return NSLOCTEXT("MetasoundFrontend", "OutputNode_Variable", "Variable");
+				}
+				break;
+
+				case EMetasoundFrontendClassType::Graph:
+				case EMetasoundFrontendClassType::Invalid:
+				case EMetasoundFrontendClassType::Literal:
+				default:
+				{
+					checkNoEntry();
+				}
+				break;
+			}
+		}
+	}
+
+	return { };
+}
 #endif // WITH_EDITOR
 
 void FMetaSoundFrontendDocumentBuilder::SetVersionNumber(const FMetasoundFrontendVersionNumber& InDocumentVersionNumber)
diff --git a/Engine/Plugins/Runtime/Metasound/Source/MetasoundFrontend/Private/MetasoundFrontendInvalidController.h b/Engine/Plugins/Runtime/Metasound/Source/MetasoundFrontend/Private/MetasoundFrontendInvalidController.h
--- a/Engine/Plugins/Runtime/Metasound/Source/MetasoundFrontend/Private/MetasoundFrontendInvalidController.h
+++ b/Engine/Plugins/Runtime/Metasound/Source/MetasoundFrontend/Private/MetasoundFrontendInvalidController.h
@@ -255,8 +255,10 @@ namespace Metasound
 			virtual const FMetasoundFrontendVersion& GetInterfaceVersion() const { return FMetasoundFrontendVersion::GetInvalid(); }
 			virtual bool IsInterfaceMember() const override { return false; }
 
+#if WITH_EDITORONLY_DATA
 			virtual bool DiffAgainstRegistryInterface(FClassInterfaceUpdates& OutInterfaceUpdates, bool bInUseHighestMinorVersion, bool bInForceRegeneratedClassInterfaceOverride) const override { return false; }
 			virtual bool CanAutoUpdate(FClassInterfaceUpdates& OutInterfaceUpdates) const override { OutInterfaceUpdates = { }; return false; }
+#endif // WITH_EDITORONLY_DATA
 
 			UE_API virtual TSharedRef<IGraphController> AsGraph() override;
 			UE_API virtual TSharedRef<const IGraphController> AsGraph() const override;
@@ -279,8 +281,8 @@ namespace Metasound
 
 			virtual const FVertexName& GetNodeName() const override { return Invalid::GetInvalidName(); }
 
-	   		virtual void SetNodeConfiguration(TInstancedStruct<FMetaSoundFrontendNodeConfiguration> InConfig) override {}
-	   		virtual TConstStructView<FMetaSoundFrontendNodeConfiguration> GetNodeConfiguration() const override { return TConstStructView<FMetaSoundFrontendNodeConfiguration>{}; }
+			virtual void SetNodeConfiguration(TInstancedStruct<FMetaSoundFrontendNodeConfiguration> InConfig) override {}
+			virtual TConstStructView<FMetaSoundFrontendNodeConfiguration> GetNodeConfiguration() const override { return TConstStructView<FMetaSoundFrontendNodeConfiguration>{}; }
 
 #if WITH_EDITOR
 			virtual FText GetDisplayName() const override { return Invalid::GetInvalidText(); }
diff --git a/Engine/Plugins/Runtime/Metasound/Source/MetasoundFrontend/Private/MetasoundFrontendNodeController.cpp b/Engine/Plugins/Runtime/Metasound/Source/MetasoundFrontend/Private/MetasoundFrontendNodeController.cpp
--- a/Engine/Plugins/Runtime/Metasound/Source/MetasoundFrontend/Private/MetasoundFrontendNodeController.cpp
+++ b/Engine/Plugins/Runtime/Metasound/Source/MetasoundFrontend/Private/MetasoundFrontendNodeController.cpp
@@ -1123,6 +1123,7 @@ namespace Metasound
 			return InRegisteredClass.GetDefaultInterface();
 		}
 
+#if WITH_EDITORONLY_DATA
 		void FBaseNodeController::DiffClassInterfaces(const FMetasoundFrontendClassInterface& InOriginalInterface, const FMetasoundFrontendClassInterface& InNewInterface, FClassInterfaceUpdates& OutClassInterfaceUpdates) const
 		{
 			Algo::Transform(InNewInterface.Inputs, OutClassInterfaceUpdates.AddedInputs, [](const FMetasoundFrontendClassInput& Input) { return &Input; });
@@ -1277,7 +1278,10 @@ namespace Metasound
 
 				constexpr bool bUseHighestMinorVersion = true;
 				constexpr bool bForceRegenerateClassInterfaceOverride = true;
+
+PRAGMA_DISABLE_DEPRECATION_WARNINGS
 				DiffAgainstRegistryInterface(OutInterfaceUpdates, bUseHighestMinorVersion, bForceRegenerateClassInterfaceOverride);
+PRAGMA_ENABLE_DEPRECATION_WARNINGS
 
 				return OutInterfaceUpdates.ContainsChanges();
 			}
@@ -1287,9 +1291,12 @@ namespace Metasound
 			// contains no changes).
 			constexpr bool bUseHighestMinorVersion = true;
 			constexpr bool bForceRegenerateClassInterfaceOverride = true;
+PRAGMA_DISABLE_DEPRECATION_WARNINGS
 			DiffAgainstRegistryInterface(OutInterfaceUpdates, bUseHighestMinorVersion, bForceRegenerateClassInterfaceOverride);
+PRAGMA_ENABLE_DEPRECATION_WARNINGS
 			return true;
 		}
+#endif // WITH_EDITORONLY_DATA
 
 		FDocumentAccess FBaseNodeController::ShareAccess()
 		{
diff --git a/Engine/Plugins/Runtime/Metasound/Source/MetasoundFrontend/Private/MetasoundFrontendNodeController.h b/Engine/Plugins/Runtime/Metasound/Source/MetasoundFrontend/Private/MetasoundFrontendNodeController.h
--- a/Engine/Plugins/Runtime/Metasound/Source/MetasoundFrontend/Private/MetasoundFrontendNodeController.h
+++ b/Engine/Plugins/Runtime/Metasound/Source/MetasoundFrontend/Private/MetasoundFrontendNodeController.h
@@ -60,12 +60,16 @@ namespace Metasound
 			virtual void SetNodeStyle(const FMetasoundFrontendNodeStyle& InStyle) override;
 #endif // WITH_EDITORONLY_DATA
 
-	   		virtual void SetNodeConfiguration(TInstancedStruct<FMetaSoundFrontendNodeConfiguration> InConfig) override;
-	   		virtual TConstStructView<FMetaSoundFrontendNodeConfiguration> GetNodeConfiguration() const override;
+			virtual void SetNodeConfiguration(TInstancedStruct<FMetaSoundFrontendNodeConfiguration> InConfig) override;
+			virtual TConstStructView<FMetaSoundFrontendNodeConfiguration> GetNodeConfiguration() const override;
 
+#if WITH_EDITORONLY_DATA
+			UE_DEPRECATED(5.6, "Moved to FMetaSoundFrontendDocumentBuilder")
 			virtual bool DiffAgainstRegistryInterface(FClassInterfaceUpdates& OutInterfaceUpdates, bool bInUseHighestMinorVersion, bool bInForceRegenerateClassInterfaceOverride) const override;
 
 			virtual bool CanAutoUpdate(FClassInterfaceUpdates& OutInterfaceUpdates) const override;
+#endif // WITH_EDITORONLY_DATA
+
 			virtual FNodeHandle ReplaceWithVersion(const FMetasoundFrontendVersionNumber& InNewVersion, TArray<FVertexNameAndType>* OutDisconnectedInputs, TArray<FVertexNameAndType>* OutDisconnectedOutputs) override;
 
 			virtual const FVertexName& GetNodeName() const override;
@@ -195,8 +199,8 @@ namespace Metasound
 			virtual FOutputHandle CreateOutputController(FGuid InVertexID, FConstVertexAccessPtr InNodeVertexPtr, FConstClassOutputAccessPtr InClassOutputPtr, FNodeHandle InOwningNode) const = 0;
 
 		private:
-	   		void SetClassInterfaceOverride(TInstancedStruct<FMetasoundFrontendClassInterface> InClassInterfaceOverride);
-	   		TConstStructView<FMetasoundFrontendClassInterface> GetClassInterfaceOverride() const;
+			void SetClassInterfaceOverride(TInstancedStruct<FMetasoundFrontendClassInterface> InClassInterfaceOverride);
+			TConstStructView<FMetasoundFrontendClassInterface> GetClassInterfaceOverride() const;
 
 			struct FNodeConfigurationUpdateData
 			{
@@ -214,7 +218,9 @@ namespace Metasound
 			const FMetasoundFrontendClassInterface& GetApplicableRegistryInterface(const FMetasoundFrontendClass& InRegisteredClass, const FNodeConfigurationUpdateData& InNodeConfigUpdates) const;
 			void FindNodeConfigurationUpdates(const FMetasoundFrontendClass& InRegisteredClass, FNodeConfigurationUpdateData& OutNodeConfigUpdates, bool bInForceRegenerateClassInterfaceOverride) const;
 
+#if WITH_EDITORONLY_DATA
 			void DiffClassInterfaces(const FMetasoundFrontendClassInterface& InOriginalInterface, const FMetasoundFrontendClassInterface& InNewInterface, FClassInterfaceUpdates& OutClassInterfaceUpdates) const;
+#endif // WITH_EDITORONLY_DATA
 
 			struct FInputConnectionInfo
 			{
diff --git a/Engine/Plugins/Runtime/Metasound/Source/MetasoundFrontend/Private/MetasoundFrontendTransform.cpp b/Engine/Plugins/Runtime/Metasound/Source/MetasoundFrontend/Private/MetasoundFrontendTransform.cpp
--- a/Engine/Plugins/Runtime/Metasound/Source/MetasoundFrontend/Private/MetasoundFrontendTransform.cpp
+++ b/Engine/Plugins/Runtime/Metasound/Source/MetasoundFrontend/Private/MetasoundFrontendTransform.cpp
@@ -520,7 +520,6 @@ namespace Metasound
 				SortAndPlaceMemberNodes(EMetasoundFrontendClassType::Output, AddedNames, GetOutputSortOrder);
 			}
 		}
-#endif // WITH_EDITOR
 
 		bool FAutoUpdateRootGraph::Transform(FDocumentHandle InDocument)
 		{
@@ -585,28 +584,22 @@ namespace Metasound
 					NodesToUpdate.Add(TPair<FNodeHandle, FMetasoundFrontendVersionNumber>(NodeHandle, UpdateVersion));
 				}
 
-				// Only update the node at this point if editor data is loaded. If it isn't and their are no interface
-				// changes but auto-update returned it was eligible, then the auto-update only contains cosmetic changes.
-#if WITH_EDITORONLY_DATA
 				else
 				{
 					NodesToUpdate.Add(TPair<FNodeHandle, FMetasoundFrontendVersionNumber>(NodeHandle, ClassMetadata.GetVersion()));
 				}
-#endif // WITH_EDITORONLY_DATA
 			}, EMetasoundFrontendClassType::External);
 
 			if (PresetReferencedMetaSoundAsset)
 			{
 				if (bIsPreset)
 				{
-#if WITH_EDITORONLY_DATA
 					// This call to FindOrBeginBuilding is a bit of a hack to guarantee instantiation of the referenced graph's builder while
 					// using the deprecated controllers to rebuild the preset root graph.  This is needed in case the referenced graph implements
 					// page data as the builder is required to resolve the correct graph data used to rebuild the preset.  Once the rebuild transform
 					// is migrated to the Builder API, this will no longer be necessary.
 					TScriptInterface<IMetaSoundDocumentInterface> RefDocInterface = PresetReferencedMetaSoundAsset->GetOwningAsset();
 					FMetaSoundFrontendDocumentBuilder& ReferenceBuilder = IDocumentBuilderRegistry::GetChecked().FindOrBeginBuilding(RefDocInterface);
-#endif // WITH_EDITORONLY_DATA
 
 					bDidEdit |= FRebuildPresetRootGraph(PresetReferencedMetaSoundAsset->GetDocumentHandle()).Transform(InDocument);
 					if (bDidEdit)
@@ -668,6 +661,7 @@ namespace Metasound
 
 			return bDidEdit;
 		}
+#endif // WITH_EDITORONLY_DATA
 
 		bool FRebuildPresetRootGraph::Transform(FDocumentHandle InDocument) const
 		{
diff --git a/Engine/Plugins/Runtime/Metasound/Source/MetasoundFrontend/Public/MetasoundFrontendController.h b/Engine/Plugins/Runtime/Metasound/Source/MetasoundFrontend/Public/MetasoundFrontendController.h
--- a/Engine/Plugins/Runtime/Metasound/Source/MetasoundFrontend/Public/MetasoundFrontendController.h
+++ b/Engine/Plugins/Runtime/Metasound/Source/MetasoundFrontend/Public/MetasoundFrontendController.h
@@ -99,50 +99,6 @@ namespace Metasound
 			FDocumentAccessPtr Document;
 		};
 
-
-		/** Provides list of interface members that have been added or removed
-		  * when querying if a node's class has been updated */
-		struct FClassInterfaceUpdates
-		{
-			FClassInterfaceUpdates() = default;
-			TArray<const FMetasoundFrontendClassInput*> AddedInputs;
-			TArray<const FMetasoundFrontendClassOutput*> AddedOutputs;
-			TArray<const FMetasoundFrontendClassInput*> RemovedInputs;
-			TArray<const FMetasoundFrontendClassOutput*> RemovedOutputs;
-			TInstancedStruct<FMetaSoundFrontendNodeConfiguration> AddedConfiguration;
-			TConstStructView<FMetaSoundFrontendNodeConfiguration> RemovedConfiguration;
-			TInstancedStruct<FMetasoundFrontendClassInterface> AddedClassInterfaceOverride;
-			TConstStructView<FMetasoundFrontendClassInterface> RemovedClassInterfaceOverride;
-
-			bool ContainsRemovedMembers() const
-			{
-				return !RemovedInputs.IsEmpty() || !RemovedOutputs.IsEmpty();
-			}
-
-			bool ContainsAddedMembers() const
-			{
-				return !AddedInputs.IsEmpty() || !AddedOutputs.IsEmpty();
-			}
-
-			bool ContainsUpdatedNodeConfig() const
-			{
-				return AddedConfiguration.IsValid() || RemovedConfiguration.IsValid();
-			}
-
-			bool ContainsUpdatedNodeClassInterfaceOverride() const
-			{
-				return AddedClassInterfaceOverride.IsValid() || RemovedClassInterfaceOverride.IsValid();
-			}
-
-			bool ContainsChanges() const
-			{
-				return ContainsRemovedMembers() || ContainsAddedMembers() || ContainsUpdatedNodeConfig() || ContainsUpdatedNodeClassInterfaceOverride();
-			}
-
-			// Cached copy of registry class potentially referenced by added members
-			FMetasoundFrontendClass RegistryClass;
-		};
-
 		/** IDocumentAccessor describes an interface for various I*Controllers to interact with
 		 * each other without exposing that interface publicly or requiring friendship 
 		 * between various controller implementation classes. 
@@ -568,11 +524,8 @@ namespace Metasound
 			virtual const FMetasoundFrontendClassStyle& GetClassStyle() const = 0;
 #endif // WITH_EDITOR
 
-			/**
-			  * Fills out the provided ClassInterfaceUpdate struct with the differences between
-			  * the registry's version of the class interface and that of the node.
-			  * @return Whether or not the interface was found in the registry.
-			  */
+#if WITH_EDITORONLY_DATA
+			UE_DEPRECATED(5.6, "Moved to FMetaSoundFrontendDocumentBuilder")
 			virtual bool DiffAgainstRegistryInterface(FClassInterfaceUpdates& OutInterfaceUpdates, bool bInUseHighestMinorVersion, bool bInForceRegenerateClassInterfaceOverride=true) const = 0;
 
 			/** Returns whether the node is eligible for auto-updating (i.e.
@@ -584,6 +537,7 @@ namespace Metasound
 			  * behavior has been modified.
 			  */
 			virtual bool CanAutoUpdate(FClassInterfaceUpdates& OutInterfaceUpdates) const = 0;
+#endif // WITH_EDITORONLY_DATA
 
 #if WITH_EDITOR
 			/** Description of the given node. */
@@ -608,8 +562,8 @@ namespace Metasound
 			virtual void SetNodeStyle(const FMetasoundFrontendNodeStyle& InStyle) = 0;
 #endif // WITH_EDITORONLY_DATA
 
-	   		virtual void SetNodeConfiguration(TInstancedStruct<FMetaSoundFrontendNodeConfiguration> InConfig) = 0;
-	   		virtual TConstStructView<FMetaSoundFrontendNodeConfiguration> GetNodeConfiguration() const = 0;
+			virtual void SetNodeConfiguration(TInstancedStruct<FMetaSoundFrontendNodeConfiguration> InConfig) = 0;
+			virtual TConstStructView<FMetaSoundFrontendNodeConfiguration> GetNodeConfiguration() const = 0;
 
 			/** Returns the ID associated with this node. */
 			virtual FGuid GetID() const = 0;
diff --git a/Engine/Plugins/Runtime/Metasound/Source/MetasoundFrontend/Public/MetasoundFrontendDocument.h b/Engine/Plugins/Runtime/Metasound/Source/MetasoundFrontend/Public/MetasoundFrontendDocument.h
--- a/Engine/Plugins/Runtime/Metasound/Source/MetasoundFrontend/Public/MetasoundFrontendDocument.h
+++ b/Engine/Plugins/Runtime/Metasound/Source/MetasoundFrontend/Public/MetasoundFrontendDocument.h
@@ -14,6 +14,7 @@
 #include "Misc/Guid.h"
 #include "StructUtils/InstancedStruct.h"
 #include "StructUtils/SharedStruct.h"
+#include "StructUtils/StructView.h"
 #include "Templates/Function.h"
 #include "Templates/Invoke.h"
 #include "Templates/TypeHash.h"
@@ -1003,6 +1004,8 @@ struct FMetasoundFrontendInterfaceStyle
 	UPROPERTY()
 	TMap<FName, FText> RequiredMembers;
 
+	UE_API void SortVertices(TArray<const FMetasoundFrontendVertex*>& OutVertices, TFunctionRef<FText(const FMetasoundFrontendVertex&)> InGetDisplayNamePredicate) const;
+
 	template <typename HandleType, typename NamePredicateType>
 	void SortDefaults(TArray<HandleType>& OutHandles, NamePredicateType InGetDisplayNamePredicate) const
 	{
@@ -2021,6 +2024,52 @@ namespace Metasound::Frontend
 
 	/** Execute the provided function for each literal on a FMetasoundFrontendNode.*/
 	METASOUNDFRONTEND_API void ForEachLiteral(const FMetasoundFrontendNode& InNode, FForEachLiteralFunctionRef OnLiteral);
+
+#if WITH_EDITORONLY_DATA
+	// Provides list of interface members that have been added or removed
+	// when querying if a node's class has been updated. Editor only as 
+	// interface can only be updated from editor builds.
+	struct METASOUNDFRONTEND_API FClassInterfaceUpdates
+	{
+		FClassInterfaceUpdates() = default;
+		TArray<const FMetasoundFrontendClassInput*> AddedInputs;
+		TArray<const FMetasoundFrontendClassOutput*> AddedOutputs;
+		TArray<const FMetasoundFrontendClassInput*> RemovedInputs;
+		TArray<const FMetasoundFrontendClassOutput*> RemovedOutputs;
+		TInstancedStruct<FMetaSoundFrontendNodeConfiguration> AddedConfiguration;
+		TConstStructView<FMetaSoundFrontendNodeConfiguration> RemovedConfiguration;
+		TInstancedStruct<FMetasoundFrontendClassInterface> AddedClassInterfaceOverride;
+		TConstStructView<FMetasoundFrontendClassInterface> RemovedClassInterfaceOverride;
+
+		bool ContainsRemovedMembers() const
+		{
+			return !RemovedInputs.IsEmpty() || !RemovedOutputs.IsEmpty();
+		}
+
+		bool ContainsAddedMembers() const
+		{
+			return !AddedInputs.IsEmpty() || !AddedOutputs.IsEmpty();
+		}
+
+		bool ContainsUpdatedNodeConfig() const
+		{
+			return AddedConfiguration.IsValid() || RemovedConfiguration.IsValid();
+		}
+
+		bool ContainsUpdatedNodeClassInterfaceOverride() const
+		{
+			return AddedClassInterfaceOverride.IsValid() || RemovedClassInterfaceOverride.IsValid();
+		}
+
+		bool ContainsChanges() const
+		{
+			return ContainsRemovedMembers() || ContainsAddedMembers() || ContainsUpdatedNodeConfig() || ContainsUpdatedNodeClassInterfaceOverride();
+		}
+
+		// Cached copy of registry class potentially referenced by added members
+		FMetasoundFrontendClass RegistryClass;
+	};
+#endif // WITH_EDITORONLY_DATA
 } // namespace Metasound::Frontend
 
 #undef UE_API
diff --git a/Engine/Plugins/Runtime/Metasound/Source/MetasoundFrontend/Public/MetasoundFrontendDocumentBuilder.h b/Engine/Plugins/Runtime/Metasound/Source/MetasoundFrontend/Public/MetasoundFrontendDocumentBuilder.h
--- a/Engine/Plugins/Runtime/Metasound/Source/MetasoundFrontend/Public/MetasoundFrontendDocumentBuilder.h
+++ b/Engine/Plugins/Runtime/Metasound/Source/MetasoundFrontend/Public/MetasoundFrontendDocumentBuilder.h
@@ -196,6 +196,12 @@ public:
 	UE_API const FMetasoundFrontendGraph& AddGraphPage(const FGuid& InPageID, bool bDuplicateLastGraph, bool bSetAsBuildGraph = true);
 #endif // WITH_EDITORONLY_DATA
 
+#if WITH_EDITOR
+	// Caches transient dependency metadata & style (class & vertex) found in the
+	// registry that is not necessary for serialization or core graph generation.
+	UE_API void CacheRegistryMetadata();
+#endif // WITH_EDITOR
+
 	// Returns whether or not the given edge can be added, which requires that its input
 	// is not already connected and the edge is valid (see function 'IsValidEdge').
 	UE_API bool CanAddEdge(const FMetasoundFrontendEdge& InEdge, const FGuid* InPageID = nullptr) const;
@@ -219,6 +225,13 @@ public:
 	UE_API bool ConvertFromPreset();
 	UE_API bool ConvertToPreset(const FMetasoundFrontendDocument& InReferencedDocument, TSharedPtr<Metasound::Frontend::FDocumentModifyDelegates> ModifyDelegates = { });
 
+#if WITH_EDITORONLY_DATA
+	// Fills out the provided ClassInterfaceUpdate struct with the differences between
+	// the registry's version of the class interface and that of the node.
+	// Returns whether or not the interface was found in the registry.
+	UE_API bool DiffAgainstRegistryInterface(const FGuid& InNodeID, const FGuid& InPageID, bool bInUseHighestMinorVersion, Metasound::Frontend::FClassInterfaceUpdates& OutInterfaceUpdates) const;
+#endif // WITH_EDITORONLY_DATA
+
 	UE_API const FMetasoundFrontendClassInput* DuplicateGraphInput(FName ExistingName, FName NewName);
 	UE_API const FMetasoundFrontendClassOutput* DuplicateGraphOutput(FName ExistingName, FName NewName);
 	UE_API const FMetasoundFrontendVariable* DuplicateGraphVariable(FName ExistingName, FName NewName, const FGuid* InPageID = nullptr);
@@ -267,6 +280,12 @@ public:
 	// Return the node's index in the document's specified paged graph's node array 
 	UE_API const int32* FindNodeIndex(const FGuid& InNodeID, const FGuid* InPageID = nullptr) const;
 
+	// Returns all inputs on a given node that are user modifiable
+	UE_API TArray<const FMetasoundFrontendVertex*> FindUserModifiableNodeInputs(const FGuid& InNodeID, const FGuid* InPageID = nullptr) const;
+
+	// Returns all outputs on a given node that are user modifiable
+	UE_API TArray<const FMetasoundFrontendVertex*> FindUserModifiableNodeOutputs(const FGuid& InNodeID, const FGuid* InPageID = nullptr) const;
+
 	UE_API const FMetasoundFrontendVertex* FindNodeInput(const FGuid& InNodeID, const FGuid& InVertexID, const FGuid* InPageID = nullptr) const;
 	UE_API const FMetasoundFrontendVertex* FindNodeInput(const FGuid& InNodeID, FName InVertexName, const FGuid* InPageID = nullptr) const;
 
@@ -302,7 +321,15 @@ public:
 
 #if WITH_EDITOR
 	// Gets the editor-only style of a node with the given ID.
-	UE_API const FMetasoundFrontendNodeStyle* GetNodeStyle(const FGuid& InNodeID, const FGuid* InPageID = nullptr);
+	UE_API const FMetasoundFrontendNodeStyle* GetNodeStyle(const FGuid& InNodeID, const FGuid* InPageID = nullptr) const;
+
+	// Title used when displaying the given node. Returned text falls back in the following order:
+	// 1. If input, output, or variable (member), returns member type name.
+	// 2. If external or template node, uses:
+		// a. Node's class display name if set
+		// b. Asset name (if class is defined as an asset)
+		// c. Class name (if class is defined in code)
+	UE_API FText GetNodeTitle(const FGuid& InNodeID, const FGuid* InPageID = nullptr) const;
 #endif // WITH_EDITOR
 
 	template<typename TObjectType>
@@ -340,6 +367,18 @@ public:
 	UE_DEPRECATED(5.5, "Use FindNodeInputClass overloads instead and use GetDefaults() on result (now supports page values)")
 	UE_API const FMetasoundFrontendLiteral* GetNodeInputClassDefault(const FGuid& InNodeID, const FGuid& InVertexID, const FGuid* InPageID = nullptr) const;
 
+#if WITH_EDITORONLY_DATA
+	// Returns friendly name of the given node input vertex. Unlike a straight look-up of the vertex metadata value, this
+	// tries and determines the best display name based on the parent node's class, if no display name is given
+	// converts the base FName to FText, etc.
+	UE_API FText GetNodeInputDisplayName(const FGuid& InNodeID, const FName VertexName, const FGuid* InPageID = nullptr) const;
+
+	// Returns friendly name of the given node output vertex. Unlike a straight look-up of the vertex metadata value, this
+	// tries and determines the best display name based on the parent node's class, if no display name is given
+	// converts the base FName to FText, etc.
+	UE_API FText GetNodeOutputDisplayName(const FGuid& InNodeID, const FName VertexName, const FGuid* InPageID = nullptr) const;
+#endif // WITH_EDITORONLY_DATA
+
 	UE_DEPRECATED(5.5, "Use FindNodeInputDefault and returned struct Value member instead")
 	UE_API const FMetasoundFrontendLiteral* GetNodeInputDefault(const FGuid& InNodeID, const FGuid& InVertexID, const FGuid* InPageID = nullptr) const;
 
@@ -369,7 +408,9 @@ public:
 
 	UE_API bool IsDependencyReferenced(const FGuid& InClassID) const;
 	UE_API bool IsNodeInputConnected(const FGuid& InNodeID, const FGuid& InVertexID, const FGuid* InPageID = nullptr) const;
+	UE_API bool IsNodeInputConnectionUserModifiable(const FGuid& InNodeID, const FGuid& InVertexID, const FGuid* InPageID = nullptr) const;
 	UE_API bool IsNodeOutputConnected(const FGuid& InNodeID, const FGuid& InVertexID, const FGuid* InPageID = nullptr) const;
+	UE_API bool IsNodeOutputConnectionUserModifiable(const FGuid& InNodeID, const FGuid& InVertexID, const FGuid* InPageID = nullptr) const;
 
 	UE_API bool IsInterfaceDeclared(FName InInterfaceName) const;
 	UE_API bool IsInterfaceDeclared(const FMetasoundFrontendVersion& InInterfaceVersion) const;
@@ -382,6 +423,9 @@ public:
 	// Note that this does not return whether or not the given edge exists, but rather if it could be legally applied to the given edge vertices.
 	UE_API Metasound::Frontend::EInvalidEdgeReason IsValidEdge(const FMetasoundFrontendEdge& InEdge, const FGuid* InPageID = nullptr) const;
 
+	// Iterates all nodes for the given page (uses build page ID if not provided).
+	UE_API void IterateNodes(Metasound::Frontend::FConstClassAndNodeFunctionRef Func, const FGuid* InPageID = nullptr) const;
+
 	// Iterates nodes that are filtered by only subscribing to a class with the given type (asserts if provided invalid class type).
 	UE_API void IterateNodesByClassType(Metasound::Frontend::FConstClassAndNodeFunctionRef Func, EMetasoundFrontendClassType ClassType, const FGuid* InPageID = nullptr) const;
 
@@ -614,42 +658,45 @@ public:
 private:
 	using FFinalizeNodeFunctionRef = TFunctionRef<void(FMetasoundFrontendNode&, const Metasound::Frontend::FNodeRegistryKey&)>;
 
-	UE_API FMetasoundFrontendNode* AddNodeInternal(const FMetasoundFrontendClassMetadata& InClassMetadata, Metasound::Frontend::FFinalizeNodeFunctionRef FinalizeNode, const FGuid& InPageID, FGuid InNodeID = FGuid::NewGuid(), int32* NewNodeIndex = nullptr);
-	UE_API FMetasoundFrontendNode* AddNodeInternal(const Metasound::Frontend::FNodeRegistryKey& InClassKey, FFinalizeNodeFunctionRef FinalizeNode, const FGuid& InPageID, FGuid InNodeID = FGuid::NewGuid(), int32* NewNodeIndex = nullptr);
-	UE_API void BeginBuilding(TSharedPtr<Metasound::Frontend::FDocumentModifyDelegates> Delegates = {}, bool bPrimeCache = false);
+	FMetasoundFrontendNode* AddNodeInternal(const FMetasoundFrontendClassMetadata& InClassMetadata, Metasound::Frontend::FFinalizeNodeFunctionRef FinalizeNode, const FGuid& InPageID, FGuid InNodeID = FGuid::NewGuid(), int32* NewNodeIndex = nullptr);
+	FMetasoundFrontendNode* AddNodeInternal(const Metasound::Frontend::FNodeRegistryKey& InClassKey, FFinalizeNodeFunctionRef FinalizeNode, const FGuid& InPageID, FGuid InNodeID = FGuid::NewGuid(), int32* NewNodeIndex = nullptr);
+	void BeginBuilding(TSharedPtr<Metasound::Frontend::FDocumentModifyDelegates> Delegates = {}, bool bPrimeCache = false);
 
 	// Conforms GraphOutput node's ClassID, Access & Data Type with the GraphOutput.
-	// creating and removing dependencies as necessary within the document dependency array. Does *NOT*
+	// Creates and removes dependencies as necessary within the document dependency array. Does *NOT*
 	// modify edge data (i.e. if the DataType is changed on the given node and it has corresponding
 	// edges, edges may then be invalid due to access type/DataType incompatibility).
-	UE_API bool ConformGraphInputNodeToClass(const FMetasoundFrontendClassInput& GraphInput);
+	bool ConformGraphInputNodeToClass(const FMetasoundFrontendClassInput& GraphInput);
 
 	// Conforms GraphOutput node's ClassID, Access & Data Type with the GraphOutput.
 	// Creates and removes dependencies as necessary within the document dependency array. Does *NOT*
 	// modify edge data (i.e. if the DataType is changed on the given node and it has corresponding
 	// edges, edges may then be invalid due to access type/DataType incompatibility).
-	UE_API bool ConformGraphOutputNodeToClass(const FMetasoundFrontendClassOutput& GraphOutput);
+	bool ConformGraphOutputNodeToClass(const FMetasoundFrontendClassOutput& GraphOutput);
+
+	FMetasoundFrontendGraph& FindBuildGraphChecked() const;
 
-	UE_API FMetasoundFrontendGraph& FindBuildGraphChecked() const;
+	FMetasoundFrontendVariable* FindGraphVariableInternal(FName InVariableName, const FGuid* InPageID = nullptr);
 
-	UE_API FMetasoundFrontendVariable* FindGraphVariableInternal(FName InVariableName, const FGuid* InPageID = nullptr);
+	bool FindNodeClassInterfaces(const FGuid& InNodeID, TSet<FMetasoundFrontendVersion>& OutInterfaces, const FGuid& InPageID) const;
 
-	UE_API bool FindNodeClassInterfaces(const FGuid& InNodeID, TSet<FMetasoundFrontendVersion>& OutInterfaces, const FGuid& InPageID) const;
-	UE_API FMetasoundFrontendNode* FindNodeInternal(const FGuid& InNodeID, const FGuid* InPageID = nullptr);
+	// Returns pointer to editable node.  Modifying data on the node directly can be inherently dangerous if associated node cache
+	// data is not updated to reflect any changes to the given node (ex. changing NodeID, interface, etc.)
+	FMetasoundFrontendNode* FindNodeInternal(const FGuid& InNodeID, const FGuid* InPageID = nullptr);
 
-	UE_API const FMetasoundFrontendNode* FindHeadNodeInVariableStack(FName VariableName, const FGuid* InPageID = nullptr) const;
-	UE_API const FMetasoundFrontendNode* FindTailNodeInVariableStack(FName VariableName, const FGuid* InPageID = nullptr) const;
+	const FMetasoundFrontendNode* FindHeadNodeInVariableStack(FName VariableName, const FGuid* InPageID = nullptr) const;
+	const FMetasoundFrontendNode* FindTailNodeInVariableStack(FName VariableName, const FGuid* InPageID = nullptr) const;
 
-	UE_API void IterateNodesConnectedWithVertex(const FMetasoundFrontendVertexHandle& Vertex, TFunctionRef<void(const FMetasoundFrontendEdge&, FMetasoundFrontendNode&)> NodeIndexIterFunc, const FGuid& InPageID);
+	void IterateNodesConnectedWithVertex(const FMetasoundFrontendVertexHandle& Vertex, TFunctionRef<void(const FMetasoundFrontendEdge&, FMetasoundFrontendNode&)> NodeIndexIterFunc, const FGuid& InPageID);
 
-	UE_API const FTopLevelAssetPath GetBuilderClassPath() const;
-	UE_API FMetasoundFrontendDocument& GetDocumentChecked() const;
-	UE_API IMetaSoundDocumentInterface& GetDocumentInterfaceChecked() const;
+	const FTopLevelAssetPath GetBuilderClassPath() const;
+	FMetasoundFrontendDocument& GetDocumentChecked() const;
+	IMetaSoundDocumentInterface& GetDocumentInterfaceChecked() const;
 
-	UE_API void RemoveSwapDependencyInternal(int32 Index);
+	void RemoveSwapDependencyInternal(int32 Index);
 
-	UE_API bool SpliceVariableNodeFromStack(const FGuid& InNodeID, const FGuid& InPageID);
-	UE_API bool UnlinkVariableNode(const FGuid& InNodeID, const FGuid& InPageID);
+	bool SpliceVariableNodeFromStack(const FGuid& InNodeID, const FGuid& InPageID);
+	bool UnlinkVariableNode(const FGuid& InNodeID, const FGuid& InPageID);
 
 	UPROPERTY(Transient)
 	TScriptInterface<IMetaSoundDocumentInterface> DocumentInterface;
diff --git a/Engine/Plugins/Runtime/Metasound/Source/MetasoundFrontend/Public/MetasoundFrontendTransform.h b/Engine/Plugins/Runtime/Metasound/Source/MetasoundFrontend/Public/MetasoundFrontendTransform.h
--- a/Engine/Plugins/Runtime/Metasound/Source/MetasoundFrontend/Public/MetasoundFrontendTransform.h
+++ b/Engine/Plugins/Runtime/Metasound/Source/MetasoundFrontend/Public/MetasoundFrontendTransform.h
@@ -189,6 +189,7 @@ namespace Metasound::Frontend
 		FConstDocumentHandle ReferencedDocument = IDocumentController::GetInvalidHandle();
 	};
 
+#if WITH_EDITORONLY_DATA
 	/** Automatically updates all nodes and respective dependencies in graph where
 		* newer versions exist in the loaded MetaSound Class Node Registry.
 		*/
@@ -217,6 +218,7 @@ namespace Metasound::Frontend
 		const FString DebugAssetPath;
 		bool bLogWarningOnDroppedConnection;
 	};
+#endif // WITH_EDITORONLY_DATA
 
 	/** Sets the document's graph class, optionally updating the namespace and variant. */
 	class FRenameRootGraphClass : public IDocumentTransform
