Fix for static analysis
#rnx

diff --git a/Engine/Plugins/Runtime/Metasound/Source/MetasoundEditor/Private/SMetasoundGraphNode.cpp b/Engine/Plugins/Runtime/Metasound/Source/MetasoundEditor/Private/SMetasoundGraphNode.cpp
--- a/Engine/Plugins/Runtime/Metasound/Source/MetasoundEditor/Private/SMetasoundGraphNode.cpp
+++ b/Engine/Plugins/Runtime/Metasound/Source/MetasoundEditor/Private/SMetasoundGraphNode.cpp
@@ -370,63 +370,61 @@ namespace Metasound
 			// Set pin hidden if the node has unconnected pins hidden
 			const FGuid NodeID = GetMetaSoundNode().GetNodeID();
 			const FMetaSoundFrontendDocumentBuilder& Builder = GetMetaSoundNode().GetBuilderChecked().GetConstBuilder();
-			const FMetasoundFrontendNode* FrontendNode = Builder.FindNode(NodeID);
-
-			if (FrontendNode)
+			if (const FMetasoundFrontendNode* FrontendNode = Builder.FindNode(NodeID))
 			{
 				const FMetasoundFrontendNodeStyle& Style = FrontendNode->Style;
 				InPin->SafeSetHidden(Style.bUnconnectedPinsHidden);
-			}
-			
-			if (const FMetasoundFrontendClass* FrontendClass = Builder.FindDependency(FrontendNode->ClassID))
-			{
-				if (InPin->Direction == EGPD_Input)
+
+				if (const FMetasoundFrontendClass* FrontendClass = Builder.FindDependency(FrontendNode->ClassID))
 				{
-					const FMetasoundFrontendVertex* InputVertex = FGraphBuilder::GetPinVertex(Builder, InPin);
-					auto IsPinMetadata = [&InputVertex](const FMetasoundFrontendClassVertex& Vertex) { return Vertex.Name == InputVertex->Name; };
-					if (const FMetasoundFrontendClassVertex* ClassVertex = FrontendClass->GetInterfaceForNode(*FrontendNode).Inputs.FindByPredicate(IsPinMetadata))
+					if (InPin->Direction == EGPD_Input)
 					{
-						if (ClassVertex->Metadata.bIsAdvancedDisplay != InPin->bAdvancedView)
+						const FMetasoundFrontendVertex* InputVertex = FGraphBuilder::GetPinVertex(Builder, InPin);
+						auto IsPinMetadata = [&InputVertex](const FMetasoundFrontendClassVertex& Vertex) { return Vertex.Name == InputVertex->Name; };
+						if (const FMetasoundFrontendClassVertex* ClassVertex = FrontendClass->GetInterfaceForNode(*FrontendNode).Inputs.FindByPredicate(IsPinMetadata))
 						{
-							FGraphBuilder::RefreshPinMetadata(*InPin, ClassVertex->Metadata.bIsAdvancedDisplay);
+							if (ClassVertex->Metadata.bIsAdvancedDisplay != InPin->bAdvancedView)
+							{
+								FGraphBuilder::RefreshPinMetadata(*InPin, ClassVertex->Metadata.bIsAdvancedDisplay);
+							}
 						}
 					}
-				}
-				else if (InPin->Direction == EGPD_Output)
-				{
-					const FMetasoundFrontendVertex* OutputVertex = FGraphBuilder::GetPinVertex(Builder, InPin);
-					auto IsPinMetadata = [&OutputVertex](const FMetasoundFrontendClassVertex& Vertex) { return Vertex.Name == OutputVertex->Name; };
-					if (const FMetasoundFrontendClassVertex* ClassVertex = FrontendClass->GetInterfaceForNode(*FrontendNode).Outputs.FindByPredicate(IsPinMetadata))
+					else if (InPin->Direction == EGPD_Output)
 					{
-						if (ClassVertex->Metadata.bIsAdvancedDisplay != InPin->bAdvancedView)
+						const FMetasoundFrontendVertex* OutputVertex = FGraphBuilder::GetPinVertex(Builder, InPin);
+						auto IsPinMetadata = [&OutputVertex](const FMetasoundFrontendClassVertex& Vertex) { return Vertex.Name == OutputVertex->Name; };
+						if (const FMetasoundFrontendClassVertex* ClassVertex = FrontendClass->GetInterfaceForNode(*FrontendNode).Outputs.FindByPredicate(IsPinMetadata))
 						{
-							FGraphBuilder::RefreshPinMetadata(*InPin, ClassVertex->Metadata.bIsAdvancedDisplay);
+							if (ClassVertex->Metadata.bIsAdvancedDisplay != InPin->bAdvancedView)
+							{
+								FGraphBuilder::RefreshPinMetadata(*InPin, ClassVertex->Metadata.bIsAdvancedDisplay);
+							}
 						}
 					}
-				}
 
-				const bool bShowPin = ShouldPinBeHidden(InPin);
-				if (bShowPin)
-				{
-					TSharedPtr<SGraphPin> NewPin = CreatePinWidget(InPin);
-					check(NewPin.IsValid());
-
-					if (InPin->Direction == EGPD_Input)
+					const bool bShowPin = ShouldPinBeHidden(InPin);
+					if (bShowPin)
 					{
-						if (!FrontendClass->Style.Display.bShowInputNames)
+						TSharedPtr<SGraphPin> NewPin = CreatePinWidget(InPin);
+						check(NewPin.IsValid());
+
+						if (InPin->Direction == EGPD_Input)
 						{
-							NewPin->SetShowLabel(false);
+							if (!FrontendClass->Style.Display.bShowInputNames)
+							{
+								NewPin->SetShowLabel(false);
+							}
 						}
-					}
-					else if (InPin->Direction == EGPD_Output)
-					{
-						if (!FrontendClass->Style.Display.bShowOutputNames)
+						else if (InPin->Direction == EGPD_Output)
 						{
-							NewPin->SetShowLabel(false);
+							if (!FrontendClass->Style.Display.bShowOutputNames)
+							{
+								NewPin->SetShowLabel(false);
+							}
 						}
-					}
 
-					AddPin(NewPin.ToSharedRef());
+						AddPin(NewPin.ToSharedRef());
+					}
 				}
 			}
 		}
