Fix crashes when using MetaSound node configuration with pages

diff --git a/Engine/Plugins/Runtime/Metasound/Source/MetasoundEditor/Private/MetasoundDetailCustomization.cpp b/Engine/Plugins/Runtime/Metasound/Source/MetasoundEditor/Private/MetasoundDetailCustomization.cpp
--- a/Engine/Plugins/Runtime/Metasound/Source/MetasoundEditor/Private/MetasoundDetailCustomization.cpp
+++ b/Engine/Plugins/Runtime/Metasound/Source/MetasoundEditor/Private/MetasoundDetailCustomization.cpp
@@ -404,6 +404,13 @@ namespace Metasound::Editor
 
 					EMetaSoundBuilderResult Result = EMetaSoundBuilderResult::Failed;
 					Builder->AddGraphPage(FName(*NameToAdd.Get()), bDuplicateLastGraph, bSetAsBuildGraph, Result);
+
+					IAssetEditorInstance* AssetEditor = GEditor->GetEditorSubsystem<UAssetEditorSubsystem>()->FindEditorForAsset(&GetMetaSound(), /*bFocusIfOpen=*/false);
+					FEditor* MetaSoundEditor = static_cast<FEditor*>(AssetEditor);
+					if (MetaSoundEditor)
+					{
+						MetaSoundEditor->RefreshDetails();
+					}
 				}
 			})
 			.Content()
@@ -498,6 +505,12 @@ namespace Metasound::Editor
 					{
 						constexpr bool bOpenEditor = false; // Already focused by user action
 						UMetaSoundEditorSubsystem::GetConstChecked().SetFocusedPage(*Builder.Get(), PageID, bOpenEditor);
+						IAssetEditorInstance* AssetEditor = GEditor->GetEditorSubsystem<UAssetEditorSubsystem>()->FindEditorForAsset(&GetMetaSound(), /*bFocusIfOpen=*/false);
+						FEditor* MetaSoundEditor = static_cast<FEditor*>(AssetEditor);
+						if (MetaSoundEditor)
+						{
+							MetaSoundEditor->RefreshDetails();
+						}
 						BuildPageName = InName;
 					}),
 					TAttribute<FText>::Create([this, InName]()
diff --git a/Engine/Plugins/Runtime/Metasound/Source/MetasoundEditor/Private/MetasoundEditorGraphBuilder.cpp b/Engine/Plugins/Runtime/Metasound/Source/MetasoundEditor/Private/MetasoundEditorGraphBuilder.cpp
--- a/Engine/Plugins/Runtime/Metasound/Source/MetasoundEditor/Private/MetasoundEditorGraphBuilder.cpp
+++ b/Engine/Plugins/Runtime/Metasound/Source/MetasoundEditor/Private/MetasoundEditorGraphBuilder.cpp
@@ -2344,7 +2344,17 @@ PRAGMA_ENABLE_DEPRECATION_WARNINGS
 			{
 				for (UMetasoundEditorGraphNode* EditorNode : IdNodePair.Value.EditorNodes)
 				{
-					SynchronizeNodePins(*EditorNode, IdNodePair.Value.Node);
+					const bool bNodeDirtied = SynchronizeNodePins(*EditorNode, IdNodePair.Value.Node);
+					if (bNodeDirtied)
+					{
+						if (TSharedPtr<FEditor> MetaSoundEditor = FGraphBuilder::GetEditorForMetasound(MetaSoundObject))
+						{
+							if (TSharedPtr<SGraphEditor> GraphEditor = MetaSoundEditor->GetGraphEditor())
+							{
+								GraphEditor->RefreshNode(*EditorNode);
+							}
+						}
+					}
 				}
 			}
 		}
