Allow partial copies of final output buffer based on On Finished trigger frame index
#swarm https://p4-swarm.epicgames.net/reviews/41887656

diff --git a/Engine/Plugins/Runtime/Metasound/Source/MetasoundGenerator/Private/MetasoundGenerator.cpp b/Engine/Plugins/Runtime/Metasound/Source/MetasoundGenerator/Private/MetasoundGenerator.cpp
--- a/Engine/Plugins/Runtime/Metasound/Source/MetasoundGenerator/Private/MetasoundGenerator.cpp
+++ b/Engine/Plugins/Runtime/Metasound/Source/MetasoundGenerator/Private/MetasoundGenerator.cpp
@@ -757,28 +757,36 @@ PRAGMA_ENABLE_DEPRECATION_WARNINGS
 		{
 			return 0;
 		}
+
+		const int32 NumSamplesToCopy = FMath::Min(InNum, MaxNumOutputSamples);
 		
+		if (bIsFinishTriggered)
+		{
+		    FMemory::Memcpy(OutAudio, InBuffer.GetData(), MaxNumOutputSamples * sizeof(float));
+			return NumSamplesToCopy;
+		}
 		
-		int32 NumSamplesToCopy = FMath::Min(InNum, MaxNumOutputSamples);
 		
 		if (FinishSample != INDEX_NONE)
 		{
-			FinishSample = NumSamplesToCopy; // maxdebug, remove this. forcing edge case
 			if (FinishSample > NumSamplesToCopy)
 			{
 				FinishSample -= NumSamplesToCopy;
 			}
 			else
 			{
-				NumSamplesToCopy = FinishSample;
 				bIsFinishTriggered = true;
-				FinishSample = INDEX_NONE;
 			}
 		}
 		
-		FMemory::Memcpy(OutAudio, InBuffer.GetData(), NumSamplesToCopy * sizeof(float));
-
-		return NumSamplesToCopy;
+		const int32 NumFramesToWrite = bIsFinishTriggered? FinishSample : NumSamplesToCopy;
+		const int32 NumFramesToZero = MaxNumOutputSamples - NumFramesToWrite;
+		FinishSample = bIsFinishTriggered? INDEX_NONE : FinishSample;
+		
+		FMemory::Memcpy(OutAudio, InBuffer.GetData(), NumFramesToWrite * sizeof(float));
+		FMemory::Memzero(OutAudio, NumFramesToZero * sizeof(float));
+		
+		return NumFramesToWrite;
 	}
 
 	void FMetasoundGenerator::InterleaveGeneratedAudio()
