- Do not remove template input nodes when swapping inputs
- Minor clean-up moving WITH_EDITOR to WITH_EDITORONLY_DATA

#rnx

diff --git a/Engine/Plugins/Runtime/Metasound/Source/MetasoundFrontend/Private/MetasoundFrontendDocumentBuilder.cpp b/Engine/Plugins/Runtime/Metasound/Source/MetasoundFrontend/Private/MetasoundFrontendDocumentBuilder.cpp
--- a/Engine/Plugins/Runtime/Metasound/Source/MetasoundFrontend/Private/MetasoundFrontendDocumentBuilder.cpp
+++ b/Engine/Plugins/Runtime/Metasound/Source/MetasoundFrontend/Private/MetasoundFrontendDocumentBuilder.cpp
@@ -4684,7 +4684,7 @@ bool FMetaSoundFrontendDocumentBuilder::RemoveGraphComment(const FGuid& InCommen
 }
 #endif // WITH_EDITORONLY_DATA
 
-bool FMetaSoundFrontendDocumentBuilder::RemoveGraphInput(FName InputName)
+bool FMetaSoundFrontendDocumentBuilder::RemoveGraphInput(FName InputName, bool bRemoveTemplateInputNodes)
 {
 	using namespace Metasound::Frontend;
 
@@ -4695,7 +4695,7 @@ bool FMetaSoundFrontendDocumentBuilder::RemoveGraphInput(FName InputName)
 		const FGuid NodeID = Inputs[*IndexPtr].NodeID;
 		FGuid ClassID;
 		bool bNodesRemoved = true;
-		Document.RootGraph.IterateGraphPages([this, &Document, &ClassID, &NodeID, &InputName, &bNodesRemoved](const FMetasoundFrontendGraph& Graph)
+		Document.RootGraph.IterateGraphPages([&](const FMetasoundFrontendGraph& Graph)
 		{
 			TArray<FGuid> NodeIDsToRemove { NodeID };
 
@@ -4709,8 +4709,11 @@ bool FMetaSoundFrontendDocumentBuilder::RemoveGraphInput(FName InputName)
 				return;
 			}
 
-			const TArray<const FMetasoundFrontendNode*> TemplateNodes = GetGraphInputTemplateNodes(InputName, &Graph.PageID);
-			Algo::Transform(TemplateNodes, NodeIDsToRemove, [](const FMetasoundFrontendNode* Node) { return Node->GetID(); });
+			if (bRemoveTemplateInputNodes)
+			{
+				const TArray<const FMetasoundFrontendNode*> TemplateNodes = GetGraphInputTemplateNodes(InputName, &Graph.PageID);
+				Algo::Transform(TemplateNodes, NodeIDsToRemove, [](const FMetasoundFrontendNode* Node) { return Node->GetID(); });
+			}
 
 			for (const FGuid& ToRemove : NodeIDsToRemove)
 			{
@@ -6111,10 +6114,10 @@ bool FMetaSoundFrontendDocumentBuilder::SwapGraphInput(const FMetasoundFrontendC
 
 	const IDocumentGraphInterfaceCache& InterfaceCache = DocumentCache->GetInterfaceCache();
 
-#if WITH_EDITOR
+#if WITH_EDITORONLY_DATA
 	using FPageNodeLocations = TMap<FGuid, FVector2D>;
 	TMap<FGuid, FPageNodeLocations> PageNodeLocations;
-#endif // WITH_EDITOR
+#endif // WITH_EDITORONLY_DATA
 
 	// 2. Gather data from existing member/node needed to swap
 	TMultiMap<FGuid, FMetasoundFrontendEdge> RemovedEdgesPerPage;
@@ -6131,9 +6134,9 @@ bool FMetaSoundFrontendDocumentBuilder::SwapGraphInput(const FMetasoundFrontendC
 		const FMetasoundFrontendNode* ExistingInputNode = NodeCache.FindNode(NodeID);
 		check(ExistingInputNode);
 
-#if WITH_EDITOR
+#if WITH_EDITORONLY_DATA
 		PageNodeLocations.Add(Graph.PageID, ExistingInputNode->Style.Display.Locations);
-#endif // WITH_EDITOR
+#endif // WITH_EDITORONLY_DATA
 
 		const FGuid VertexID = ExistingInputNode->Interface.Outputs.Last().VertexID;
 		TArray<const FMetasoundFrontendEdge*> Edges = DocumentCache->GetEdgeCache(Graph.PageID).FindEdges(NodeID, VertexID);
@@ -6145,7 +6148,12 @@ bool FMetaSoundFrontendDocumentBuilder::SwapGraphInput(const FMetasoundFrontendC
 
 	// 3. Remove existing graph vertex
 	{
-		const bool bRemovedVertex = RemoveGraphInput(InExistingInputVertex.Name);
+		// Access & Data Types will be preserved, so no reason to remove template nodes.
+		// (Removal can additionally cause associated edges to be removed rendering cached
+		// RemovedEdgesPerPage above to be stale, so leaving template input nodes in place
+		// preserves that data's validity.
+		constexpr bool bRemoveTemplateInputNodes = false;
+		const bool bRemovedVertex = RemoveGraphInput(InExistingInputVertex.Name, bRemoveTemplateInputNodes);
 		checkf(bRemovedVertex, TEXT("Failed to swap MetaSound input expected to exist"));
 	}
 
@@ -6154,9 +6162,9 @@ bool FMetaSoundFrontendDocumentBuilder::SwapGraphInput(const FMetasoundFrontendC
 	{
 		FMetasoundFrontendClassInput NewInput = InNewInputVertex;
 		NewInput.NodeID = NodeID;
-#if WITH_EDITOR
+#if WITH_EDITORONLY_DATA
 		NewInput.Metadata.SetSerializeText(InExistingInputVertex.Metadata.GetSerializeText());
-#endif // WITH_EDITOR
+#endif // WITH_EDITORONLY_DATA
 
 		const FMetasoundFrontendNode* NewInputNode = AddGraphInput(MoveTemp(NewInput));
 		checkf(NewInputNode, TEXT("Failed to add new Input node when swapping graph inputs"));
@@ -6166,7 +6174,7 @@ bool FMetaSoundFrontendDocumentBuilder::SwapGraphInput(const FMetasoundFrontendC
 
 	Document.RootGraph.IterateGraphPages([&](FMetasoundFrontendGraph& Graph)
 	{
-#if WITH_EDITOR
+#if WITH_EDITORONLY_DATA
 		// 5a. Add to new copy existing node locations
 		if (const FPageNodeLocations* Locations = PageNodeLocations.Find(Graph.PageID))
 		{
@@ -6176,7 +6184,7 @@ bool FMetaSoundFrontendDocumentBuilder::SwapGraphInput(const FMetasoundFrontendC
 			FMetasoundFrontendNode& NewNode = Graph.Nodes[*NodeIndex];
 			NewNode.Style.Display.Locations = *Locations;
 		}
-#endif // WITH_EDITOR
+#endif // WITH_EDITORONLY_DATA
 
 		// 5b. Add to new copy existing node edges
 		TArray<FMetasoundFrontendEdge> RemovedEdges;
@@ -6208,10 +6216,10 @@ bool FMetaSoundFrontendDocumentBuilder::SwapGraphOutput(const FMetasoundFrontend
 
 	const IDocumentGraphInterfaceCache& InterfaceCache = DocumentCache->GetInterfaceCache();
 
-#if WITH_EDITOR
+#if WITH_EDITORONLY_DATA
 	using FPageNodeLocations = TMap<FGuid, FVector2D>;
 	TMap<FGuid, FPageNodeLocations> PageNodeLocations;
-#endif // WITH_EDITOR
+#endif // WITH_EDITORONLY_DATA
 
 	// 2. Gather data from existing page member/node needed to swap
 	TMultiMap<FGuid, FMetasoundFrontendEdge> RemovedEdgesPerPage;
@@ -6228,9 +6236,9 @@ bool FMetaSoundFrontendDocumentBuilder::SwapGraphOutput(const FMetasoundFrontend
 		const FMetasoundFrontendNode* ExistingOutputNode = NodeCache.FindNode(NodeID);
 		check(ExistingOutputNode);
 
-#if WITH_EDITOR
+#if WITH_EDITORONLY_DATA
 		PageNodeLocations.Add(Graph.PageID, ExistingOutputNode->Style.Display.Locations);
-#endif // WITH_EDITOR
+#endif // WITH_EDITORONLY_DATA
 
 		const FGuid VertexID = ExistingOutputNode->Interface.Inputs.Last().VertexID;
 		TArray<const FMetasoundFrontendEdge*> Edges = DocumentCache->GetEdgeCache(Graph.PageID).FindEdges(NodeID, VertexID);
@@ -6251,9 +6259,9 @@ bool FMetaSoundFrontendDocumentBuilder::SwapGraphOutput(const FMetasoundFrontend
 	{
 		FMetasoundFrontendClassOutput NewOutput = InNewOutputVertex;
 		NewOutput.NodeID = NodeID;
-#if WITH_EDITOR
+#if WITH_EDITORONLY_DATA
 		NewOutput.Metadata.SetSerializeText(InExistingOutputVertex.Metadata.GetSerializeText());
-#endif // WITH_EDITOR
+#endif // WITH_EDITORONLY_DATA
 
 		const FMetasoundFrontendNode* NewOutputNode = AddGraphOutput(MoveTemp(NewOutput));
 		checkf(NewOutputNode, TEXT("Failed to add new output node when swapping graph outputs"));
@@ -6263,7 +6271,7 @@ bool FMetaSoundFrontendDocumentBuilder::SwapGraphOutput(const FMetasoundFrontend
 
 	Document.RootGraph.IterateGraphPages([&](FMetasoundFrontendGraph& Graph)
 	{
-#if WITH_EDITOR
+#if WITH_EDITORONLY_DATA
 		// 5a. Add to new copy existing node locations
 		if (const FPageNodeLocations* Locations = PageNodeLocations.Find(Graph.PageID))
 		{
@@ -6273,7 +6281,7 @@ bool FMetaSoundFrontendDocumentBuilder::SwapGraphOutput(const FMetasoundFrontend
 			FMetasoundFrontendNode& NewNode = Graph.Nodes[*NodeIndex];
 			NewNode.Style.Display.Locations = *Locations;
 		}
-#endif // WITH_EDITOR
+#endif // WITH_EDITORONLY_DATA
 
 		// 5b. Add to new copy existing node edges
 		TArray<FMetasoundFrontendEdge> RemovedEdges;
diff --git a/Engine/Plugins/Runtime/Metasound/Source/MetasoundFrontend/Public/MetasoundFrontendDocumentBuilder.h b/Engine/Plugins/Runtime/Metasound/Source/MetasoundFrontend/Public/MetasoundFrontendDocumentBuilder.h
--- a/Engine/Plugins/Runtime/Metasound/Source/MetasoundFrontend/Public/MetasoundFrontendDocumentBuilder.h
+++ b/Engine/Plugins/Runtime/Metasound/Source/MetasoundFrontend/Public/MetasoundFrontendDocumentBuilder.h
@@ -456,7 +456,7 @@ public:
 	UE_API bool RemoveGraphComment(const FGuid& InCommentID, const FGuid* InPageID = nullptr);
 #endif // WITH_EDITORONLY_DATA
 
-	UE_API bool RemoveGraphInput(FName InputName);
+	UE_API bool RemoveGraphInput(FName InputName, bool bRemoveTemplateInputNodes = true);
 	UE_API bool RemoveGraphOutput(FName OutputName);
 
 #if WITH_EDITORONLY_DATA
