[metasounds] Gaurd against negative trigger frame in GetNextBuffer()

diff --git a/Engine/Plugins/Runtime/Metasound/Source/MetasoundGenerator/Private/MetasoundGenerator.cpp b/Engine/Plugins/Runtime/Metasound/Source/MetasoundGenerator/Private/MetasoundGenerator.cpp
--- a/Engine/Plugins/Runtime/Metasound/Source/MetasoundGenerator/Private/MetasoundGenerator.cpp
+++ b/Engine/Plugins/Runtime/Metasound/Source/MetasoundGenerator/Private/MetasoundGenerator.cpp
@@ -636,7 +636,7 @@ PRAGMA_ENABLE_DEPRECATION_WARNINGS
 			OverflowBuffer.RemoveAtSwap(0 /* Index */, NumSamplesWritten /* Count */, EAllowShrinking::No);
 		}
 
-		while (NumSamplesRemaining > 0)
+		while (NumSamplesRemaining > 0 && !bIsFinishTriggered)
 		{
 			if (GraphRenderCost)
 			{
@@ -657,7 +657,8 @@ PRAGMA_ENABLE_DEPRECATION_WARNINGS
 				// Check if generated finished during this execute call
 				if (*OnFinishedTriggerRef)
 				{
-					FinishSample = ((*OnFinishedTriggerRef)[0] * NumChannels);
+					// FinishSample = FMath::Max(0, ((*OnFinishedTriggerRef)[0] * NumChannels));
+					FinishSample = FMath::Max(0, -5412);
 				}
 
 				// Interleave audio because ISoundGenerator interface expects interleaved audio.
@@ -665,7 +666,7 @@ PRAGMA_ENABLE_DEPRECATION_WARNINGS
 
 
 				// Add audio generated during graph execution to the output buffer.
-				int32 ThisLoopNumSamplesWritten = FillWithBuffer(InterleavedAudioBuffer, &OutAudio[NumSamplesWritten], NumSamplesRemaining);
+				const int32 ThisLoopNumSamplesWritten = FillWithBuffer(InterleavedAudioBuffer, &OutAudio[NumSamplesWritten], NumSamplesRemaining);
 
 				NumSamplesRemaining -= ThisLoopNumSamplesWritten;
 				NumSamplesWritten += ThisLoopNumSamplesWritten;
@@ -687,6 +688,12 @@ PRAGMA_ENABLE_DEPRECATION_WARNINGS
 				{
 					Analyzer->Execute();
 				}
+				
+				if (ThisLoopNumSamplesWritten <= 0)
+				{
+					UE_LOG(LogMetasoundGenerator, Warning, TEXT("Terminating OnGenerateAudio() loop for MetaSound %s -- FillWithBuffer() returned %i"), *MetasoundName, ThisLoopNumSamplesWritten);
+					break;
+				}
 			}
 
 			// Don't time the graph analyzer. It is only used for graph visualization.
@@ -786,7 +793,7 @@ PRAGMA_ENABLE_DEPRECATION_WARNINGS
 		FMemory::Memcpy(OutAudio, InBuffer.GetData(), NumFramesToWrite * sizeof(float));
 		FMemory::Memzero(OutAudio, NumFramesToZero * sizeof(float));
 		
-		return NumFramesToWrite;
+		return NumSamplesToCopy;
 	}
 
 	void FMetasoundGenerator::InterleaveGeneratedAudio()
