Fixing a crasher in MetasoundEditorModule that would occur when unloading the module in -game/-server.

#rnx

diff --git a/Engine/Plugins/Runtime/Metasound/Source/MetasoundEditor/Private/MetasoundEditorModule.cpp b/Engine/Plugins/Runtime/Metasound/Source/MetasoundEditor/Private/MetasoundEditorModule.cpp
--- a/Engine/Plugins/Runtime/Metasound/Source/MetasoundEditor/Private/MetasoundEditorModule.cpp
+++ b/Engine/Plugins/Runtime/Metasound/Source/MetasoundEditor/Private/MetasoundEditorModule.cpp
@@ -972,12 +972,15 @@ namespace Metasound::Editor
 					case ERegistrationAssetContext::Removing:
 					case ERegistrationAssetContext::Renaming:
 					{
-						UAssetEditorSubsystem* AssetEditorSubsystem = GEditor->GetEditorSubsystem<UAssetEditorSubsystem>();
-						if (AssetEditorSubsystem)
+						if (GIsEditor)
 						{
-							// Close the editors so the internal reference to the builder doesn't suddenly 
-							// upon GC enter an invalid state (pointing to a null MetaSound asset)
-							AssetEditorSubsystem->CloseAllEditorsForAsset(&InMetaSound);
+							UAssetEditorSubsystem* AssetEditorSubsystem = GEditor->GetEditorSubsystem<UAssetEditorSubsystem>();
+							if (AssetEditorSubsystem)
+							{
+								// Close the editors so the internal reference to the builder doesn't suddenly 
+								// upon GC enter an invalid state (pointing to a null MetaSound asset)
+								AssetEditorSubsystem->CloseAllEditorsForAsset(&InMetaSound);
+							}
 						}
 						break;
 					}
