Follow up fix to remove debug code/comment

diff --git a/Engine/Plugins/Runtime/Metasound/Source/MetasoundGenerator/Private/MetasoundGenerator.cpp b/Engine/Plugins/Runtime/Metasound/Source/MetasoundGenerator/Private/MetasoundGenerator.cpp
--- a/Engine/Plugins/Runtime/Metasound/Source/MetasoundGenerator/Private/MetasoundGenerator.cpp
+++ b/Engine/Plugins/Runtime/Metasound/Source/MetasoundGenerator/Private/MetasoundGenerator.cpp
@@ -657,8 +657,7 @@ PRAGMA_ENABLE_DEPRECATION_WARNINGS
 				// Check if generated finished during this execute call
 				if (*OnFinishedTriggerRef)
 				{
-					// FinishSample = FMath::Max(0, ((*OnFinishedTriggerRef)[0] * NumChannels));
-					FinishSample = FMath::Max(0, -5412);
+					FinishSample = FMath::Max(0, ((*OnFinishedTriggerRef)[0] * NumChannels));
 				}
 
 				// Interleave audio because ISoundGenerator interface expects interleaved audio.
