Fix for diff objects registering & unregistering improperly with MetaSound asset manager in project-specific contexts

#tests local cook, open/close/reload/rename MetaSound asset, repro from bug
#rnx

diff --git a/Engine/Plugins/Runtime/Metasound/Source/MetasoundEngine/Private/MetasoundAssetSubsystem.cpp b/Engine/Plugins/Runtime/Metasound/Source/MetasoundEngine/Private/MetasoundAssetSubsystem.cpp
--- a/Engine/Plugins/Runtime/Metasound/Source/MetasoundEngine/Private/MetasoundAssetSubsystem.cpp
+++ b/Engine/Plugins/Runtime/Metasound/Source/MetasoundEngine/Private/MetasoundAssetSubsystem.cpp
@@ -54,8 +54,17 @@ namespace Metasound::Engine
 					return ClassInfo.AssetPath.GetAssetName() == InAssetPath.GetAssetName();
 				};
 
-				if (AssetInfoArray->RemoveAllSwap(ComparePaths, EAllowShrinking::No) > 0)
+				const int32 NumRemoved = AssetInfoArray->RemoveAllSwap(ComparePaths, EAllowShrinking::No);
+				if (NumRemoved > 0)
 				{
+					if (NumRemoved > 1 && InAssetPath.GetPackageName().IsNone())
+					{
+						UE_LOG(LogMetaSound, Display,
+							TEXT("MetaSoundAssetManager: Multiple assets registered with key '%s' and is removing all asset class info with provided asset path missing package.  "
+							"Likely caused by diff and request for removal is currently amidst diff object distruction."),
+							*AssetKey.ToString());
+					}
+
 					if (AssetInfoArray->IsEmpty())
 					{
 						Map.Remove(AssetKey);
@@ -149,8 +158,10 @@ namespace Metasound::Engine
 			const FMetaSoundAssetKey AssetKey(Document.RootGraph.Metadata);
 			if (!ContainsKey(AssetKey))
 			{
-				AddOrUpdateFromObject(*InAssetBase.GetOwningAsset());
-				UE_LOG(LogMetaSound, Verbose, TEXT("Adding asset '%s' to MetaSoundAsset registry."), *InAssetBase.GetOwningAssetName());
+				if (AddOrUpdateFromObject(*InAssetBase.GetOwningAsset()).IsValid())
+				{
+					UE_LOG(LogMetaSound, Verbose, TEXT("Adding asset '%s' to MetaSoundAsset registry."), *InAssetBase.GetOwningAssetName());
+				}
 			}
 		}
 
@@ -209,11 +220,22 @@ namespace Metasound::Engine
 #endif // WITH_EDITORONLY_DATA
 
 	FMetaSoundAssetKey FMetaSoundAssetManager::AddOrUpdateFromObject(const UObject& InObject)
+	{
+		// Don't add temporary assets used for diffing
+		if (const UPackage* Package = InObject.GetPackage(); !Package || Package->HasAnyPackageFlags(PKG_ForDiffing))
+		{
+			return FMetaSoundAssetKey::GetInvalid();
+		}
+
+		return AddOrUpdateFromObjectInternal(InObject);
+	}
+
+	FMetaSoundAssetKey FMetaSoundAssetManager::AddOrUpdateFromObjectInternal(const UObject& InObject)
 	{
 		using namespace AssetSubsystemPrivate;
 		using namespace Frontend;
 
-		METASOUND_TRACE_CPUPROFILER_EVENT_SCOPE(FMetaSoundAssetManager::AddOrUpdateFromObject);
+		METASOUND_TRACE_CPUPROFILER_EVENT_SCOPE(FMetaSoundAssetManager::AddOrUpdateFromObjectInternal);
 
 		const FMetasoundAssetBase* MetaSoundAsset = Metasound::IMetasoundUObjectRegistry::Get().GetObjectAsAssetBase(&InObject);
 		check(MetaSoundAsset);
@@ -221,7 +243,7 @@ namespace Metasound::Engine
 		TScriptInterface<const IMetaSoundDocumentInterface> ScriptInterface(&InObject);
 		check(ScriptInterface.GetObject() != nullptr);
 		const IMetaSoundDocumentInterface& DocInterface = *ScriptInterface;
-		return AssetSubsystemPrivate::AddClassInfo(TagDataMapCriticalSection, ClassInfoMap, FMetaSoundAssetClassInfo (DocInterface));
+		return AssetSubsystemPrivate::AddClassInfo(TagDataMapCriticalSection, ClassInfoMap, FMetaSoundAssetClassInfo(DocInterface));
 	}
 
 	void FMetaSoundAssetManager::AddOrLoadAndUpdateFromObjectAsync(const FAssetData& InAssetData, Frontend::IMetaSoundAssetManager::FOnUpdatedAssetLoaded&& OnUpdatedAssetLoaded)
@@ -241,7 +263,7 @@ namespace Metasound::Engine
 		if (UObject* Object = Path.ResolveObject())
 		{
 			UE_LOG(LogMetaSound, VeryVerbose, TEXT("FMetaSoundAssetManager::AddOrLoadAndUpdateFromObjectAsync called, but object ''%s' already loaded. Updating entry using in-memory object definition."), *Path.ToString());
-			const FMetaSoundAssetKey AssetKey = AddOrUpdateFromObject(*Object);
+			const FMetaSoundAssetKey AssetKey = AddOrUpdateFromObjectInternal(*Object);
 			OnUpdatedAssetLoaded(AssetKey, *Object);
 			OnUpdatedAssetLoaded.Reset();
 			return;
@@ -282,7 +304,7 @@ namespace Metasound::Engine
 		if (UObject* Object = Path.ResolveObject())
 		{
 			UE_LOG(LogMetaSound, VeryVerbose, TEXT("AddOrUpdateFromAssetData called, but object '%s' already loaded. Updating entry using in-memory object definition."), *Path.ToString());
-			AddOrUpdateFromObject(*Object);
+			AddOrUpdateFromObjectInternal(*Object);
 		}
 		else
 		{
@@ -596,7 +618,7 @@ namespace Metasound::Engine
 #endif // WITH_EDITORONLY_DATA
 
 						const FMetaSoundAssetKey AssetKey = AssetManager->AddOrUpdateFromObject(*MetaSoundObj);
-						if (PackageLoadedArgs.OnUpdatedAssetLoaded)
+						if (AssetKey.IsValid() && PackageLoadedArgs.OnUpdatedAssetLoaded)
 						{
 							PackageLoadedArgs.OnUpdatedAssetLoaded(AssetKey, *MetaSoundObj);
 						}
@@ -932,18 +954,18 @@ namespace Metasound::Engine
 	{
 		using namespace Frontend;
 
+		// Don't need to remove assets used for diffing as they can't be added.
+		if (UPackage* Package = InObject.GetPackage(); Package && Package->HasAnyPackageFlags(PKG_ForDiffing))
+		{
+			return;
+		}
+
 		TScriptInterface<const IMetaSoundDocumentInterface> DocInterface(&InObject);
 		check(DocInterface.GetObject());
 		const FMetasoundFrontendDocument& Document = DocInterface->GetConstDocument();
 		const FMetasoundFrontendClassMetadata& Metadata = Document.RootGraph.Metadata;
 
 		const FTopLevelAssetPath AssetPath(&InObject);
-		if (IDocumentBuilderRegistry* BuilderRegistry = IDocumentBuilderRegistry::Get())
-		{
-			constexpr bool bForceUnregister = true;
-			BuilderRegistry->FinishBuilding(Metadata.GetClassName(), AssetPath, bForceUnregister);
-		}
-
 		const FMetaSoundAssetKey AssetKey(Metadata.GetClassName(), Metadata.GetVersion());
 		AssetSubsystemPrivate::RemoveClassInfo(TagDataMapCriticalSection, ClassInfoMap, AssetKey, AssetPath);
 	}
diff --git a/Engine/Plugins/Runtime/Metasound/Source/MetasoundEngine/Public/MetasoundAssetSubsystem.h b/Engine/Plugins/Runtime/Metasound/Source/MetasoundEngine/Public/MetasoundAssetSubsystem.h
--- a/Engine/Plugins/Runtime/Metasound/Source/MetasoundEngine/Public/MetasoundAssetSubsystem.h
+++ b/Engine/Plugins/Runtime/Metasound/Source/MetasoundEngine/Public/MetasoundAssetSubsystem.h
@@ -156,6 +156,8 @@ PRAGMA_ENABLE_DEPRECATION_WARNINGS
 		virtual FString GetReferencerName() const override { return TEXT("FMetaSoundAssetManager"); }
 
 	private:
+		FMetaSoundAssetKey AddOrUpdateFromObjectInternal(const UObject& InObject);
+
 		TArray<FMetaSoundAsyncAssetDependencies> LoadingDependencies;
 
 		static UE_API void DepthFirstVisit_AssetKey(const FMetaSoundAssetKey& InKey, TFunctionRef<TSet<FMetaSoundAssetKey>(const FMetaSoundAssetKey&)> VisitFunction);
@@ -164,14 +166,14 @@ PRAGMA_ENABLE_DEPRECATION_WARNINGS
 		// of assets, improper plugin migration, etc. but should usually just be a single entry) registered
 		// with the manager. Info may not be returned if the asset has yet to be registered with the MetaSoundAssetManager
 		// or may be out-of-date if edited after last registration attempt via AddOrUpdate... calls.
-		UE_API TArray<Frontend::FMetaSoundAssetClassInfo> FindAssetClassInfoInternal(const FMetaSoundAssetKey& InKey) const;
+		TArray<Frontend::FMetaSoundAssetClassInfo> FindAssetClassInfoInternal(const FMetaSoundAssetKey& InKey) const;
 
-		UE_API FMetaSoundAsyncAssetDependencies* FindLoadingDependencies(const UObject* InParentAsset);
-		UE_API FMetaSoundAsyncAssetDependencies* FindLoadingDependencies(int32 InLoadID);
-		UE_API void RemoveLoadingDependencies(int32 InLoadID);
-		UE_API void OnReferencedAssetsLoaded(int32 InLoadID);
+		FMetaSoundAsyncAssetDependencies* FindLoadingDependencies(const UObject* InParentAsset);
+		FMetaSoundAsyncAssetDependencies* FindLoadingDependencies(int32 InLoadID);
+		void RemoveLoadingDependencies(int32 InLoadID);
+		void OnReferencedAssetsLoaded(int32 InLoadID);
 
-		UE_API bool GetReferencedPresetHierarchyInternal(FMetasoundAssetBase& InAsset, TArray<FMetasoundAssetBase*>& OutReferencedAssets) const;
+		bool GetReferencedPresetHierarchyInternal(FMetasoundAssetBase& InAsset, TArray<FMetasoundAssetBase*>& OutReferencedAssets) const;
 
 		struct FPackageLoadedArgs
 		{
@@ -181,7 +183,7 @@ PRAGMA_ENABLE_DEPRECATION_WARNINGS
 			FOnUpdatedAssetLoaded OnUpdatedAssetLoaded;
 		};
 
-		static UE_API void OnPackageLoaded(const FPackageLoadedArgs& PackageLoadedArgs);
+		static void OnPackageLoaded(const FPackageLoadedArgs& PackageLoadedArgs);
 
 		FStreamableManager StreamableManager;
 		int32 AsyncLoadIDCounter = 0;
diff --git a/Engine/Plugins/Runtime/Metasound/Source/MetasoundEngine/Public/MetasoundBuilderBase.cpp b/Engine/Plugins/Runtime/Metasound/Source/MetasoundEngine/Public/MetasoundBuilderBase.cpp
--- a/Engine/Plugins/Runtime/Metasound/Source/MetasoundEngine/Public/MetasoundBuilderBase.cpp
+++ b/Engine/Plugins/Runtime/Metasound/Source/MetasoundEngine/Public/MetasoundBuilderBase.cpp
@@ -47,7 +47,7 @@ void UMetaSoundBuilderBase::BeginDestroy()
 		if (IDocumentBuilderRegistry* BuilderRegistry = IDocumentBuilderRegistry::Get())
 		{
 			const FMetasoundFrontendClassName& MetaSoundClassName = Builder.GetConstDocumentChecked().RootGraph.Metadata.GetClassName();
-			BuilderRegistry->FinishBuilding(MetaSoundClassName);
+			BuilderRegistry->FinishBuilding(MetaSoundClassName, Builder.GetHintPath());
 		}
 
 		// The registry may have not been active or it was and the internal weak pointer record of this object no longer accessible.
diff --git a/Engine/Plugins/Runtime/Metasound/Source/MetasoundFrontend/Private/MetasoundFrontendDocumentBuilder.cpp b/Engine/Plugins/Runtime/Metasound/Source/MetasoundFrontend/Private/MetasoundFrontendDocumentBuilder.cpp
--- a/Engine/Plugins/Runtime/Metasound/Source/MetasoundFrontend/Private/MetasoundFrontendDocumentBuilder.cpp
+++ b/Engine/Plugins/Runtime/Metasound/Source/MetasoundFrontend/Private/MetasoundFrontendDocumentBuilder.cpp
@@ -3171,7 +3171,7 @@ const FMetasoundFrontendDocument& FMetaSoundFrontendDocumentBuilder::GetConstDoc
 const IMetaSoundDocumentInterface& FMetaSoundFrontendDocumentBuilder::GetConstDocumentInterfaceChecked() const
 {
 	const IMetaSoundDocumentInterface* Interface = DocumentInterface.GetInterface();
-	checkf(Interface, TEXT("Failed to return document; interface must always be valid while builder is operating on MetaSound UObject!"));
+	checkf(Interface, TEXT("Failed to return document; interface must always be valid while builder is operating on MetaSound UObject! Builder constructed with asset at %s"), *HintPath.ToString());
 	return *Interface;
 }
 
@@ -3186,7 +3186,7 @@ const FString FMetaSoundFrontendDocumentBuilder::GetDebugName() const
 const FMetasoundFrontendDocument& FMetaSoundFrontendDocumentBuilder::GetDocument() const
 {
 	const IMetaSoundDocumentInterface* Interface = DocumentInterface.GetInterface();
-	checkf(Interface, TEXT("Failed to return document; interface must always be valid while builder is operating on MetaSound UObject!"));
+	checkf(Interface, TEXT("Failed to return document; interface must always be valid while builder is operating on MetaSound UObject! Builder constructed with asset at %s"), *HintPath.ToString());
 	return Interface->GetConstDocument();
 }
 
@@ -3203,14 +3203,14 @@ Metasound::Frontend::FDocumentModifyDelegates& FMetaSoundFrontendDocumentBuilder
 const IMetaSoundDocumentInterface& FMetaSoundFrontendDocumentBuilder::GetDocumentInterface() const
 {
 	const IMetaSoundDocumentInterface* Interface = DocumentInterface.GetInterface();
-	checkf(Interface, TEXT("Failed to return document; interface must always be valid while builder is operating on MetaSound UObject!"));
+	checkf(Interface, TEXT("Failed to return document; interface must always be valid while builder is operating on MetaSound UObject! Builder constructed with asset at %s"), *HintPath.ToString());
 	return *Interface;
 }
 
 IMetaSoundDocumentInterface& FMetaSoundFrontendDocumentBuilder::GetDocumentInterfaceChecked() const
 {
 	IMetaSoundDocumentInterface* Interface = DocumentInterface.GetInterface();
-	checkf(Interface, TEXT("Failed to return document; interface must always be valid while builder is operating on MetaSound UObject!"));
+	checkf(Interface, TEXT("Failed to return document; interface must always be valid while builder is operating on MetaSound UObject! Builder constructed with asset at %s"), *HintPath.ToString());
 	return *Interface;
 }
 
@@ -3265,6 +3265,10 @@ const TSet<FName>* FMetaSoundFrontendDocumentBuilder::GetGraphInputsInheritingDe
 	return nullptr;
 }
 
+const FTopLevelAssetPath& FMetaSoundFrontendDocumentBuilder::GetHintPath() const
+{
+	return HintPath;
+}
 
 FMetasoundAssetBase& FMetaSoundFrontendDocumentBuilder::GetMetasoundAsset() const
 {
@@ -4216,8 +4220,11 @@ void FMetaSoundFrontendDocumentBuilder::BeginBuilding(TSharedPtr<Metasound::Fron
 {
 	using namespace Metasound::Frontend;
 
+	HintPath = { };
 	if (DocumentInterface)
 	{
+		HintPath = DocumentInterface->GetAssetPathChecked();
+
 		// Potentially at cook and runtime, the default graph may have been cooked away,
 		// so initialize build page to a valid page ID if possible. On initial construction,
 		// it may be possible the default graph has yet to be initialized, so don't error if
diff --git a/Engine/Plugins/Runtime/Metasound/Source/MetasoundFrontend/Private/MetasoundFrontendNodeClassRegistry.cpp b/Engine/Plugins/Runtime/Metasound/Source/MetasoundFrontend/Private/MetasoundFrontendNodeClassRegistry.cpp
--- a/Engine/Plugins/Runtime/Metasound/Source/MetasoundFrontend/Private/MetasoundFrontendNodeClassRegistry.cpp
+++ b/Engine/Plugins/Runtime/Metasound/Source/MetasoundFrontend/Private/MetasoundFrontendNodeClassRegistry.cpp
@@ -89,7 +89,7 @@ namespace Metasound::Frontend
 				if (!bIsBuilding)
 				{
 					const FMetasoundFrontendClassName& ClassName = Document.RootGraph.Metadata.GetClassName();
-					IDocumentBuilderRegistry::GetChecked().FinishBuilding(ClassName);
+					IDocumentBuilderRegistry::GetChecked().FinishBuilding(ClassName, OriginalDocBuilder.GetHintPath());
 				}
 			}
 		#endif // WITH_EDITORONLY_DATA
diff --git a/Engine/Plugins/Runtime/Metasound/Source/MetasoundFrontend/Public/MetasoundFrontendDocumentBuilder.h b/Engine/Plugins/Runtime/Metasound/Source/MetasoundFrontend/Public/MetasoundFrontendDocumentBuilder.h
--- a/Engine/Plugins/Runtime/Metasound/Source/MetasoundFrontend/Public/MetasoundFrontendDocumentBuilder.h
+++ b/Engine/Plugins/Runtime/Metasound/Source/MetasoundFrontend/Public/MetasoundFrontendDocumentBuilder.h
@@ -350,6 +350,11 @@ public:
 
 	UE_DEPRECATED(5.5, "Use GetConstDocumentInterfaceChecked instead")
 	UE_API const IMetaSoundDocumentInterface& GetDocumentInterface() const;
+
+	// Path for document object provided at construction time. Cached on builder as a useful means of debugging
+	// and enables weak reference removal from the builder registry should the object be mid-destruction.
+	UE_API const FTopLevelAssetPath& GetHintPath() const;
+
 	UE_API FMetasoundAssetBase& GetMetasoundAsset() const;
 	
 	// Get the asset referenced by this builder's preset asset, nullptr if builder is not a preset.
@@ -708,6 +713,8 @@ private:
 
 	TSharedPtr<Metasound::Frontend::IDocumentCache> DocumentCache;
 	TSharedPtr<Metasound::Frontend::FDocumentModifyDelegates> DocumentDelegates;
+
+	FTopLevelAssetPath HintPath;
 };
 
 #undef UE_API
