Resample metasound audio bus reader/writer when metasound sample rate does not match audio mixer sample rate.
[RN] MetaSound Audio Bus Reader and Writer nodes resample audio when there are sample rate discrepencies.

diff --git a/Engine/Plugins/Runtime/Metasound/Source/MetasoundEngine/Private/MetasoundAudioBusPrivate.cpp b/Engine/Plugins/Runtime/Metasound/Source/MetasoundEngine/Private/MetasoundAudioBusPrivate.cpp
new file mode 100644
--- /dev/null
+++ b/Engine/Plugins/Runtime/Metasound/Source/MetasoundEngine/Private/MetasoundAudioBusPrivate.cpp
@@ -0,0 +1,128 @@
+// Copyright Epic Games, Inc. All Rights Reserved.
+
+#include "MetasoundAudioBusPrivate.h"
+
+#include "DSP/MultithreadedPatching.h"
+#include "DSP/RuntimeResampler.h"
+#include "HAL/IConsoleManager.h"
+#include "MetasoundLog.h"
+
+FAutoConsoleVariableRef CVarAudioBusReaderNodeEnableResampledAudioBus(
+	TEXT("au.MetaSound.EnableAudioBusResampler"),
+	Metasound::AudioBusPrivate::EnableResampledAudioBus,
+	TEXT("Enable the use of a resampler when the AudioMixer sample rate does not match the MetaSound sample rate.\n")
+	TEXT("0: disabled, 1: enabled (default)"),
+	ECVF_Default);
+
+namespace Metasound::AudioBusPrivate
+{
+	int32 EnableResampledAudioBus = 1;
+
+	FResampledPatchOutput::FResampledPatchOutput(int32 InNumChannels, float InAudioBusSampleRate, float InMetaSoundSampleRate, int32 InMetaSoundBlockSize, TSharedRef<Audio::FPatchOutput> InPatchOutput)
+	: NumChannels(InNumChannels)
+	, Resampler(InNumChannels) 
+	, PatchOutput(MoveTemp(InPatchOutput))
+	{
+		check(InNumChannels > 0);
+		check(InAudioBusSampleRate > 0.f);
+		check(InMetaSoundSampleRate> 0.f);
+		check(InMetaSoundBlockSize > 0);
+
+		// Set sample rate to read/write rate. It is assumed that all audio buses
+		// read/write at the AudioBusSampleRate. MetaSounds which do not
+		// render at the AudioBusSampleRate are resampled outside of the MetaSound
+		// system to match the AudioBusSampleRate. This SampleRateRatio
+		// accounts for the resampling that occurs outside of the MetaSound Source.
+		const float SampleRateRatio = InAudioBusSampleRate / InMetaSoundSampleRate;
+
+		Resampler.SetFrameRatio(SampleRateRatio);
+
+		// A temporary buffer is required to interact with the FPatchOutput 
+		// API. The FPatchOutput API could be reworked to not require the
+		// use of a temporary buffer by providing a Peek method which returns
+		// a const view of the array already existing in the FPatchOutput
+		const int32 NumFramesNeededFromAudioBus = Resampler.GetNumInputFramesNeededToProduceOutputFrames(InMetaSoundBlockSize) + FMath::CeilToInt(SampleRateRatio);
+		ScratchBuffer.AddUninitialized(NumFramesNeededFromAudioBus * NumChannels);
+
+	}
+
+	int32 FResampledPatchOutput::PopAudio(float* OutAudio, int32 InNumSamplesToPop, bool bInUseLatestAudio)
+	{
+		const int32 NumFramesNeededFromAudioBus = Resampler.GetNumInputFramesNeededToProduceOutputFrames(InNumSamplesToPop / NumChannels);
+		const int32 NumSamplesNeededFromAudioBus = NumFramesNeededFromAudioBus * NumChannels;
+
+		// Check that buffer allocations are not needed during rendering. 
+		if (!ensureMsgf(NumSamplesNeededFromAudioBus <= ScratchBuffer.Num(), TEXT("More initial slack is needed in allocation of AudioBuBuffer. Allocated %d, Requested: %d"), ScratchBuffer.Num(), NumSamplesNeededFromAudioBus))
+		{
+			ScratchBuffer.Reset();
+			ScratchBuffer.AddUninitialized(NumSamplesNeededFromAudioBus);
+		}
+
+		int32 NumSamplesPopped = PatchOutput->PopAudio(ScratchBuffer.GetData(), NumSamplesNeededFromAudioBus, bInUseLatestAudio);
+
+		int32 NumFramesConsumed = -1;
+		int32 NumFramesProduced = -1;
+
+		Resampler.ProcessInterleaved(TArrayView<const float>{ScratchBuffer.GetData(), NumSamplesPopped}, TArrayView<float>{OutAudio, InNumSamplesToPop}, NumFramesConsumed, NumFramesProduced);
+
+		// Check that all the input frames have been consumed so that they
+		// do not need to be maintained here. 
+		if ((NumChannels * NumFramesConsumed) < NumSamplesPopped)
+		{
+			UE_LOG(LogMetaSound, Warning, TEXT("Dropping %d samples"), NumSamplesPopped - (NumChannels * NumFramesConsumed));
+		}
+		return NumFramesProduced * NumChannels;
+	};
+
+
+	FResampledPatchInput::FResampledPatchInput(int32 InNumChannels, float InAudioBusSampleRate, float InMetaSoundSampleRate, int32 InMetaSoundBlockSize, Audio::FPatchInput InPatchInput)
+	: NumChannels(InNumChannels)
+	, Resampler(InNumChannels) 
+	, PatchInput(MoveTemp(InPatchInput))
+	{
+		check(InNumChannels > 0);
+		check(InAudioBusSampleRate > 0.f);
+		check(InMetaSoundSampleRate> 0.f);
+		check(InMetaSoundBlockSize > 0);
+
+		// Set sample rate to read/write rate. It is assumed that all audio buses
+		// read/write at the AudioBusSampleRate. MetaSounds which do not
+		// render at the AudioBusSampleRate are resampled outside of the MetaSound
+		// system to match the AudioBusSampleRate. This SampleRateRatio
+		// accounts for the resampling that occurs outside of the MetaSound Source.
+		const float SampleRateRatio = InMetaSoundSampleRate / InAudioBusSampleRate;
+
+		Resampler.SetFrameRatio(SampleRateRatio);
+
+		const int32 MaxOutputBufferNumFrames = Resampler.GetNumOutputFramesProducedByInputFrames(InMetaSoundBlockSize + 1) + FMath::CeilToInt(SampleRateRatio);
+		ScratchBuffer.AddUninitialized(MaxOutputBufferNumFrames * NumChannels);
+	}
+
+	int32 FResampledPatchInput::PushAudio(float const* InAudio, int32 InNumSamplesToPush)
+	{
+		const int32 NumResampledFramesToPush = Resampler.GetNumOutputFramesProducedByInputFrames(InNumSamplesToPush / NumChannels);
+		const int32 NumResampledSamplesToPush = NumResampledFramesToPush * NumChannels;
+
+		// Check that buffer allocations are not needed during rendering. 
+		if (!ensureMsgf(NumResampledSamplesToPush <= ScratchBuffer.Num(), TEXT("More initial slack is needed in allocation of AudioBuBuffer. Allocated %d, Requested: %d"), ScratchBuffer.Num(), NumResampledSamplesToPush))
+		{
+			ScratchBuffer.Reset();
+			ScratchBuffer.AddUninitialized(NumResampledSamplesToPush);
+		}
+
+		int32 NumFramesConsumed = -1;
+		int32 NumFramesProduced = -1;
+		Resampler.ProcessInterleaved(TArrayView<const float>{InAudio, InNumSamplesToPush}, TArrayView<float>{ScratchBuffer.GetData(), NumResampledSamplesToPush}, NumFramesConsumed, NumFramesProduced);
+
+		// Check that all the input frames have been consumed so that they
+		// do not need to be maintained here. 
+		if ((NumChannels * NumFramesConsumed) < InNumSamplesToPush)
+		{
+			UE_LOG(LogMetaSound, Warning, TEXT("Dropping %d samples"), InNumSamplesToPush - (NumChannels * NumFramesConsumed));
+		}
+
+		int32 NumSamplesPushed = PatchInput.PushAudio(ScratchBuffer.GetData(), NumFramesProduced * NumChannels);
+
+		return NumFramesConsumed * NumChannels;
+	}
+}
diff --git a/Engine/Plugins/Runtime/Metasound/Source/MetasoundEngine/Private/MetasoundAudioBusPrivate.h b/Engine/Plugins/Runtime/Metasound/Source/MetasoundEngine/Private/MetasoundAudioBusPrivate.h
new file mode 100644
--- /dev/null
+++ b/Engine/Plugins/Runtime/Metasound/Source/MetasoundEngine/Private/MetasoundAudioBusPrivate.h
@@ -0,0 +1,47 @@
+// Copyright Epic Games, Inc. All Rights Reserved.
+
+#pragma once
+
+#include "DSP/AlignedBuffer.h"
+#include "DSP/MultithreadedPatching.h"
+#include "DSP/RuntimeResampler.h"
+#include "Templates/SharedPointer.h"
+
+namespace Metasound::AudioBusPrivate
+{
+	// Enablement flag for resampling of audio when the MetaSound sample rate
+	// does not match the AudioMixer sample rate. 
+	extern int32 EnableResampledAudioBus;
+
+	// Resample audio popped from an audio bus. 
+	class FResampledPatchOutput
+	{
+	public:
+		FResampledPatchOutput(int32 InNumChannels, float InAudioBusSampleRate, float InMetaSoundSampleRate, int32 InMetaSoundBlockSize, TSharedRef<Audio::FPatchOutput> InPatchOutput);
+
+		int32 PopAudio(float* OutAudio, int32 InNumSamplesToPop, bool bInUseLatestAudio);
+
+	private:
+
+		int32 NumChannels = 0;
+		Audio::FRuntimeResampler Resampler;
+		Audio::FAlignedFloatBuffer ScratchBuffer;
+		TSharedRef<Audio::FPatchOutput> PatchOutput;
+	};
+
+	// Resample audio pushed to an audio bus. 
+	class FResampledPatchInput
+	{
+	public:
+		FResampledPatchInput(int32 InNumChannels, float InAudioBusSampleRate, float InMetaSoundSampleRate, int32 InMetaSoundBlockSize, Audio::FPatchInput InPatchInput);
+
+		int32 PushAudio(float const* InAudio, int32 InNumSamplesToPush);
+
+	private:
+
+		int32 NumChannels = 0;
+		Audio::FRuntimeResampler Resampler;
+		Audio::FAlignedFloatBuffer ScratchBuffer;
+		Audio::FPatchInput PatchInput;
+	};
+}
diff --git a/Engine/Plugins/Runtime/Metasound/Source/MetasoundEngine/Private/MetasoundAudioBusReaderNode.cpp b/Engine/Plugins/Runtime/Metasound/Source/MetasoundEngine/Private/MetasoundAudioBusReaderNode.cpp
--- a/Engine/Plugins/Runtime/Metasound/Source/MetasoundEngine/Private/MetasoundAudioBusReaderNode.cpp
+++ b/Engine/Plugins/Runtime/Metasound/Source/MetasoundEngine/Private/MetasoundAudioBusReaderNode.cpp
@@ -1,11 +1,14 @@
 // Copyright Epic Games, Inc. All Rights Reserved.
 
-#include "AudioDevice.h"
-#include "Sound/AudioBus.h"
 #include "AudioBusSubsystem.h"
+#include "AudioDevice.h"
+#include "DSP/AlignedBuffer.h"
 #include "Internationalization/Text.h"
+#include "Sound/AudioBus.h"
+
 #include "MetasoundAudioBuffer.h"
 #include "MetasoundAudioBus.h"
+#include "MetasoundAudioBusPrivate.h"
 #include "MetasoundEngineNodesNames.h"
 #include "MetasoundExecutableOperator.h"
 #include "MetasoundFacade.h"
@@ -29,6 +32,7 @@ namespace Metasound
 		METASOUND_PARAM(InParamAudioBusInput, "Audio Bus", "Audio Bus Asset.")
 
 		METASOUND_PARAM(OutParamAudio, "Out {0}", "Audio bus output for channel {0}.");
+
 	}
 
 	int32 AudioBusReaderNodeInitialNumBlocks(int32 BlockSizeFrames, int32 AudioMixerOutputFrames)
@@ -115,6 +119,7 @@ namespace Metasound
 
 			bool bHasEnvironmentVars = InParams.Environment.Contains<Audio::FDeviceId>(SourceInterface::Environment::DeviceID);
 			bHasEnvironmentVars &= InParams.Environment.Contains<int32>(SourceInterface::Environment::AudioMixerNumOutputFrames);
+			bHasEnvironmentVars &= InParams.Environment.Contains<uint64>(SourceInterface::Environment::TransmitterID);
 			
 			if (bHasEnvironmentVars)
 			{
@@ -127,6 +132,7 @@ namespace Metasound
 					, *SourceInterface::Environment::DeviceID.ToString(), *SourceInterface::Environment::AudioMixerNumOutputFrames.ToString());
 				return nullptr;
 			}
+
 		}
 
 		TAudioBusReaderOperator(const FBuildOperatorParams& InParams, const FAudioBusAssetReadRef& InAudioBusAsset) : AudioBusAsset(InAudioBusAsset)
@@ -184,7 +190,17 @@ namespace Metasound
 
 			// Pop off the interleaved data from the audio bus
 			int32 NumSamplesToPop = BlockSizeFrames * AudioBusChannels;
-			int32 SamplesPopped = AudioBusPatchOutput->PopAudio(InterleavedBuffer.GetData(), NumSamplesToPop, false);
+			int32 SamplesPopped = -1;
+
+			if (ResampledPatchOutput.IsValid())
+			{
+				SamplesPopped = ResampledPatchOutput->PopAudio(InterleavedBuffer.GetData(), NumSamplesToPop, false);
+			}
+			else
+			{
+				SamplesPopped = AudioBusPatchOutput->PopAudio(InterleavedBuffer.GetData(), NumSamplesToPop, false);
+			}
+
 			if (SamplesPopped < NumSamplesToPop && !bWasUnderrunReported)
 			{
 				UE_LOG(LogMetaSound, Warning, TEXT("Underrun detected in audio bus reader node."));
@@ -204,6 +220,8 @@ namespace Metasound
 
 		void CreatePatchOutput()
 		{
+			using namespace AudioBusPrivate;
+
 			const FAudioBusProxyPtr& AudioBusProxy = AudioBusAsset->GetAudioBusProxy();
 			if (AudioBusProxy.IsValid())
 			{
@@ -237,7 +255,25 @@ namespace Metasound
 
 				AudioBusPatchOutput = AudioBusSubsystem->AddPatchOutputForSoundAndAudioBus(InstanceID, AudioBusKey, BlockSizeFrames, int32(AudioBusChannels));
 				PatchInput = AudioBusPatchOutput;
-				PatchInput.PushAudio(nullptr, NumBlocksToNumSamples(InitialNumBlocks()));
+
+				// Handle case of mismatched sample rate between audio mixer and
+				// metasound. 
+				if (EnableResampledAudioBus && (AudioMixerSampleRate != SampleRate) && (AudioMixerSampleRate > 0.f) && (SampleRate > 0.f) && AudioBusPatchOutput.IsValid())
+				{
+					UE_LOG(LogMetaSound, Warning, TEXT("Using a audio bus reader node is inefficient if the MetaSound sample rate %f does not match the AudioMixer sample rate %f. Please update MetaSound SampleRate to match the AudioMixer's SampleRate"), SampleRate, AudioMixerSampleRate);
+
+					ResampledPatchOutput = MakeUnique<FResampledPatchOutput>(AudioBusChannels, AudioMixerSampleRate, SampleRate, BlockSizeFrames, AudioBusPatchOutput.ToSharedRef());
+
+					// Sample rate matches between audio mixer and metasound. This
+					// node will consume audio from the patch in approximate block
+					// sizes of (BlockSize * AudioMixerSampleRate / SampleRate)
+					PatchInput.PushAudio(nullptr, NumBlocksToNumSamples(InitialNumBlocks(), AudioMixerSampleRate / SampleRate));
+				}
+				else
+				{
+					// Sample rate matches between audio mixer and metasound. 
+					PatchInput.PushAudio(nullptr, NumBlocksToNumSamples(InitialNumBlocks()));
+				}
 
 				InterleavedBuffer.Reset();
 				InterleavedBuffer.AddUninitialized(NumBlocksToNumSamples(1));
@@ -250,7 +286,9 @@ namespace Metasound
 			using namespace AudioBusReaderNode;
 
 			InterleavedBuffer.Reset();
+			ResampledPatchOutput.Reset();
 			AudioMixerOutputFrames = INDEX_NONE;
+			AudioMixerSampleRate = 1.f;
 			AudioDeviceId = INDEX_NONE;
 			InstanceID = 0;
 			AudioBusId = 0;
@@ -258,7 +296,8 @@ namespace Metasound
 			AudioBusPatchOutput.Reset();
 			PatchInput.Reset();
 			AudioBusChannels = INDEX_NONE;
-			BlockSizeFrames = 0;
+			SampleRate = InParams.OperatorSettings.GetSampleRate();
+			BlockSizeFrames = InParams.OperatorSettings.GetNumFramesPerBlock();
 			bWasUnderrunReported = false;
 
 			bool bHasEnvironmentVars = InParams.Environment.Contains<Audio::FDeviceId>(SourceInterface::Environment::DeviceID);
@@ -276,14 +315,25 @@ namespace Metasound
 				UE_LOG(LogMetaSound, Warning, TEXT("Audio bus reader node requires audio device ID '%s', audio mixer num output frames '%s' and transmitter id '%s' environment variables")
 					, *SourceInterface::Environment::DeviceID.ToString(), *SourceInterface::Environment::AudioMixerNumOutputFrames.ToString(), *SourceInterface::Environment::TransmitterID.ToString());
 			}
+
+			// Audio mixer sample rate is a newer addition to the set of required environment variables in UE 5.6.  Check separately
+			// and fall back to old behavior if it does not exist. 
+			if (InParams.Environment.Contains<float>(SourceInterface::Environment::AudioMixerSampleRate))
+			{
+				AudioMixerSampleRate = InParams.Environment.GetValue<float>(SourceInterface::Environment::AudioMixerSampleRate);
+			}
+			else
+			{
+				UE_LOG(LogMetaSound, Warning, TEXT("Audio bus reader node may not render correctly without the audio mixer sample rate '%s' environment variable"), *SourceInterface::Environment::AudioMixerSampleRate.ToString());
+				// Assume matching sample rate if environment variable is missing.
+				AudioMixerSampleRate = SampleRate;
+			}
 			
 			for(const FAudioBufferWriteRef& Buffer : AudioOutputs)
 			{
 				Buffer->Zero();
 			}
 			
-			SampleRate = InParams.OperatorSettings.GetSampleRate();
-			BlockSizeFrames = InParams.OperatorSettings.GetNumFramesPerBlock();
 		}
 
 	private:
@@ -292,16 +342,26 @@ namespace Metasound
 			return AudioBusReaderNodeInitialNumBlocks(BlockSizeFrames, AudioMixerOutputFrames);
 		}
 
-		int32 NumBlocksToNumSamples(int32 NumBlocks) const
+		int32 NumBlocksToNumSamples(int32 NumBlocks, float InSampleRateRatio=-1.f) const
 		{
-			return NumBlocks * BlockSizeFrames * AudioBusChannels;
+			if (InSampleRateRatio > 0.f)
+			{
+				return FMath::CeilToInt(NumBlocks * BlockSizeFrames * InSampleRateRatio) * AudioBusChannels;
+			}
+			else
+			{
+				return NumBlocks * BlockSizeFrames * AudioBusChannels;
+			}
 		}
 
+
 		FAudioBusAssetReadRef AudioBusAsset;
 		TArray<FAudioBufferWriteRef> AudioOutputs;
 
-		TArray<float> InterleavedBuffer;
+		Audio::FAlignedFloatBuffer InterleavedBuffer;
+		TUniquePtr<AudioBusPrivate::FResampledPatchOutput> ResampledPatchOutput;
 		int32 AudioMixerOutputFrames = INDEX_NONE;
+		float AudioMixerSampleRate = -1.f;
 		Audio::FDeviceId AudioDeviceId = INDEX_NONE;
 		uint64 InstanceID = 0;
 		uint32 AudioBusId = 0;
diff --git a/Engine/Plugins/Runtime/Metasound/Source/MetasoundEngine/Private/MetasoundAudioBusWriterNode.cpp b/Engine/Plugins/Runtime/Metasound/Source/MetasoundEngine/Private/MetasoundAudioBusWriterNode.cpp
--- a/Engine/Plugins/Runtime/Metasound/Source/MetasoundEngine/Private/MetasoundAudioBusWriterNode.cpp
+++ b/Engine/Plugins/Runtime/Metasound/Source/MetasoundEngine/Private/MetasoundAudioBusWriterNode.cpp
@@ -9,6 +9,7 @@
 #include "MediaPacket.h"
 #include "MetasoundAudioBuffer.h"
 #include "MetasoundAudioBus.h"
+#include "MetasoundAudioBusPrivate.h"
 #include "MetasoundExecutableOperator.h"
 #include "MetasoundFacade.h"
 #include "MetasoundNodeRegistrationMacro.h"
@@ -140,6 +141,8 @@ namespace Metasound
 
 		void CreatePatchInput()
 		{
+			using namespace AudioBusPrivate;
+
 			const FAudioBusProxyPtr& AudioBusProxy = AudioBusAsset->GetAudioBusProxy();
 			if (AudioBusProxy.IsValid())
 			{
@@ -173,8 +176,26 @@ namespace Metasound
 
 				AudioBusPatchInput = AudioBusSubsystem->AddPatchInputForSoundAndAudioBus(InstanceID, AudioBusKey, BlockSizeFrames, int32(AudioBusChannels));
 				int32 NumBlocksToPush = InitialNumBlocks();
-				if (NumBlocksToPush > 0)
+
+				// Handle case of mismatched sample rate between audio mixer and
+				// metasound. 
+				if (EnableResampledAudioBus && (AudioMixerSampleRate != SampleRate) && (AudioMixerSampleRate > 0.f) && (SampleRate > 0.f))
 				{
+					UE_LOG(LogMetaSound, Warning, TEXT("Using a audio bus writer node is inefficient if the MetaSound sample rate %f does not match the AudioMixer sample rate %f. Please update MetaSound SampleRate to match the AudioMixer's SampleRate"), SampleRate, AudioMixerSampleRate);
+
+					if (NumBlocksToPush)
+					{
+						// Sample rate matches between audio mixer and metasound. This
+						// node will produce audio to the patch in approximate block
+						// sizes of (BlockSize * SampleRate / AudioMixerSampleRate)
+						AudioBusPatchInput.PushAudio(nullptr, NumBlocksToNumSamples(NumBlocksToPush, SampleRate / AudioMixerSampleRate));
+					}
+
+					ResampledPatchInput = MakeUnique<FResampledPatchInput>(AudioBusChannels, AudioMixerSampleRate, SampleRate, BlockSizeFrames, AudioBusPatchInput);
+				}
+				else  if (NumBlocksToPush > 0)
+				{
+					// Sample rate matches between audio mixer and metasound. 
 					AudioBusPatchInput.PushAudio(nullptr, NumBlocksToNumSamples(NumBlocksToPush));
 				}
 
@@ -192,12 +213,15 @@ namespace Metasound
 
 			InterleavedBuffer.Reset();
 			AudioMixerOutputFrames = INDEX_NONE;
+			AudioMixerSampleRate = -1.f;
 			AudioDeviceId = INDEX_NONE;
 			SampleRate = InParams.OperatorSettings.GetSampleRate();
 			AudioBusPatchInput.Reset();
+			ResampledPatchInput.Reset();
 			AudioBusChannels = INDEX_NONE;
 			AudioBusId = 0;
 			InstanceID = 0;
+			SampleRate = InParams.OperatorSettings.GetSampleRate();
 			BlockSizeFrames = InParams.OperatorSettings.GetNumFramesPerBlock();
 			bWasUnderrunReported = false;
 
@@ -216,6 +240,19 @@ namespace Metasound
 				UE_LOG(LogMetaSound, Warning, TEXT("Audio bus writer node requires audio device ID '%s', audio mixer num output frames '%s' and transmitter id '%s' environment variables")
 					, *SourceInterface::Environment::DeviceID.ToString(), *SourceInterface::Environment::AudioMixerNumOutputFrames.ToString(), *SourceInterface::Environment::TransmitterID.ToString());
 			}
+
+			// Audio mixer sample rate is a newer addition to the set of required environment variables in UE 5.6.  Check separately
+			// and fall back to old behavior if it does not exist. 
+			if (InParams.Environment.Contains<float>(SourceInterface::Environment::AudioMixerSampleRate))
+			{
+				AudioMixerSampleRate = InParams.Environment.GetValue<float>(SourceInterface::Environment::AudioMixerSampleRate);
+			}
+			else
+			{
+				UE_LOG(LogMetaSound, Warning, TEXT("Audio bus writer node may not render correctly without the audio mixer sample rate '%s' environment variable"), *SourceInterface::Environment::AudioMixerSampleRate.ToString());
+				// Assume matching sample rate if environment variable is missing.
+				AudioMixerSampleRate = SampleRate;
+			}
 		}
 
 
@@ -291,8 +328,19 @@ namespace Metasound
 				}
 			}
 
-			// Pushes the interleaved data to the audio bus
-			const int32 SamplesPushed = AudioBusPatchInput.PushAudio(InterleavedBuffer.GetData(), InterleavedBuffer.Num());
+			int32 SamplesPushed = -1;
+			if (ResampledPatchInput.IsValid())
+			{
+				// Perform resampling when pushing audio in case where audio mixer sample rate
+				// does not match the metasound sample rate. 
+				SamplesPushed = ResampledPatchInput->PushAudio(InterleavedBuffer.GetData(), InterleavedBuffer.Num());
+			}
+			else
+			{
+				// Pushes the interleaved data to the audio bus
+				SamplesPushed = AudioBusPatchInput.PushAudio(InterleavedBuffer.GetData(), InterleavedBuffer.Num());
+			}
+
 			if (SamplesPushed < InterleavedBuffer.Num() && !bWasUnderrunReported)
 			{
 				UE_LOG(LogMetaSound, Warning, TEXT("Underrun detected in audio bus writer node."));
@@ -306,16 +354,25 @@ namespace Metasound
 			return AudioBusWriterNodeInitialNumBlocks(BlockSizeFrames, AudioMixerOutputFrames);
 		}
 
-		int32 NumBlocksToNumSamples(int32 NumBlocks) const
+		int32 NumBlocksToNumSamples(int32 NumBlocks, float InSampleRateRatio=-1.f) const
 		{
-			return NumBlocks * BlockSizeFrames * AudioBusChannels;
+			if (InSampleRateRatio > 0.f)
+			{
+				return FMath::CeilToInt(NumBlocks * BlockSizeFrames * InSampleRateRatio) * AudioBusChannels;
+			}
+			else
+			{
+				return NumBlocks * BlockSizeFrames * AudioBusChannels;
+			}
 		}
 
 		FAudioBusAssetReadRef AudioBusAsset;
 		TArray<FAudioBufferReadRef> AudioInputs;
 
-		TArray<float> InterleavedBuffer;
+		Audio::FAlignedFloatBuffer InterleavedBuffer;
+		TUniquePtr<AudioBusPrivate::FResampledPatchInput> ResampledPatchInput;
 		int32 AudioMixerOutputFrames = INDEX_NONE;
+		float AudioMixerSampleRate = -1.f;
 		Audio::FDeviceId AudioDeviceId = INDEX_NONE;
 		float SampleRate = 0.0f;
 		Audio::FPatchInput AudioBusPatchInput;
diff --git a/Engine/Plugins/Runtime/Metasound/Source/MetasoundEngine/Private/MetasoundSource.cpp b/Engine/Plugins/Runtime/Metasound/Source/MetasoundEngine/Private/MetasoundSource.cpp
--- a/Engine/Plugins/Runtime/Metasound/Source/MetasoundEngine/Private/MetasoundSource.cpp
+++ b/Engine/Plugins/Runtime/Metasound/Source/MetasoundEngine/Private/MetasoundSource.cpp
@@ -1568,6 +1568,7 @@ Metasound::FMetasoundEnvironment UMetaSoundSource::CreateEnvironment(const FSoun
 	Environment.SetValue<uint64>(SourceInterface::Environment::TransmitterID, InParams.InstanceID);
 	Environment.SetValue<Audio::FDeviceId>(SourceInterface::Environment::DeviceID, InParams.AudioDeviceID);
 	Environment.SetValue<int32>(SourceInterface::Environment::AudioMixerNumOutputFrames, InParams.AudioMixerNumOutputFrames);
+	Environment.SetValue<float>(SourceInterface::Environment::AudioMixerSampleRate, InParams.SampleRate);
 
 #if WITH_METASOUND_DEBUG_ENVIRONMENT
 	Environment.SetValue<FString>(SourceInterface::Environment::GraphName, GetFullName());
diff --git a/Engine/Plugins/Runtime/Metasound/Source/MetasoundFrontend/Private/Interfaces/MetasoundFrontendSourceInterface.cpp b/Engine/Plugins/Runtime/Metasound/Source/MetasoundFrontend/Private/Interfaces/MetasoundFrontendSourceInterface.cpp
--- a/Engine/Plugins/Runtime/Metasound/Source/MetasoundFrontend/Private/Interfaces/MetasoundFrontendSourceInterface.cpp
+++ b/Engine/Plugins/Runtime/Metasound/Source/MetasoundFrontend/Private/Interfaces/MetasoundFrontendSourceInterface.cpp
@@ -237,6 +237,7 @@ namespace Metasound::Frontend
 			const FName SoundUniqueID = AUDIO_PARAMETER_INTERFACE_MEMBER_DEFINE("SoundUniqueID");
 			const FName TransmitterID = AUDIO_PARAMETER_INTERFACE_MEMBER_DEFINE("TransmitterID");
 			const FName AudioMixerNumOutputFrames = AUDIO_PARAMETER_INTERFACE_MEMBER_DEFINE("AudioMixerNumOutputFrames");
+			const FName AudioMixerSampleRate = AUDIO_PARAMETER_INTERFACE_MEMBER_DEFINE("AudioMixerSampleRate");
 		}
 
 		Audio::FParameterInterfacePtr CreateInterface(const UClass& InClass)
@@ -296,6 +297,12 @@ namespace Metasound::Frontend
 							LOCTEXT("AudioMixerOutputFramesDescription", "The number of output frames rendered per buffer in the audio mixer."),
 							FName(), 
 							Environment::AudioMixerNumOutputFrames
+						},
+						{
+							LOCTEXT("AudioMixerSampleRateDisplayName", "Audio Mixer Sample Rate"),
+							LOCTEXT("AudioMixerSampleRateDescription", "The sample rate of the audio mixer."),
+							FName(), 
+							Environment::AudioMixerSampleRate
 						}
 					};
 				}
diff --git a/Engine/Plugins/Runtime/Metasound/Source/MetasoundFrontend/Public/Interfaces/MetasoundFrontendSourceInterface.h b/Engine/Plugins/Runtime/Metasound/Source/MetasoundFrontend/Public/Interfaces/MetasoundFrontendSourceInterface.h
--- a/Engine/Plugins/Runtime/Metasound/Source/MetasoundFrontend/Public/Interfaces/MetasoundFrontendSourceInterface.h
+++ b/Engine/Plugins/Runtime/Metasound/Source/MetasoundFrontend/Public/Interfaces/MetasoundFrontendSourceInterface.h
@@ -81,6 +81,7 @@ namespace Metasound
 				METASOUNDFRONTEND_API const extern FName SoundUniqueID;
 				METASOUNDFRONTEND_API const extern FName TransmitterID;
 				METASOUNDFRONTEND_API const extern FName AudioMixerNumOutputFrames;
+				METASOUNDFRONTEND_API const extern FName AudioMixerSampleRate;
 			}
 
 			METASOUNDFRONTEND_API const FMetasoundFrontendVersion& GetVersion();
