Fix for missing disconnect call when swapping MetaSound Graph Output DataType.
- Minor: Remove redundant type update
[RN] Fix for swapping graph output data type on paged graphs in MetaSound failing to remove invalid connections causing graph to output runtime errors on execution.

diff --git a/Engine/Plugins/Runtime/Metasound/Source/MetasoundFrontend/Private/MetasoundFrontendDocumentBuilder.cpp b/Engine/Plugins/Runtime/Metasound/Source/MetasoundFrontend/Private/MetasoundFrontendDocumentBuilder.cpp
--- a/Engine/Plugins/Runtime/Metasound/Source/MetasoundFrontend/Private/MetasoundFrontendDocumentBuilder.cpp
+++ b/Engine/Plugins/Runtime/Metasound/Source/MetasoundFrontend/Private/MetasoundFrontendDocumentBuilder.cpp
@@ -5647,7 +5647,6 @@ bool FMetaSoundFrontendDocumentBuilder::SetGraphOutputDataType(FName OutputName,
 	if (GraphOutput.TypeName != DataType)
 	{
 		GraphOutput.TypeName = DataType;
-
 		RootGraph.IterateGraphPages([this, &GraphOutput, &DataType](FMetasoundFrontendGraph& Graph)
 		{
 			const IDocumentGraphNodeCache& NodeCache = DocumentCache->GetNodeCache(Graph.PageID);
@@ -5660,8 +5659,7 @@ bool FMetaSoundFrontendDocumentBuilder::SetGraphOutputDataType(FName OutputName,
 				FMetasoundFrontendVertex& NodeInput = Node.Interface.Inputs.Last();
 				Node.InputLiterals = { FMetasoundFrontendVertexLiteral { NodeInput.VertexID, MoveTemp(DefaultLiteral) } };
 
-				RemoveEdgeToNodeInput(GraphOutput.NodeID, NodeInput.VertexID);
-				GraphOutput.TypeName = DataType;
+				RemoveEdgeToNodeInput(GraphOutput.NodeID, NodeInput.VertexID, &Graph.PageID);
 			}
 		});
 
