Remove Slate depdency from MetasoundFrontend

diff --git a/Engine/Plugins/Experimental/MetasoundExperimental/Source/MetasoundExperimentalEditor/MetasoundExperimentalEditor.build.cs b/Engine/Plugins/Experimental/MetasoundExperimental/Source/MetasoundExperimentalEditor/MetasoundExperimentalEditor.build.cs
--- a/Engine/Plugins/Experimental/MetasoundExperimental/Source/MetasoundExperimentalEditor/MetasoundExperimentalEditor.build.cs
+++ b/Engine/Plugins/Experimental/MetasoundExperimental/Source/MetasoundExperimentalEditor/MetasoundExperimentalEditor.build.cs
@@ -24,21 +24,10 @@ public class MetasoundExperimentalEditor : ModuleRules
 		);
 
 
-		PublicDependencyModuleNames.AddRange(
-			new string[]
-			{
-				"Core",
-				// ... add other public dependencies that you statically link with here ...
-			}
-		);
-
 		PrivateDependencyModuleNames.AddRange(
 			new string[]
 			{
-				"CoreUObject",
-				"Engine",
-
-				// ... add private dependencies that you statically link with here ...	
+				"Core",
 			}
 		);
 
diff --git a/Engine/Plugins/Runtime/Metasound/Source/MetasoundEditor/Private/MetasoundEditorGraphCommentNode.cpp b/Engine/Plugins/Runtime/Metasound/Source/MetasoundEditor/Private/MetasoundEditorGraphCommentNode.cpp
--- a/Engine/Plugins/Runtime/Metasound/Source/MetasoundEditor/Private/MetasoundEditorGraphCommentNode.cpp
+++ b/Engine/Plugins/Runtime/Metasound/Source/MetasoundEditor/Private/MetasoundEditorGraphCommentNode.cpp
@@ -43,8 +43,8 @@ void UMetasoundEditorGraphCommentNode::ConvertToFrontendComment(const UEdGraphNo
 	OutComment.Depth = InEdNode.CommentDepth;
 	OutComment.FontSize = InEdNode.FontSize;
 	OutComment.MoveMode = InEdNode.MoveMode == ECommentBoxMode::GroupMovement ? EMetaSoundFrontendGraphCommentMoveMode::GroupMovement : EMetaSoundFrontendGraphCommentMoveMode::NoGroupMovement;
-	OutComment.Position = FVector2D(InEdNode.NodePosX, InEdNode.NodePosY);
-	OutComment.Size = FVector2D(InEdNode.NodeWidth, InEdNode.NodeHeight);
+	OutComment.Position = FVector2f(InEdNode.NodePosX, InEdNode.NodePosY);
+	OutComment.Size = FVector2f(InEdNode.NodeWidth, InEdNode.NodeHeight);
 }
 
 void UMetasoundEditorGraphCommentNode::ConvertFromFrontendComment(const FMetaSoundFrontendGraphComment& InComment, UEdGraphNode_Comment& OutEdNode)
@@ -131,7 +131,7 @@ void UMetasoundEditorGraphCommentNode::UpdateFrontendNodeLocation()
 {
 	using namespace Metasound::Frontend;
 	FMetaSoundFrontendGraphComment& FrontendComment = GetBuilderChecked().FindOrAddGraphComment(CommentID);
-	FrontendComment.Position = FVector2D(NodePosX, NodePosY);
-	FrontendComment.Size = FVector2D(NodeWidth, NodeHeight);
+	FrontendComment.Position = FVector2f(NodePosX, NodePosY);
+	FrontendComment.Size = FVector2f(NodeWidth, NodeHeight);
 }
 #undef LOCTEXT_NAMESPACE
diff --git a/Engine/Plugins/Runtime/Metasound/Source/MetasoundFrontend/MetasoundFrontend.Build.cs b/Engine/Plugins/Runtime/Metasound/Source/MetasoundFrontend/MetasoundFrontend.Build.cs
--- a/Engine/Plugins/Runtime/Metasound/Source/MetasoundFrontend/MetasoundFrontend.Build.cs
+++ b/Engine/Plugins/Runtime/Metasound/Source/MetasoundFrontend/MetasoundFrontend.Build.cs
@@ -23,7 +23,6 @@ namespace UnrealBuildTool.Rules
 					"Serialization",
 					"SignalProcessing",
 					"MetasoundGraphCore",
-					"SlateCore"
 				}
 			);
 
diff --git a/Engine/Plugins/Runtime/Metasound/Source/MetasoundFrontend/Public/MetasoundFrontendDocument.h b/Engine/Plugins/Runtime/Metasound/Source/MetasoundFrontend/Public/MetasoundFrontendDocument.h
--- a/Engine/Plugins/Runtime/Metasound/Source/MetasoundFrontend/Public/MetasoundFrontendDocument.h
+++ b/Engine/Plugins/Runtime/Metasound/Source/MetasoundFrontend/Public/MetasoundFrontendDocument.h
@@ -200,10 +200,10 @@ struct FMetaSoundFrontendGraphComment
 	int32 FontSize = 0;
 
 	UPROPERTY()
-	FDeprecateSlateVector2D Position = FVector2f::ZeroVector;
+	FVector2f Position = FVector2f::ZeroVector;
 
 	UPROPERTY()
-	FDeprecateSlateVector2D Size = FVector2f::ZeroVector;
+	FVector2f Size = FVector2f::ZeroVector;
 
 	UPROPERTY()
 	EMetaSoundFrontendGraphCommentMoveMode MoveMode = EMetaSoundFrontendGraphCommentMoveMode::GroupMovement;
