[Metasound Engine] Metasound Interfaces with incremented minor versions auto update metasound assets that implement them

[RN] Incrementing the minor version of Metasound Interfaces auto-updates any metasound assetss that implement them

diff --git a/Engine/Plugins/Runtime/Metasound/Source/MetasoundEngine/Private/Interfaces/MetasoundInterface.cpp b/Engine/Plugins/Runtime/Metasound/Source/MetasoundEngine/Private/Interfaces/MetasoundInterface.cpp
--- a/Engine/Plugins/Runtime/Metasound/Source/MetasoundEngine/Private/Interfaces/MetasoundInterface.cpp
+++ b/Engine/Plugins/Runtime/Metasound/Source/MetasoundEngine/Private/Interfaces/MetasoundInterface.cpp
@@ -69,9 +69,20 @@ namespace Metasound::Engine
 		{
 			return UpdateTransform->Transform(InDocument);
 		}
+		
 		return false;
 	}
 
+	bool FInterfaceRegistryEntry::UpdateRootGraphInterface(FMetaSoundFrontendDocumentBuilder& InDocumentBuilder) const
+	{
+		if (UpdateTransform.IsValid())
+		{
+			return UpdateRootGraphInterface(InDocumentBuilder.GetMetasoundAsset().GetDocumentHandle());
+		}
+
+		return InDocumentBuilder.AddInterface(GetInterface().Metadata.Version.Name);
+	}
+
 	void RegisterAudioFormatInterfaces()
 	{
 		using namespace Frontend;
diff --git a/Engine/Plugins/Runtime/Metasound/Source/MetasoundEngine/Private/Interfaces/MetasoundInterface.h b/Engine/Plugins/Runtime/Metasound/Source/MetasoundEngine/Private/Interfaces/MetasoundInterface.h
--- a/Engine/Plugins/Runtime/Metasound/Source/MetasoundEngine/Private/Interfaces/MetasoundInterface.h
+++ b/Engine/Plugins/Runtime/Metasound/Source/MetasoundEngine/Private/Interfaces/MetasoundInterface.h
@@ -4,6 +4,7 @@
 #include "IAudioParameterInterfaceRegistry.h"
 #include "Interfaces/MetasoundFrontendInterfaceRegistry.h"
 #include "MetasoundFrontendDocument.h"
+#include "MetasoundFrontendDocumentBuilder.h"
 #include "MetasoundFrontendTransform.h"
 
 
@@ -21,6 +22,7 @@ namespace Metasound::Engine
 		virtual const FMetasoundFrontendInterface& GetInterface() const override;
 		virtual bool IsDeprecated() const override;
 		virtual bool UpdateRootGraphInterface(Frontend::FDocumentHandle InDocument) const override;
+		virtual bool UpdateRootGraphInterface(FMetaSoundFrontendDocumentBuilder& InDocumentBuilder) const override;
 
 	private:
 		FMetasoundFrontendInterface Interface;
diff --git a/Engine/Plugins/Runtime/Metasound/Source/MetasoundEngineTest/Private/Interfaces/MetasoundTestInterfaces.cpp b/Engine/Plugins/Runtime/Metasound/Source/MetasoundEngineTest/Private/Interfaces/MetasoundTestInterfaces.cpp
new file mode 100644
--- /dev/null
+++ b/Engine/Plugins/Runtime/Metasound/Source/MetasoundEngineTest/Private/Interfaces/MetasoundTestInterfaces.cpp
@@ -0,0 +1,146 @@
+// Copyright Epic Games, Inc. All Rights Reserved.
+
+#include "MetasoundTestInterfaces.h"
+
+#include "MetasoundFrontendDocument.h"
+#include "MetasoundPrimitives.h"
+#include "MetasoundTrigger.h"
+
+#define LOCTEXT_NAMESPACE "MetasoundTestInterfaces"
+
+namespace Metasound::Test
+{
+#define AUDIO_PARAMETER_INTERFACE_NAMESPACE "UE.Test.Update"
+	namespace UpdateTestInterface_0_1
+	{
+		const FMetasoundFrontendVersion& GetVersion()
+		{
+			static const FMetasoundFrontendVersion Version = { AUDIO_PARAMETER_INTERFACE_NAMESPACE, { 0, 1 } };
+			return Version;
+		}
+
+		namespace Inputs
+		{
+			const FName InputTrigger = AUDIO_PARAMETER_INTERFACE_MEMBER_DEFINE("InputTrigger");
+		}
+
+		namespace Outputs
+		{
+			const FName OutputTrigger = AUDIO_PARAMETER_INTERFACE_MEMBER_DEFINE("OutputTrigger");
+		}
+
+		Audio::FParameterInterfacePtr CreateInterface(const UClass& InClass)
+		{
+			struct FInterface : public Audio::FParameterInterface
+			{
+				FInterface(const FTopLevelAssetPath& InClassPath)
+					: FParameterInterface(UpdateTestInterface_0_1::GetVersion().Name, UpdateTestInterface_0_1::GetVersion().Number.ToInterfaceVersion())
+				{
+					constexpr bool bIsModifiable = true;
+					constexpr bool bIsDefault = false;
+					UClassOptions = TArray<FClassOptions>
+					{
+						{ InClassPath, bIsModifiable, bIsDefault }
+					};
+
+					Inputs =
+					{
+						{
+							LOCTEXT("InputTrigger", "Input Trigger"),
+							LOCTEXT("InputTrigger Description", "A trigger input used for testing Metasound Interfaces."),
+							GetMetasoundDataTypeName<FTrigger>(),
+							Inputs::InputTrigger
+						}
+					};
+
+					Outputs =
+					{
+						{
+							LOCTEXT("OutputTrigger", "Output Trigger"),
+							LOCTEXT("OutputTrigger Description", "A trigger output used for testing Metasound Interfaces."),
+							GetMetasoundDataTypeName<FTrigger>(),
+							Outputs::OutputTrigger
+						}
+					};
+				}
+			};
+
+			return MakeShared<FInterface>(InClass.GetClassPathName());
+		}
+	}
+
+	namespace UpdateTestInterface_0_2
+	{
+		const FMetasoundFrontendVersion& GetVersion()
+		{
+			static const FMetasoundFrontendVersion Version = { AUDIO_PARAMETER_INTERFACE_NAMESPACE, { 0, 2 } };
+			return Version;
+		}
+
+		namespace Inputs
+		{
+			const FName InputTrigger = AUDIO_PARAMETER_INTERFACE_MEMBER_DEFINE("InputTrigger");
+			const FName InputFloat = AUDIO_PARAMETER_INTERFACE_MEMBER_DEFINE("InputFloat");
+		}
+
+		namespace Outputs
+		{
+			const FName OutputTrigger = AUDIO_PARAMETER_INTERFACE_MEMBER_DEFINE("OutputTrigger");
+			const FName OutputFloat = AUDIO_PARAMETER_INTERFACE_MEMBER_DEFINE("OutputFloat");
+		}
+
+		Audio::FParameterInterfacePtr CreateInterface(const UClass& InClass)
+		{
+			struct FInterface : public Audio::FParameterInterface
+			{
+				FInterface(const FTopLevelAssetPath& InClassPath)
+					: FParameterInterface(UpdateTestInterface_0_2::GetVersion().Name, UpdateTestInterface_0_2::GetVersion().Number.ToInterfaceVersion())
+				{
+					constexpr bool bIsModifiable = true;
+					constexpr bool bIsDefault = false;
+					UClassOptions = TArray<FClassOptions>
+					{
+						{ InClassPath, bIsModifiable, bIsDefault }
+					};
+
+					Inputs =
+					{
+						{
+							LOCTEXT("InputTrigger", "Input Trigger"),
+							LOCTEXT("InputTrigger Description", "A trigger input used for testing Metasound Interfaces."),
+							GetMetasoundDataTypeName<FTrigger>(),
+							Inputs::InputTrigger
+						},
+						{
+							LOCTEXT("InputFloat", "Input Float"),
+							LOCTEXT("InputFloat Description", "A float input used for testing Metasound Interfaces."),
+							GetMetasoundDataTypeName<float>(),
+							Inputs::InputFloat
+						}
+					};
+
+					Outputs =
+					{
+						{
+							LOCTEXT("OutputTrigger", "Output Trigger"),
+							LOCTEXT("OutputTrigger Description", "A trigger output used for testing Metasound Interfaces."),
+							GetMetasoundDataTypeName<FTrigger>(),
+							Outputs::OutputTrigger
+						},
+						{
+							LOCTEXT("OutputFloat", "Output Float"),
+							LOCTEXT("OutputFloat Description", "A float output used for testing Metasound Interfaces."),
+							GetMetasoundDataTypeName<float>(),
+							Outputs::OutputFloat
+						}
+					};
+				}
+			};
+
+			return MakeShared<FInterface>(InClass.GetClassPathName());
+		}
+	}
+#undef AUDIO_PARAMETER_INTERFACE_NAMESPACE
+}
+
+#undef LOCTEXT_NAMESPACE
\ No newline at end of file
diff --git a/Engine/Plugins/Runtime/Metasound/Source/MetasoundEngineTest/Private/Interfaces/MetasoundTestInterfaces.h b/Engine/Plugins/Runtime/Metasound/Source/MetasoundEngineTest/Private/Interfaces/MetasoundTestInterfaces.h
new file mode 100644
--- /dev/null
+++ b/Engine/Plugins/Runtime/Metasound/Source/MetasoundEngineTest/Private/Interfaces/MetasoundTestInterfaces.h
@@ -0,0 +1,48 @@
+// Copyright Epic Games, Inc. All Rights Reserved.
+#pragma once
+
+#include "UObject/NameTypes.h"
+#include "IAudioParameterInterfaceRegistry.h"
+
+struct FMetasoundFrontendVersion;
+
+namespace Metasound::Test
+{
+#define AUDIO_PARAMETER_INTERFACE_NAMESPACE "UE.Test.Update"
+	namespace UpdateTestInterface_0_1
+	{
+		const FMetasoundFrontendVersion& GetVersion();
+
+		namespace Inputs
+		{
+			const extern FName InputTrigger;
+		}
+
+		namespace Outputs
+		{
+			const extern FName OutputTrigger;
+		}
+
+		Audio::FParameterInterfacePtr CreateInterface(const UClass& InClass);
+	}
+
+	namespace UpdateTestInterface_0_2
+	{
+		const FMetasoundFrontendVersion& GetVersion();
+
+		namespace Inputs
+		{
+			const extern FName InputTrigger;
+			const extern FName InputFloat;
+		}
+
+		namespace Outputs
+		{
+			const extern FName OutputTrigger;
+			const extern FName OutputFloat;
+		}
+
+		Audio::FParameterInterfacePtr CreateInterface(const UClass& InClass);
+	}
+#undef AUDIO_PARAMETER_INTERFACE_NAMESPACE
+}
\ No newline at end of file
diff --git a/Engine/Plugins/Runtime/Metasound/Source/MetasoundEngineTest/Private/MetasoundEngineTestModule.cpp b/Engine/Plugins/Runtime/Metasound/Source/MetasoundEngineTest/Private/MetasoundEngineTestModule.cpp
--- a/Engine/Plugins/Runtime/Metasound/Source/MetasoundEngineTest/Private/MetasoundEngineTestModule.cpp
+++ b/Engine/Plugins/Runtime/Metasound/Source/MetasoundEngineTest/Private/MetasoundEngineTestModule.cpp
@@ -1,5 +1,8 @@
 // Copyright Epic Games, Inc. All Rights Reserved.
 
+#include "IAudioParameterInterfaceRegistry.h"
+#include "Interfaces/MetasoundTestInterfaces.h"
+#include "MetasoundSource.h"
 #include "Modules/ModuleInterface.h"
 #include "Modules/ModuleManager.h"
 
@@ -13,6 +16,12 @@ namespace Metasound
 			FModuleManager::Get().LoadModuleChecked("MetasoundFrontend");
 			FModuleManager::Get().LoadModuleChecked("MetasoundStandardNodes");
 			FModuleManager::Get().LoadModuleChecked("MetasoundEngine");
+
+			using namespace Metasound::Test;
+
+			Audio::IAudioParameterInterfaceRegistry& Registry = Audio::IAudioParameterInterfaceRegistry::Get();
+			Registry.RegisterInterface(UpdateTestInterface_0_1::CreateInterface(*UMetaSoundSource::StaticClass()));
+			Registry.RegisterInterface(UpdateTestInterface_0_2::CreateInterface(*UMetaSoundSource::StaticClass()));
 		}
 	};
 }
diff --git a/Engine/Plugins/Runtime/Metasound/Source/MetasoundEngineTest/Private/MetasoundInterfaceUpdateTests.cpp b/Engine/Plugins/Runtime/Metasound/Source/MetasoundEngineTest/Private/MetasoundInterfaceUpdateTests.cpp
new file mode 100644
--- /dev/null
+++ b/Engine/Plugins/Runtime/Metasound/Source/MetasoundEngineTest/Private/MetasoundInterfaceUpdateTests.cpp
@@ -0,0 +1,172 @@
+// Copyright Epic Games, Inc. All Rights Reserved.
+
+#include "IAudioParameterInterfaceRegistry.h"
+#include "Interfaces/MetasoundTestInterfaces.h"
+#include "MetasoundBuilderSubsystem.h"
+#include "MetasoundSource.h"
+#include "MetasoundStandardNodesNames.h"
+#include "Misc/AutomationTest.h"
+
+#if WITH_DEV_AUTOMATION_TESTS
+
+namespace Metasound::Test::InterfaceUpdate
+{
+	struct FInitTestBuilderSourceOutput
+	{
+		FMetaSoundBuilderNodeOutputHandle OnPlayOutput;
+		FMetaSoundBuilderNodeInputHandle OnFinishedInput;
+		TArray<FMetaSoundBuilderNodeInputHandle> AudioOutNodeInputs;
+	};
+
+	UMetaSoundSourceBuilder& CreateSourceBuilder(
+		FAutomationTestBase& Test,
+		EMetaSoundOutputAudioFormat OutputFormat,
+		bool bIsOneShot,
+		FInitTestBuilderSourceOutput& Output)
+	{
+		using namespace Audio;
+		using namespace Metasound;
+		using namespace Metasound::Engine;
+		using namespace Metasound::Frontend;
+
+		EMetaSoundBuilderResult Result;
+		UMetaSoundSourceBuilder* Builder = UMetaSoundBuilderSubsystem::GetChecked().CreateSourceBuilder(
+			"Unit Test Graph Builder",
+			Output.OnPlayOutput,
+			Output.OnFinishedInput,
+			Output.AudioOutNodeInputs,
+			Result,
+			OutputFormat,
+			bIsOneShot);
+		checkf(Builder, TEXT("Failed to create MetaSoundSourceBuilder"));
+		Test.AddErrorIfFalse(Result == EMetaSoundBuilderResult::Succeeded, TEXT("Builder created but CreateSourceBuilder did not result in 'Succeeded' state"));
+
+		return *Builder;
+	}
+
+	IMPLEMENT_SIMPLE_AUTOMATION_TEST(
+		FMetasoundInterfaceMinorUpdateTest,
+		"Audio.Metasound.Interface.MinorVersionUpdate",
+		EAutomationTestFlags::EditorContext | EAutomationTestFlags::EngineFilter)
+		bool FMetasoundInterfaceMinorUpdateTest::RunTest(const FString&)
+	{
+		using namespace Audio;
+		using namespace Metasound::Frontend;
+
+		FInitTestBuilderSourceOutput BuilderSourceOutput;
+		UMetaSoundSourceBuilder& MetasoundSourceBuilder = CreateSourceBuilder(*this, EMetaSoundOutputAudioFormat::Mono, true, BuilderSourceOutput);
+
+		FMetaSoundFrontendDocumentBuilder& MetasoundFrontendBuilder = MetasoundSourceBuilder.GetBuilder();
+
+		TArray<FMetasoundFrontendInterface> TestInterfaceVersions;
+
+		IAudioParameterInterfaceRegistry::Get().IterateInterfaces([&TestInterfaceVersions](FParameterInterfacePtr InterfacePtr)
+			{
+				if (InterfacePtr->GetName() == UpdateTestInterface_0_1::GetVersion().Name)
+				{
+					TestInterfaceVersions.Add(InterfacePtr);
+				}
+			});
+
+		TestEqual(TEXT("There should be two versions of the test interface"), TestInterfaceVersions.Num(), 2);
+
+		const FMetasoundFrontendInterface TestInterfaceV0_1 = TestInterfaceVersions[0];
+		const FMetasoundFrontendInterface TestInterfaceV0_2 = TestInterfaceVersions[1];
+
+		TestEqual("Test Interface V0.1 version numbers should be correct", TestInterfaceV0_1.Metadata.Version.Number.Major, 0);
+		TestEqual("Test Interface V0.1 version numbers should be correct", TestInterfaceV0_1.Metadata.Version.Number.Minor, 1);
+		TestEqual("Test Interface V0.2 version numbers should be correct", TestInterfaceV0_2.Metadata.Version.Number.Major, 0);
+		TestEqual("Test Interface V0.2 version numbers should be correct", TestInterfaceV0_2.Metadata.Version.Number.Minor, 2);
+
+		{
+			FModifyInterfaceOptions ModifyInterfaceOptions({ }, { TestInterfaceV0_1 });
+
+			MetasoundFrontendBuilder.ModifyInterfaces(MoveTemp(ModifyInterfaceOptions));
+		}
+
+		const bool bHasInterfaceV0_1 = MetasoundFrontendBuilder.IsInterfaceDeclared(TestInterfaceV0_1.Metadata.Version);
+
+		TestTrue("Test Interface V0.1 should be declared at this point", bHasInterfaceV0_1);
+		
+		EMetaSoundBuilderResult BuilderResult{ EMetaSoundBuilderResult::Failed };
+
+		// This block creates a chain of Trigger Input -> Trigger Toggle -> SuperOscillator, Trigger Toggle -> Trigger Output
+		if (FMetaSoundNodeHandle SuperOscillatorNodeHandle = MetasoundSourceBuilder.AddNodeByClassName(FMetasoundFrontendClassName(StandardNodes::Namespace, FName("SuperOscillatorMono"), StandardNodes::AudioVariant), BuilderResult);
+			SuperOscillatorNodeHandle.IsSet() && BuilderResult == EMetaSoundBuilderResult::Succeeded)
+		{
+			FMetaSoundNodeHandle TriggerToggleNodeHandle = MetasoundSourceBuilder.AddNodeByClassName(FMetasoundFrontendClassName(StandardNodes::Namespace, FName("Trigger Toggle"), ""), BuilderResult);
+			AddErrorIfFalse(BuilderResult == EMetaSoundBuilderResult::Succeeded, "Failed to add trigger toggle node");
+
+			FMetaSoundBuilderNodeInputHandle TriggerToggleOnInput = MetasoundSourceBuilder.FindNodeInputByName(TriggerToggleNodeHandle, FName("On"), BuilderResult);
+			AddErrorIfFalse(BuilderResult == EMetaSoundBuilderResult::Succeeded, "Failed to find \"On\" input pin of trigger toggle node");
+
+			FMetaSoundNodeHandle InterfaceInputNodeHandle = MetasoundSourceBuilder.FindGraphInputNode(UpdateTestInterface_0_1::Inputs::InputTrigger, BuilderResult);
+			AddErrorIfFalse(BuilderResult == EMetaSoundBuilderResult::Succeeded, "Failed to find interface input node");
+			 
+			FMetaSoundBuilderNodeOutputHandle InterfaceInputNodeOutputHandle = MetasoundSourceBuilder.FindNodeOutputByName(InterfaceInputNodeHandle, UpdateTestInterface_0_1::Inputs::InputTrigger, BuilderResult);
+			AddErrorIfFalse(BuilderResult == EMetaSoundBuilderResult::Succeeded, "Failed to find interface input node output pin handle");
+
+			MetasoundSourceBuilder.ConnectNodes(InterfaceInputNodeOutputHandle, TriggerToggleOnInput, BuilderResult);
+			AddErrorIfFalse(BuilderResult == EMetaSoundBuilderResult::Succeeded, "Failed to connect interface trigger input pin to Trigger Toggle Node's \"On\" pin.");
+
+			FMetaSoundBuilderNodeInputHandle SuperOscEnabledInput = MetasoundSourceBuilder.FindNodeInputByName(SuperOscillatorNodeHandle, FName("Enabled"), BuilderResult);
+			AddErrorIfFalse(BuilderResult == EMetaSoundBuilderResult::Succeeded, "Failed to find \"Enabled\" input pin of SuperOscillator node");
+
+			FMetaSoundBuilderNodeOutputHandle TriggerToggleValueOutput = MetasoundSourceBuilder.FindNodeOutputByName(TriggerToggleNodeHandle, FName("Value"), BuilderResult);
+			AddErrorIfFalse(BuilderResult == EMetaSoundBuilderResult::Succeeded, "Failed to find \"Value\" output pin of trigger toggle node");
+
+			MetasoundSourceBuilder.ConnectNodes(TriggerToggleValueOutput, SuperOscEnabledInput, BuilderResult);
+			AddErrorIfFalse(BuilderResult == EMetaSoundBuilderResult::Succeeded, "Failed to connect Trigger Toggle Node's \"Value\" pin to SuperOscillator node's \"Enabled\" pin.");
+
+			FMetaSoundNodeHandle InterfaceOutputNodeHandle = MetasoundSourceBuilder.FindGraphOutputNode(UpdateTestInterface_0_1::Outputs::OutputTrigger, BuilderResult);
+			AddErrorIfFalse(BuilderResult == EMetaSoundBuilderResult::Succeeded, "Failed to find interface output node");
+
+			FMetaSoundBuilderNodeInputHandle InterfaceOutputNodeInputHandle = MetasoundSourceBuilder.FindNodeInputByName(InterfaceOutputNodeHandle, UpdateTestInterface_0_1::Outputs::OutputTrigger, BuilderResult);
+			AddErrorIfFalse(BuilderResult == EMetaSoundBuilderResult::Succeeded, "Failed to find interface input node output pin handle");
+
+			FMetaSoundBuilderNodeOutputHandle TriggerToggleOutOutput = MetasoundSourceBuilder.FindNodeOutputByName(TriggerToggleNodeHandle, FName("Out"), BuilderResult);
+			AddErrorIfFalse(BuilderResult == EMetaSoundBuilderResult::Succeeded, "Failed to find \"Out\" output pin of trigger toggle node");
+
+			MetasoundSourceBuilder.ConnectNodes(TriggerToggleOutOutput, InterfaceOutputNodeInputHandle, BuilderResult);
+			AddErrorIfFalse(BuilderResult == EMetaSoundBuilderResult::Succeeded, "Failed to connect Trigger Toggle Node's \"Out\" pin to graph output pin.");
+		}
+		else
+		{
+			AddError("Failed to add superoscillator node to test graph");
+			return false;
+		}
+
+		{
+			const FMetasoundFrontendClassInput* InputTriggerNode = MetasoundFrontendBuilder.FindGraphInput(UpdateTestInterface_0_1::Inputs::InputTrigger);
+			if (TestNotNull("Graph has the trigger input defined by interface V0.1", InputTriggerNode))
+			{
+				TestEqual("Input Trigger Node should have FTrigger as type", InputTriggerNode->TypeName, GetMetasoundDataTypeName<FTrigger>());
+			}
+
+			const FMetasoundFrontendClassOutput* OutputTriggerNode = MetasoundFrontendBuilder.FindGraphOutput(UpdateTestInterface_0_1::Outputs::OutputTrigger);
+			if (TestNotNull("Graph has the trigger output defined by interface V0.1", OutputTriggerNode))
+			{
+				TestEqual("Output Trigger Node should have FTrigger as type", OutputTriggerNode->TypeName, GetMetasoundDataTypeName<FTrigger>());
+			}
+		}
+
+		TScriptInterface<IMetaSoundDocumentInterface> MetasoundSourceInterface = MetasoundSourceBuilder.BuildNewMetaSound(FName("TestMetaSoundSource"));
+		if (!TestNotNull("Built MetasoundDocumentInterface is null", MetasoundSourceInterface.GetObject()))
+		{
+			return false;
+		}
+
+		TObjectPtr<UMetaSoundSource> MetasoundSource = Cast<UMetaSoundSource>(MetasoundSourceInterface.GetObject());
+		if (!TestNotNull("Cast to Metasound Source failed", MetasoundSource.Get()))
+		{
+			return false;
+		}
+
+		// This will internally throw an error if there are problems when versioning the asset
+		MetasoundSource->VersionAsset(MetasoundFrontendBuilder);
+
+		return true;
+	}
+}
+
+#endif
\ No newline at end of file
diff --git a/Engine/Plugins/Runtime/Metasound/Source/MetasoundFrontend/Private/MetasoundAssetBase.cpp b/Engine/Plugins/Runtime/Metasound/Source/MetasoundFrontend/Private/MetasoundAssetBase.cpp
--- a/Engine/Plugins/Runtime/Metasound/Source/MetasoundFrontend/Private/MetasoundAssetBase.cpp
+++ b/Engine/Plugins/Runtime/Metasound/Source/MetasoundFrontend/Private/MetasoundAssetBase.cpp
@@ -160,14 +160,14 @@ namespace Metasound
 				}
 			}
 
-			bool UpdateDocumentInterface(const TArray<const IInterfaceRegistryEntry*>& InUpgradePath, const FMetasoundFrontendVersion& InterfaceVersion, FDocumentHandle InDocument)
+			bool UpdateDocumentInterface(const TArray<const IInterfaceRegistryEntry*>& InUpgradePath, const FMetasoundFrontendVersion& InterfaceVersion, FMetaSoundFrontendDocumentBuilder& Builder)
 			{
 				const FMetasoundFrontendVersionNumber* LastVersionUpdated = nullptr;
 				for (const IInterfaceRegistryEntry* Entry : InUpgradePath)
 				{
 					if (ensure(nullptr != Entry))
 					{
-						if (Entry->UpdateRootGraphInterface(InDocument))
+						if (Entry->UpdateRootGraphInterface(Builder))
 						{
 							LastVersionUpdated = &Entry->GetInterface().Metadata.Version.Number;
 						}
@@ -176,10 +176,11 @@ namespace Metasound
 
 				if (LastVersionUpdated)
 				{
+					const FMetasoundFrontendClassMetadata& Metadata = Builder.GetMetasoundAsset().GetDocumentChecked().RootGraph.Metadata;
 #if WITH_EDITOR
-					const FString AssetName = *InDocument->GetRootGraphClass().Metadata.GetDisplayName().ToString();
+					const FString AssetName = *Metadata.GetDisplayName().ToString();
 #else
-					const FString AssetName = *InDocument->GetRootGraphClass().Metadata.GetClassName().ToString();
+					const FString AssetName = *Metadata.GetClassName().ToString();
 #endif // !WITH_EDITOR
 					UE_LOG(LogMetaSound, Display, TEXT("Asset '%s' interface '%s' updated: '%s' --> '%s'"),
 						*AssetName,
@@ -563,7 +564,7 @@ bool FMetasoundAssetBase::VersionAsset(FMetaSoundFrontendDocumentBuilder& Builde
 
 			for (const FMetasoundFrontendVersion& Version : Versions)
 			{
-				bPassUpdated |= TryUpdateInterfaceFromVersion(Version);
+				bPassUpdated |= TryUpdateInterfaceFromVersion(Version, Builder);
 			}
 
 			bInterfaceUpdated |= bPassUpdated;
@@ -998,7 +999,7 @@ bool FMetasoundAssetBase::AutoUpdate(bool bInLogWarningsOnDroppedConnection)
 #endif // !WITH_EDITORONLY_DATA
 }
 
-bool FMetasoundAssetBase::TryUpdateInterfaceFromVersion(const FMetasoundFrontendVersion& Version)
+bool FMetasoundAssetBase::TryUpdateInterfaceFromVersion(const FMetasoundFrontendVersion& Version, FMetaSoundFrontendDocumentBuilder& Builder)
 {
 	using namespace Metasound::Frontend;
 	using namespace AssetBasePrivate;
@@ -1008,7 +1009,7 @@ bool FMetasoundAssetBase::TryUpdateInterfaceFromVersion(const FMetasoundFrontend
 	{
 		TArray<const IInterfaceRegistryEntry*> UpgradePath;
 		GetUpdatePathForDocument(Version, TargetInterface.Metadata.Version, UpgradePath);
-		const bool bUpdated = UpdateDocumentInterface(UpgradePath, Version, GetDocumentHandle());
+		const bool bUpdated = UpdateDocumentInterface(UpgradePath, Version, Builder);
 		ensureMsgf(bUpdated, TEXT("Target interface '%s' was out-of-date but interface failed to be updated"), *TargetInterface.Metadata.Version.ToString());
 		return bUpdated;
 	}
diff --git a/Engine/Plugins/Runtime/Metasound/Source/MetasoundFrontend/Private/MetasoundFrontendDocumentBuilder.cpp b/Engine/Plugins/Runtime/Metasound/Source/MetasoundFrontend/Private/MetasoundFrontendDocumentBuilder.cpp
--- a/Engine/Plugins/Runtime/Metasound/Source/MetasoundFrontend/Private/MetasoundFrontendDocumentBuilder.cpp
+++ b/Engine/Plugins/Runtime/Metasound/Source/MetasoundFrontend/Private/MetasoundFrontendDocumentBuilder.cpp
@@ -1684,6 +1684,22 @@ bool FMetaSoundFrontendDocumentBuilder::AddInterface(FName InterfaceName)
 			return true;
 		}
 
+		// Get all versions of the interface, excluding the latest
+		TArray<FMetasoundFrontendVersion> PreviousVersions = ISearchEngine::Get().FindAllRegisteredInterfacesWithName(Interface.Metadata.Version.Name);
+		PreviousVersions.Remove(Interface.Metadata.Version);
+
+		// Collect the metasound frontend interface definitions registered to the previous versions
+		TArray<FMetasoundFrontendInterface> PreviousVersionInterfaces;
+
+		for (const FMetasoundFrontendVersion& PreviousVersion : PreviousVersions)
+		{
+			FMetasoundFrontendInterface PreviousVersionInterface;
+			if (IInterfaceRegistry::Get().FindInterface(GetInterfaceRegistryKey(PreviousVersion), PreviousVersionInterface))
+			{
+				PreviousVersionInterfaces.Add(MoveTemp(PreviousVersionInterface));
+			}
+		}
+
 		const FInterfaceRegistryKey Key = GetInterfaceRegistryKey(Interface.Metadata.Version);
 		if (const IInterfaceRegistryEntry* Entry = IInterfaceRegistry::Get().FindInterfaceRegistryEntry(Key))
 		{
@@ -1696,9 +1712,10 @@ bool FMetaSoundFrontendDocumentBuilder::AddInterface(FName InterfaceName)
 				return false;
 			}
 
+			// Remove old versions and add new interfaces simultaneously to support node reconnection
 			TArray<FMetasoundFrontendInterface> InterfacesToAdd;
 			InterfacesToAdd.Add(Entry->GetInterface());
-			FModifyInterfaceOptions Options({ }, MoveTemp(InterfacesToAdd));
+			FModifyInterfaceOptions Options(MoveTemp(PreviousVersionInterfaces), MoveTemp(InterfacesToAdd));
 			return ModifyInterfaces(MoveTemp(Options));
 		}
 	}
diff --git a/Engine/Plugins/Runtime/Metasound/Source/MetasoundFrontend/Public/Interfaces/MetasoundFrontendInterfaceRegistry.h b/Engine/Plugins/Runtime/Metasound/Source/MetasoundFrontend/Public/Interfaces/MetasoundFrontendInterfaceRegistry.h
--- a/Engine/Plugins/Runtime/Metasound/Source/MetasoundFrontend/Public/Interfaces/MetasoundFrontendInterfaceRegistry.h
+++ b/Engine/Plugins/Runtime/Metasound/Source/MetasoundFrontend/Public/Interfaces/MetasoundFrontendInterfaceRegistry.h
@@ -32,7 +32,10 @@ namespace Metasound::Frontend
 		virtual FName GetRouterName() const = 0;
 
 		// How to update a given document if versioning is required to this interface from a deprecated version.
-		virtual bool UpdateRootGraphInterface(FDocumentHandle InDocument) const = 0;
+		UE_DEPRECATED(5.7, "Please implement UpdateRootGraphInterface(FMetaSoundFrontendDocumentBuilder& InDocumentBuilder) instead")
+		virtual bool UpdateRootGraphInterface(FDocumentHandle InDocument) const { return false; };
+
+		virtual bool UpdateRootGraphInterface(FMetaSoundFrontendDocumentBuilder& InDocumentBuilder) const = 0;
 	};
 
 	class FInterfaceRegistryTransaction
diff --git a/Engine/Plugins/Runtime/Metasound/Source/MetasoundFrontend/Public/MetasoundAssetBase.h b/Engine/Plugins/Runtime/Metasound/Source/MetasoundFrontend/Public/MetasoundAssetBase.h
--- a/Engine/Plugins/Runtime/Metasound/Source/MetasoundFrontend/Public/MetasoundAssetBase.h
+++ b/Engine/Plugins/Runtime/Metasound/Source/MetasoundFrontend/Public/MetasoundAssetBase.h
@@ -253,7 +253,7 @@ private:
 #endif // WITH_EDITORONLY_DATA
 
 	// Checks if version is up-to-date. If so, returns true. If false, updates the interfaces within the given asset's document to the most recent version.
-	UE_API bool TryUpdateInterfaceFromVersion(const FMetasoundFrontendVersion& Version);
+	UE_API bool TryUpdateInterfaceFromVersion(const FMetasoundFrontendVersion& Version, FMetaSoundFrontendDocumentBuilder& Builder);
 
 	// Versions dependencies to most recent version where applicable. If asset is a preset, MetaSound is rebuilt to accommodate any referenced node class interface changes.
 	// Otherwise, automatically updates any nodes and respective dependent classes to accommodate changes to interfaces therein preserving edges/connections where possible.
