Slate Vector 2 Mismatch Type Fix-up, and remove Slate depdency from MetasoundFrontend round 2

#rnx

diff --git a/Engine/Plugins/Runtime/Metasound/Source/MetasoundEditor/Private/MetasoundEditorGraphCommentNode.cpp b/Engine/Plugins/Runtime/Metasound/Source/MetasoundEditor/Private/MetasoundEditorGraphCommentNode.cpp
--- a/Engine/Plugins/Runtime/Metasound/Source/MetasoundEditor/Private/MetasoundEditorGraphCommentNode.cpp
+++ b/Engine/Plugins/Runtime/Metasound/Source/MetasoundEditor/Private/MetasoundEditorGraphCommentNode.cpp
@@ -43,8 +43,8 @@ void UMetasoundEditorGraphCommentNode::ConvertToFrontendComment(const UEdGraphNo
 	OutComment.Depth = InEdNode.CommentDepth;
 	OutComment.FontSize = InEdNode.FontSize;
 	OutComment.MoveMode = InEdNode.MoveMode == ECommentBoxMode::GroupMovement ? EMetaSoundFrontendGraphCommentMoveMode::GroupMovement : EMetaSoundFrontendGraphCommentMoveMode::NoGroupMovement;
-	OutComment.Position = FVector2D(InEdNode.NodePosX, InEdNode.NodePosY);
-	OutComment.Size = FVector2D(InEdNode.NodeWidth, InEdNode.NodeHeight);
+	OutComment.Position = FIntVector2(InEdNode.NodePosX, InEdNode.NodePosY);
+	OutComment.Size = FIntVector2(InEdNode.NodeWidth, InEdNode.NodeHeight);
 }
 
 void UMetasoundEditorGraphCommentNode::ConvertFromFrontendComment(const FMetaSoundFrontendGraphComment& InComment, UEdGraphNode_Comment& OutEdNode)
@@ -82,7 +82,7 @@ void UMetasoundEditorGraphCommentNode::ResizeNode(const FVector2f& NewSize)
 	{
 		UMetasoundEditorGraph& Graph = *CastChecked<UMetasoundEditorGraph>(GetGraph());
 		UMetaSoundBuilderBase& Builder = FDocumentBuilderRegistry::GetChecked().FindOrBeginBuilding(Graph.GetMetasoundChecked());
-		Builder.FindOrAddGraphComment(CommentID).Size = FVector2f(NodeWidth, NodeHeight);
+		Builder.FindOrAddGraphComment(CommentID).Size = FIntVector2(NodeWidth, NodeHeight);
 	}
 }
 
@@ -131,7 +131,7 @@ void UMetasoundEditorGraphCommentNode::UpdateFrontendNodeLocation()
 {
 	using namespace Metasound::Frontend;
 	FMetaSoundFrontendGraphComment& FrontendComment = GetBuilderChecked().FindOrAddGraphComment(CommentID);
-	FrontendComment.Position = FVector2D(NodePosX, NodePosY);
-	FrontendComment.Size = FVector2D(NodeWidth, NodeHeight);
+	FrontendComment.Position = FIntVector2(NodePosX, NodePosY);
+	FrontendComment.Size = FIntVector2(NodeWidth, NodeHeight);
 }
 #undef LOCTEXT_NAMESPACE
diff --git a/Engine/Plugins/Runtime/Metasound/Source/MetasoundFrontend/Private/MetasoundFrontendDocument.cpp b/Engine/Plugins/Runtime/Metasound/Source/MetasoundFrontend/Private/MetasoundFrontendDocument.cpp
--- a/Engine/Plugins/Runtime/Metasound/Source/MetasoundFrontend/Private/MetasoundFrontendDocument.cpp
+++ b/Engine/Plugins/Runtime/Metasound/Source/MetasoundFrontend/Private/MetasoundFrontendDocument.cpp
@@ -280,8 +280,91 @@ void FMetasoundFrontendDocumentModifyContext::AddNodeIDsModified(const TSet<FGui
 	bDocumentModified = true;
 	NodeIDsModified.Append(InNodesModified);
 }
-#endif // WITH_EDITORONLY_DATA
 
+FMetasoundCommentNodeIntVector::FMetasoundCommentNodeIntVector(const FIntVector2& InValue)
+	: FIntVector2(InValue)
+{
+}
+
+FMetasoundCommentNodeIntVector::FMetasoundCommentNodeIntVector(const FVector2f& InValue)
+	: FIntVector2(static_cast<int32>(InValue.X), static_cast<int32>(InValue.Y))
+{
+}
+
+FMetasoundCommentNodeIntVector::FMetasoundCommentNodeIntVector(const FVector2d& InValue)
+	: FIntVector2(static_cast<int32>(InValue.X), static_cast<int32>(InValue.Y))
+{
+}
+
+FMetasoundCommentNodeIntVector& FMetasoundCommentNodeIntVector::operator=(const FVector2f& InValue)
+{
+	X = static_cast<int32>(InValue.X);
+	Y = static_cast<int32>(InValue.Y);
+	return *this;
+}
+
+FMetasoundCommentNodeIntVector& FMetasoundCommentNodeIntVector::operator=(const FVector2d& InValue)
+{
+	X = static_cast<int32>(InValue.X);
+	Y = static_cast<int32>(InValue.Y);
+	return *this;
+}
+
+FMetasoundCommentNodeIntVector& FMetasoundCommentNodeIntVector::operator=(const FIntVector2& InValue)
+{
+	X = InValue.X;
+	Y = InValue.Y;
+	return *this;
+}
+
+
+bool FMetasoundCommentNodeIntVector::Serialize(FStructuredArchive::FSlot Slot)
+{
+	Slot << static_cast<FIntVector2&>(*this);
+	return true;
+}
+
+bool FMetasoundCommentNodeIntVector::SerializeFromMismatchedTag(const struct FPropertyTag& Tag, FStructuredArchive::FSlot Slot)
+{
+	if (Tag.GetType().IsStruct(NAME_Vector2d))
+	{
+		FVector2d OldVector;
+		Slot << OldVector;
+		X = static_cast<int32>(OldVector.X);
+		Y = static_cast<int32>(OldVector.Y);
+		return true;
+	}
+	else if (Tag.GetType().IsStruct(NAME_Vector2f))
+	{
+		FVector2f OldVector;
+		Slot << OldVector;
+		X = static_cast<int32>(OldVector.X);
+		Y = static_cast<int32>(OldVector.Y);
+		return true;
+	}
+	else if (Tag.GetType().IsStruct(NAME_Vector2D))
+	{
+		if (Slot.GetUnderlyingArchive().UEVer() < EUnrealEngineObjectUE5Version::LARGE_WORLD_COORDINATES)
+		{
+			FVector2f OldVector;
+			Slot << OldVector;
+			X = static_cast<int32>(OldVector.X);
+			Y = static_cast<int32>(OldVector.Y);
+			return true;
+		}
+		else
+		{
+			FVector2d OldVector;
+			Slot << OldVector;
+			X = static_cast<int32>(OldVector.X);
+			Y = static_cast<int32>(OldVector.Y);
+			return true;
+		}
+	}
+
+	return false;
+}
+#endif // WITH_EDITORONLY_DATA
 
 const FMetasoundFrontendVersionNumber& FMetasoundFrontendVersionNumber::GetInvalid()
 {
diff --git a/Engine/Plugins/Runtime/Metasound/Source/MetasoundFrontend/Public/MetasoundFrontendDocument.h b/Engine/Plugins/Runtime/Metasound/Source/MetasoundFrontend/Public/MetasoundFrontendDocument.h
--- a/Engine/Plugins/Runtime/Metasound/Source/MetasoundFrontend/Public/MetasoundFrontendDocument.h
+++ b/Engine/Plugins/Runtime/Metasound/Source/MetasoundFrontend/Public/MetasoundFrontendDocument.h
@@ -181,6 +181,45 @@ enum class EMetaSoundFrontendGraphCommentMoveMode : uint8
 	NoGroupMovement UMETA(DisplayName = "Comment")
 };
 
+/**
+  * Migratory type to avoid adding dependency on Slate FDeprecateSlateVector2D, and by extension,
+  * bring in unnecessary Engine dependencies therein.  At one point, this dependency was incorrectly
+  * added leading to in-determinant serialization as either a double or a float vector. This type
+  * exists to resolve that discrepancy properly. Considered soft deprecated and not to be used for runtime.
+ */
+USTRUCT()
+struct FMetasoundCommentNodeIntVector : public FIntVector2
+{
+	GENERATED_BODY()
+
+#if WITH_EDITORONLY_DATA
+	UE_API FMetasoundCommentNodeIntVector() = default;
+	UE_API FMetasoundCommentNodeIntVector(const FIntVector2& InValue);
+	UE_API FMetasoundCommentNodeIntVector(const FVector2f& InValue);
+	UE_API FMetasoundCommentNodeIntVector(const FVector2d& InValue);
+
+	UE_API FMetasoundCommentNodeIntVector& operator=(const FVector2f& InValue);
+	UE_API FMetasoundCommentNodeIntVector& operator=(const FVector2d& InValue);
+	UE_API FMetasoundCommentNodeIntVector& operator=(const FIntVector2& InValue);
+
+	UE_API bool Serialize(FStructuredArchive::FSlot Slot);
+	UE_API bool SerializeFromMismatchedTag(const struct FPropertyTag& Tag, FStructuredArchive::FSlot Slot);
+#endif // WITH_EDITORONLY_DATA
+};
+
+#if WITH_EDITORONLY_DATA
+template<>
+struct TStructOpsTypeTraits<FMetasoundCommentNodeIntVector>
+	: TStructOpsTypeTraitsBase2<FMetasoundCommentNodeIntVector>
+{
+	enum
+	{
+		WithStructuredSerializeFromMismatchedTag = true,
+	};
+};
+#endif // WITH_EDITORONLY_DATA
+
+
 USTRUCT()
 struct FMetaSoundFrontendGraphComment
 {
@@ -200,10 +239,10 @@ struct FMetaSoundFrontendGraphComment
 	int32 FontSize = 0;
 
 	UPROPERTY()
-	FDeprecateSlateVector2D Position = FVector2f::ZeroVector;
+	FMetasoundCommentNodeIntVector Position = FIntVector2::ZeroValue;
 
 	UPROPERTY()
-	FDeprecateSlateVector2D Size = FVector2f::ZeroVector;
+	FMetasoundCommentNodeIntVector Size = FIntVector2::ZeroValue;
 
 	UPROPERTY()
 	EMetaSoundFrontendGraphCommentMoveMode MoveMode = EMetaSoundFrontendGraphCommentMoveMode::GroupMovement;
