Breaking link to reroute node does not display warning

#rnx

diff --git a/Engine/Plugins/Runtime/Metasound/Source/MetasoundFrontend/Private/MetasoundFrontendDocumentBuilder.cpp b/Engine/Plugins/Runtime/Metasound/Source/MetasoundFrontend/Private/MetasoundFrontendDocumentBuilder.cpp
--- a/Engine/Plugins/Runtime/Metasound/Source/MetasoundFrontend/Private/MetasoundFrontendDocumentBuilder.cpp
+++ b/Engine/Plugins/Runtime/Metasound/Source/MetasoundFrontend/Private/MetasoundFrontendDocumentBuilder.cpp
@@ -4625,7 +4625,8 @@ bool FMetaSoundFrontendDocumentBuilder::RemoveEdgesFromNodeOutput(const FGuid& I
 	const TArrayView<const int32> Indices = EdgeCache.FindEdgeIndicesFromNodeOutput(InNodeID, InVertexID);
 	if (!Indices.IsEmpty())
 	{
-		FMetasoundFrontendGraph& Graph = GetDocumentChecked().RootGraph.FindGraphChecked(PageID);
+		FMetasoundFrontendDocument& Document = GetDocumentChecked();
+		FMetasoundFrontendGraph& Graph = Document.RootGraph.FindGraphChecked(PageID);
 
 		// Copy off indices and sort descending as the edge array will be modified when notifying the cache in the loop below
 		TArray<int32> IndicesCopy(Indices.GetData(), Indices.Num());
@@ -4649,6 +4650,10 @@ bool FMetaSoundFrontendDocumentBuilder::RemoveEdgesFromNodeOutput(const FGuid& I
 			Graph.Edges.RemoveAtSwap(Index, EAllowShrinking::No);
 		}
 
+#if WITH_EDITORONLY_DATA
+		Document.Metadata.ModifyContext.AddNodeIDModified(InNodeID);
+#endif // WITH_EDITORONLY_DATA
+
 		return true;
 	}
 
