Fix for graph comment position serialized with FDeprecateSlateVector2D
#rnx

diff --git a/Engine/Plugins/Runtime/Metasound/Source/MetasoundFrontend/MetasoundFrontend.Build.cs b/Engine/Plugins/Runtime/Metasound/Source/MetasoundFrontend/MetasoundFrontend.Build.cs
--- a/Engine/Plugins/Runtime/Metasound/Source/MetasoundFrontend/MetasoundFrontend.Build.cs
+++ b/Engine/Plugins/Runtime/Metasound/Source/MetasoundFrontend/MetasoundFrontend.Build.cs
@@ -22,11 +22,21 @@ namespace UnrealBuildTool.Rules
 					"JsonUtilities",
 					"Serialization",
 					"SignalProcessing",
-					"MetasoundGraphCore",
-					"SlateCore"
+					"MetasoundGraphCore"
 				}
 			);
 
+			if (Target.bBuildWithEditorOnlyData)
+			{
+				PrivateDependencyModuleNames.AddRange(
+					new string[]
+					{
+						"Slate",
+						"SlateCore"
+					}
+				);
+			}
+
 			PublicDefinitions.Add("WITH_METASOUND_FRONTEND=1");
 
 			// In UE 5.6, registered node are expected to support the constructor 
diff --git a/Engine/Plugins/Runtime/Metasound/Source/MetasoundFrontend/Private/MetasoundFrontendDocument.cpp b/Engine/Plugins/Runtime/Metasound/Source/MetasoundFrontend/Private/MetasoundFrontendDocument.cpp
--- a/Engine/Plugins/Runtime/Metasound/Source/MetasoundFrontend/Private/MetasoundFrontendDocument.cpp
+++ b/Engine/Plugins/Runtime/Metasound/Source/MetasoundFrontend/Private/MetasoundFrontendDocument.cpp
@@ -19,6 +19,7 @@
 
 #if WITH_EDITORONLY_DATA
 #include "Internationalization/Internationalization.h"
+#include "Types/SlateVector2.h"
 #endif // WITH_EDITORONLY_DATA
 
 #include UE_INLINE_GENERATED_CPP_BY_NAME(MetasoundFrontendDocument)
@@ -325,8 +326,16 @@ bool FMetasoundCommentNodeIntVector::Serialize(FStructuredArchive::FSlot Slot)
 }
 
 bool FMetasoundCommentNodeIntVector::SerializeFromMismatchedTag(const struct FPropertyTag& Tag, FStructuredArchive::FSlot Slot)
-{
-	if (Tag.GetType().IsStruct(NAME_Vector2d))
+{;
+	if (Tag.GetType().GetName() == FDeprecateSlateVector2D::StaticStruct()->GetFName())
+	{
+		FDeprecateSlateVector2D OldVector;
+		Slot << OldVector;
+		X = static_cast<int32>(OldVector.X);
+		Y = static_cast<int32>(OldVector.Y);
+		return true;
+	}
+	else if (Tag.GetType().IsStruct(NAME_Vector2d))
 	{
 		FVector2d OldVector;
 		Slot << OldVector;
