Fix MetaSound comment node positions not being serialized correctly after resizing

diff --git a/Engine/Plugins/Runtime/Metasound/Source/MetasoundEditor/Private/MetasoundEditorGraphCommentNode.cpp b/Engine/Plugins/Runtime/Metasound/Source/MetasoundEditor/Private/MetasoundEditorGraphCommentNode.cpp
--- a/Engine/Plugins/Runtime/Metasound/Source/MetasoundEditor/Private/MetasoundEditorGraphCommentNode.cpp
+++ b/Engine/Plugins/Runtime/Metasound/Source/MetasoundEditor/Private/MetasoundEditorGraphCommentNode.cpp
@@ -80,9 +80,8 @@ void UMetasoundEditorGraphCommentNode::ResizeNode(const FVector2f& NewSize)
 	Super::ResizeNode(NewSize);
 	if (bCanResizeNode) 
 	{
-		UMetasoundEditorGraph& Graph = *CastChecked<UMetasoundEditorGraph>(GetGraph());
-		UMetaSoundBuilderBase& Builder = FDocumentBuilderRegistry::GetChecked().FindOrBeginBuilding(Graph.GetMetasoundChecked());
-		Builder.FindOrAddGraphComment(CommentID).Size = FIntVector2(NodeWidth, NodeHeight);
+		// Update both location and size, since resizing from certain corners can change the location too
+		UpdateFrontendNodeLocation();
 	}
 }
 
