Fix incorrect comparison of type in mismatched property tag fix-up
#rnx

diff --git a/Engine/Plugins/Runtime/Metasound/Source/MetasoundFrontend/Private/MetasoundFrontendDocument.cpp b/Engine/Plugins/Runtime/Metasound/Source/MetasoundFrontend/Private/MetasoundFrontendDocument.cpp
--- a/Engine/Plugins/Runtime/Metasound/Source/MetasoundFrontend/Private/MetasoundFrontendDocument.cpp
+++ b/Engine/Plugins/Runtime/Metasound/Source/MetasoundFrontend/Private/MetasoundFrontendDocument.cpp
@@ -326,8 +326,10 @@ bool FMetasoundCommentNodeIntVector::Serialize(FStructuredArchive::FSlot Slot)
 }
 
 bool FMetasoundCommentNodeIntVector::SerializeFromMismatchedTag(const struct FPropertyTag& Tag, FStructuredArchive::FSlot Slot)
-{;
-	if (Tag.GetType().GetName() == FDeprecateSlateVector2D::StaticStruct()->GetFName())
+{
+	const UScriptStruct* DeprecatedSlateType = FDeprecateSlateVector2D::StaticStruct();
+	check(DeprecatedSlateType);
+	if (Tag.GetType().IsStruct(DeprecatedSlateType->GetFName()))
 	{
 		FDeprecateSlateVector2D OldVector;
 		Slot << OldVector;
