Fix issues with MetaSound preset creation source settings in restricted mode and when using template soundwave option

diff --git a/Engine/Plugins/Runtime/Metasound/Source/MetasoundEditor/Private/MetasoundAssetDefinitions.cpp b/Engine/Plugins/Runtime/Metasound/Source/MetasoundEditor/Private/MetasoundAssetDefinitions.cpp
--- a/Engine/Plugins/Runtime/Metasound/Source/MetasoundEditor/Private/MetasoundAssetDefinitions.cpp
+++ b/Engine/Plugins/Runtime/Metasound/Source/MetasoundEditor/Private/MetasoundAssetDefinitions.cpp
@@ -423,11 +423,26 @@ namespace MenuExtension_MetaSoundSourceTemplate
 			{
 				FString PackagePath;
 				FString AssetName;
+				UObject* NewMetaSound;
 
 				IAssetTools::Get().CreateUniqueAssetName(ReferencedMetaSound->GetOutermost()->GetName(), TEXT("_Preset"), PackagePath, AssetName);
 
-				// Duplicate asset to preserve properties from referenced asset (ex. quality settings, soundwave properties)
-				UObject* NewMetaSound = IAssetTools::Get().DuplicateAssetWithDialogAndTitle(AssetName, FPackageName::GetLongPackagePath(PackagePath), ReferencedMetaSound, LOCTEXT("CreateMetaSoundPresetTitle", "Create MetaSound Preset"));
+				Metasound::Editor::IMetasoundEditorModule& MetaSoundEditorModule = FModuleManager::GetModuleChecked<Metasound::Editor::IMetasoundEditorModule>("MetaSoundEditor");
+				if (MetaSoundEditorModule.IsRestrictedMode())
+				{
+					// Cannot duplicate cooked assets in restricted mode, so create new object and copy over properties
+					// for sources in InitAsset below. Since copying properties is done manually, 
+					// SetSoundWaveSettingsFromTemplate may need to be updated with properties to be copied. 
+					UMetaSoundBaseFactory* Factory = NewObject<TFactoryClass>();
+					Factory->ReferencedMetaSoundObject = ReferencedMetaSound;
+
+					NewMetaSound = IAssetTools::Get().CreateAssetWithDialog(AssetName, FPackageName::GetLongPackagePath(PackagePath), Factory->GetSupportedClass(), Factory);
+				}
+				else
+				{
+					// Duplicate asset to preserve properties from referenced asset (ex. quality settings, soundwave properties)
+					NewMetaSound = IAssetTools::Get().DuplicateAssetWithDialogAndTitle(AssetName, FPackageName::GetLongPackagePath(PackagePath), ReferencedMetaSound, LOCTEXT("CreateMetaSoundPresetTitle", "Create MetaSound Preset"));
+				}
 				
 				if (NewMetaSound)
 				{
diff --git a/Engine/Plugins/Runtime/Metasound/Source/MetasoundEditor/Private/MetasoundEditorSubsystem.cpp b/Engine/Plugins/Runtime/Metasound/Source/MetasoundEditor/Private/MetasoundEditorSubsystem.cpp
--- a/Engine/Plugins/Runtime/Metasound/Source/MetasoundEditor/Private/MetasoundEditorSubsystem.cpp
+++ b/Engine/Plugins/Runtime/Metasound/Source/MetasoundEditor/Private/MetasoundEditorSubsystem.cpp
@@ -98,10 +98,9 @@ TScriptInterface<IMetaSoundDocumentInterface> UMetaSoundEditorSubsystem::BuildTo
 			// Apply template SoundWave settings 
 			{
 				// Template SoundWave settings only apply to sources
-				// and do not apply to presets, which were duplicated from the referenced preset asset 
-				// and therefore have those object properties 
+				// and will override settings from a preset's referenced asset settings
 				const bool bIsSource = &MetaSoundUClass == UMetaSoundSource::StaticClass();
-				if (TemplateSoundWave != nullptr && bIsSource && !NewDocBuilder.IsPreset())
+				if (TemplateSoundWave != nullptr && bIsSource)
 				{
 					SetSoundWaveSettingsFromTemplate(*CastChecked<USoundWave>(NewMetaSound), *TemplateSoundWave);
 				}
@@ -354,6 +353,17 @@ void UMetaSoundEditorSubsystem::InitAsset(UObject& InNewMetaSound, UObject* InRe
 		// Initialize preset with referenced Metasound 
 		TScriptInterface<IMetaSoundDocumentInterface> ReferencedDocInterface = InReferencedMetaSound;
 		Builder.ConvertToPreset(ReferencedDocInterface->GetConstDocument());
+		
+		if (&DocInterface->GetBaseMetaSoundUClass() == UMetaSoundSource::StaticClass())
+		{
+			// If in restricted mode, copy source soundwave settings from referenced MetaSound 
+			// (In non restricted mode, soundwave settings are automatically copied over by creating the new asset using duplicate) 
+			Metasound::Editor::IMetasoundEditorModule& MetaSoundEditorModule = FModuleManager::GetModuleChecked<Metasound::Editor::IMetasoundEditorModule>("MetaSoundEditor");
+			if (MetaSoundEditorModule.IsRestrictedMode())
+			{
+				SetSoundWaveSettingsFromTemplate(*CastChecked<USoundWave>(&InNewMetaSound), *CastChecked<USoundWave>(InReferencedMetaSound));
+			}
+		}
 	}
 }
 
