- Clean-up from prior build fix CL
- Ignore failures to parse editor-only vector migration in cook

#rnx

diff --git a/Engine/Plugins/Runtime/Metasound/Source/MetasoundFrontend/MetasoundFrontend.Build.cs b/Engine/Plugins/Runtime/Metasound/Source/MetasoundFrontend/MetasoundFrontend.Build.cs
--- a/Engine/Plugins/Runtime/Metasound/Source/MetasoundFrontend/MetasoundFrontend.Build.cs
+++ b/Engine/Plugins/Runtime/Metasound/Source/MetasoundFrontend/MetasoundFrontend.Build.cs
@@ -28,7 +28,7 @@ namespace UnrealBuildTool.Rules
 
 			if (Target.bBuildWithEditorOnlyData)
 			{
-				PrivateDependencyModuleNames.AddRange(
+				PublicDependencyModuleNames.AddRange(
 					new string[]
 					{
 						"Slate",
diff --git a/Engine/Plugins/Runtime/Metasound/Source/MetasoundFrontend/Private/MetasoundFrontendDocument.cpp b/Engine/Plugins/Runtime/Metasound/Source/MetasoundFrontend/Private/MetasoundFrontendDocument.cpp
--- a/Engine/Plugins/Runtime/Metasound/Source/MetasoundFrontend/Private/MetasoundFrontendDocument.cpp
+++ b/Engine/Plugins/Runtime/Metasound/Source/MetasoundFrontend/Private/MetasoundFrontendDocument.cpp
@@ -6,6 +6,7 @@
 #include "Algo/ForEach.h"
 #include "Algo/Sort.h"
 #include "Algo/Transform.h"
+#include "CoreGlobals.h"
 #include "IAudioParameterInterfaceRegistry.h"
 #include "Logging/LogMacros.h"
 #include "MetasoundDocumentInterface.h"
@@ -373,7 +374,10 @@ bool FMetasoundCommentNodeIntVector::SerializeFromMismatchedTag(const struct FPr
 		}
 	}
 
-	return false;
+	// Hack for running cook commandlet, where we don't really care
+	// if it fails as cooked content will never be editable, so don't
+	// bother reporting if old data was not translated.
+	return IsRunningCookCommandlet();
 }
 #endif // WITH_EDITORONLY_DATA
 
diff --git a/Engine/Plugins/Runtime/Metasound/Source/MetasoundFrontend/Public/MetasoundFrontendDocument.h b/Engine/Plugins/Runtime/Metasound/Source/MetasoundFrontend/Public/MetasoundFrontendDocument.h
--- a/Engine/Plugins/Runtime/Metasound/Source/MetasoundFrontend/Public/MetasoundFrontendDocument.h
+++ b/Engine/Plugins/Runtime/Metasound/Source/MetasoundFrontend/Public/MetasoundFrontendDocument.h
@@ -18,9 +18,12 @@
 #include "Templates/Function.h"
 #include "Templates/Invoke.h"
 #include "Templates/TypeHash.h"
-#include "Types/SlateVector2.h"
 #include "UObject/NoExportTypes.h"
 
+#if WITH_EDITORONLY_DATA
+#include "Types/SlateVector2.h"
+#endif // WITH_EDITORONLY_DATA
+
 #include "MetasoundFrontendDocument.generated.h"
 
 #define UE_API METASOUNDFRONTEND_API
