Expose type traits for non-editor only data builds in a last-ditch attempt to fix fast AutoTest build warnings
#rnx

diff --git a/Engine/Plugins/Runtime/Metasound/Source/MetasoundFrontend/Private/MetasoundFrontendDocument.cpp b/Engine/Plugins/Runtime/Metasound/Source/MetasoundFrontend/Private/MetasoundFrontendDocument.cpp
--- a/Engine/Plugins/Runtime/Metasound/Source/MetasoundFrontend/Private/MetasoundFrontendDocument.cpp
+++ b/Engine/Plugins/Runtime/Metasound/Source/MetasoundFrontend/Private/MetasoundFrontendDocument.cpp
@@ -282,6 +282,7 @@ void FMetasoundFrontendDocumentModifyContext::AddNodeIDsModified(const TSet<FGui
 	bDocumentModified = true;
 	NodeIDsModified.Append(InNodesModified);
 }
+#endif // WITH_EDITORONLY_DATA
 
 FMetasoundCommentNodeIntVector::FMetasoundCommentNodeIntVector(const FIntVector2& InValue)
 	: FIntVector2(InValue)
@@ -328,6 +329,7 @@ bool FMetasoundCommentNodeIntVector::Serialize(FStructuredArchive::FSlot Slot)
 
 bool FMetasoundCommentNodeIntVector::SerializeFromMismatchedTag(const struct FPropertyTag& Tag, FStructuredArchive::FSlot Slot)
 {
+#if WITH_EDITORONLY_DATA
 	const UScriptStruct* DeprecatedSlateType = FDeprecateSlateVector2D::StaticStruct();
 	check(DeprecatedSlateType);
 	if (Tag.GetType().IsStruct(DeprecatedSlateType->GetFName()))
@@ -338,7 +340,18 @@ bool FMetasoundCommentNodeIntVector::SerializeFromMismatchedTag(const struct FPr
 		Y = static_cast<int32>(OldVector.Y);
 		return true;
 	}
-	else if (Tag.GetType().IsStruct(NAME_Vector2d))
+#else
+	if (Tag.GetType().IsStruct("DeprecateSlateVector2D"))
+	{
+		// Missing type, don't care about visualizing this data.
+		// Ignore the serialization.
+		X = 0;
+		Y = 0;
+		return true;
+	}
+#endif // !WITH_EDITORONLY_DATA
+
+	if (Tag.GetType().IsStruct(NAME_Vector2d))
 	{
 		FVector2d OldVector;
 		Slot << OldVector;
@@ -379,7 +392,6 @@ bool FMetasoundCommentNodeIntVector::SerializeFromMismatchedTag(const struct FPr
 	// bother reporting if old data was not translated.
 	return IsRunningCookCommandlet();
 }
-#endif // WITH_EDITORONLY_DATA
 
 const FMetasoundFrontendVersionNumber& FMetasoundFrontendVersionNumber::GetInvalid()
 {
diff --git a/Engine/Plugins/Runtime/Metasound/Source/MetasoundFrontend/Public/MetasoundFrontendDocument.h b/Engine/Plugins/Runtime/Metasound/Source/MetasoundFrontend/Public/MetasoundFrontendDocument.h
--- a/Engine/Plugins/Runtime/Metasound/Source/MetasoundFrontend/Public/MetasoundFrontendDocument.h
+++ b/Engine/Plugins/Runtime/Metasound/Source/MetasoundFrontend/Public/MetasoundFrontendDocument.h
@@ -195,7 +195,6 @@ struct FMetasoundCommentNodeIntVector : public FIntVector2
 {
 	GENERATED_BODY()
 
-#if WITH_EDITORONLY_DATA
 	UE_API FMetasoundCommentNodeIntVector() = default;
 	UE_API FMetasoundCommentNodeIntVector(const FIntVector2& InValue);
 	UE_API FMetasoundCommentNodeIntVector(const FVector2f& InValue);
@@ -207,10 +206,8 @@ struct FMetasoundCommentNodeIntVector : public FIntVector2
 
 	UE_API bool Serialize(FStructuredArchive::FSlot Slot);
 	UE_API bool SerializeFromMismatchedTag(const struct FPropertyTag& Tag, FStructuredArchive::FSlot Slot);
-#endif // WITH_EDITORONLY_DATA
 };
 
-#if WITH_EDITORONLY_DATA
 template<>
 struct TStructOpsTypeTraits<FMetasoundCommentNodeIntVector>
 	: TStructOpsTypeTraitsBase2<FMetasoundCommentNodeIntVector>
@@ -220,7 +217,6 @@ struct TStructOpsTypeTraits<FMetasoundCommentNodeIntVector>
 		WithStructuredSerializeFromMismatchedTag = true,
 	};
 };
-#endif // WITH_EDITORONLY_DATA
 
 
 USTRUCT()
