Utility/Helper functions for creating MetaSounds from c++ code.

diff --git a/Engine/Plugins/Runtime/Metasound/Source/MetasoundEngine/Public/MetasoundBuilderBase.cpp b/Engine/Plugins/Runtime/Metasound/Source/MetasoundEngine/Public/MetasoundBuilderBase.cpp
--- a/Engine/Plugins/Runtime/Metasound/Source/MetasoundEngine/Public/MetasoundBuilderBase.cpp
+++ b/Engine/Plugins/Runtime/Metasound/Source/MetasoundEngine/Public/MetasoundBuilderBase.cpp
@@ -435,6 +435,15 @@ void UMetaSoundBuilderBase::ConnectNodes(const FMetaSoundBuilderNodeOutputHandle
 	}
 }
 
+void UMetaSoundBuilderBase::ConnectNodes(const FMetaSoundNodeHandle& SourceNode, const FName& OutputName, const FMetaSoundNodeHandle& DestinationNode, const FName& InputName, EMetaSoundBuilderResult& OutResult)
+{
+	FMetaSoundBuilderNodeOutputHandle Out = FindNodeOutputByName(SourceNode, OutputName, OutResult);
+	if (OutResult == EMetaSoundBuilderResult::Failed) { return; }
+	FMetaSoundBuilderNodeInputHandle In = FindNodeInputByName(DestinationNode, InputName, OutResult);
+	if (OutResult == EMetaSoundBuilderResult::Failed) { return; }
+	ConnectNodes(Out, In, OutResult);
+}
+
 void UMetaSoundBuilderBase::ConnectNodesByInterfaceBindings(const FMetaSoundNodeHandle& FromNodeHandle, const FMetaSoundNodeHandle& ToNodeHandle, EMetaSoundBuilderResult& OutResult)
 {
 	const bool bEdgesAdded = Builder.AddEdgesByNodeClassInterfaceBindings(FromNodeHandle.NodeID, ToNodeHandle.NodeID);
@@ -499,6 +508,26 @@ void UMetaSoundBuilderBase::ConnectNodeOutputToGraphOutput(FName GraphOutputName
 	}
 }
 
+void UMetaSoundBuilderBase::ConnectNodeToGraphOutput(const FMetaSoundNodeHandle& SourceNode, const FName& NodeOutputName, const FName& GraphOutputName, EMetaSoundBuilderResult& OutResult)
+{
+	FMetaSoundBuilderNodeOutputHandle NodeOut = UMetaSoundBuilderBase::FindNodeOutputByName(SourceNode, NodeOutputName, OutResult);
+	if (OutResult != EMetaSoundBuilderResult::Succeeded)
+	{
+		return;
+	}
+	ConnectNodeOutputToGraphOutput(GraphOutputName, NodeOut, OutResult);
+}
+
+void UMetaSoundBuilderBase::ConnectNodeToGraphOutput(const FMetaSoundNodeHandle& SourceNode, const FName& NodeOutputName, const FMetaSoundBuilderNodeInputHandle& GraphOutput, EMetaSoundBuilderResult& OutResult)
+{
+	FMetaSoundBuilderNodeOutputHandle NodeOut = UMetaSoundBuilderBase::FindNodeOutputByName(SourceNode, NodeOutputName, OutResult);
+	if (OutResult != EMetaSoundBuilderResult::Succeeded)
+	{
+		return;
+	}
+	ConnectNodes(NodeOut, GraphOutput, OutResult);
+}
+
 void UMetaSoundBuilderBase::ConnectNodeInputToGraphInput(FName GraphInputName, const FMetaSoundBuilderNodeInputHandle& NodeInputHandle, EMetaSoundBuilderResult& OutResult)
 {
 	using namespace Metasound::Frontend;
@@ -523,6 +552,16 @@ void UMetaSoundBuilderBase::ConnectNodeInputToGraphInput(FName GraphInputName, c
 	}
 }
 
+void UMetaSoundBuilderBase::ConnectGraphInputToNode(const FName& GraphInputName, const FMetaSoundNodeHandle& DestinationNode, const FName& InputName, EMetaSoundBuilderResult& OutResult)
+{
+	FMetaSoundBuilderNodeInputHandle In = FindNodeInputByName(DestinationNode, InputName, OutResult);
+	if (OutResult == EMetaSoundBuilderResult::Failed)
+	{
+		return;
+	}
+	ConnectNodeInputToGraphInput(GraphInputName, In, OutResult);
+}
+
 void UMetaSoundBuilderBase::ConvertFromPreset(EMetaSoundBuilderResult& OutResult)
 {
 	const bool bSuccess = Builder.ConvertFromPreset();
diff --git a/Engine/Plugins/Runtime/Metasound/Source/MetasoundEngine/Public/MetasoundBuilderBase.h b/Engine/Plugins/Runtime/Metasound/Source/MetasoundEngine/Public/MetasoundBuilderBase.h
--- a/Engine/Plugins/Runtime/Metasound/Source/MetasoundEngine/Public/MetasoundBuilderBase.h
+++ b/Engine/Plugins/Runtime/Metasound/Source/MetasoundEngine/Public/MetasoundBuilderBase.h
@@ -210,10 +210,30 @@ public:
 	UFUNCTION(BlueprintCallable, Category = "Audio|MetaSound|Builder", meta = (ExpandEnumAsExecs = "OutResult"))
 	UE_API void ConnectNodeOutputToGraphOutput(FName GraphOutputName, const FMetaSoundBuilderNodeOutputHandle& NodeOutputHandle, EMetaSoundBuilderResult& OutResult);
 
+	// Connects a given node's named output to the graph output with the given name.
+	UFUNCTION(BlueprintCallable, Category = "Audio|MetaSound|Builder", meta = (ExpandEnumAsExecs = "OutResult"))
+	FORCEINLINE void ConnectNamedNodeOutputToNamedGraphOutput(const FMetaSoundNodeHandle& SourceNode, const FName& NodeOutputName, const FName& GraphOutputName, EMetaSoundBuilderResult& OutResult)
+	{
+		ConnectNodeToGraphOutput(SourceNode, NodeOutputName, GraphOutputName, OutResult);
+	}
+	UE_API void ConnectNodeToGraphOutput(const FMetaSoundNodeHandle& SourceNode, const FName& NodeOutputName, const FName& GraphOutputName, EMetaSoundBuilderResult& OutResult);
+
+	// Connects a given node's named output to the graph output with the given handle.
+	UFUNCTION(BlueprintCallable, Category = "Audio|MetaSound|Builder", meta = (ExpandEnumAsExecs = "OutResult"))
+	FORCEINLINE void ConnectNamedNodeOutputToGraphOutput(const FMetaSoundNodeHandle& SourceNode, const FName& NodeOutputName, const FMetaSoundBuilderNodeInputHandle& GraphOutput, EMetaSoundBuilderResult& OutResult)
+	{
+		ConnectNodeToGraphOutput(SourceNode, NodeOutputName, GraphOutput, OutResult);
+	}
+	UE_API void ConnectNodeToGraphOutput(const FMetaSoundNodeHandle& SourceNode, const FName& NodeOutputName, const FMetaSoundBuilderNodeInputHandle& GraphOutput, EMetaSoundBuilderResult& OutResult);
+
 	// Connects a given node input to the graph input with the given name.
 	UFUNCTION(BlueprintCallable, Category = "Audio|MetaSound|Builder", meta = (ExpandEnumAsExecs = "OutResult"))
 	UE_API void ConnectNodeInputToGraphInput(FName GraphInputName, const FMetaSoundBuilderNodeInputHandle& NodeInputHandle, EMetaSoundBuilderResult& OutResult);
 
+	UE_API void ConnectGraphInputToNode(const FName& GraphInputName, const FMetaSoundNodeHandle& DestinationNode, const FName& InputName, EMetaSoundBuilderResult& OutResult);
+
+	UE_API void ConnectNodes(const FMetaSoundNodeHandle& SourceNode, const FName& OutputName, const FMetaSoundNodeHandle& DestinationNode, const FName& InputName, EMetaSoundBuilderResult& OutResult);
+
 	// Returns whether node exists.
 	UFUNCTION(BlueprintCallable, Category = "Audio|MetaSound|Builder")
 	UE_API UPARAM(DisplayName = "IsValid") bool ContainsNode(const FMetaSoundNodeHandle& Node) const;
@@ -440,6 +460,20 @@ public:
 	UFUNCTION(BlueprintCallable, Category = "Audio|MetaSound|Builder", meta = (ExpandEnumAsExecs = "OutResult"))
 	UE_API void SetNodeInputDefault(const FMetaSoundBuilderNodeInputHandle& NodeInputHandle, const FMetasoundFrontendLiteral& Literal, EMetaSoundBuilderResult& OutResult);
 
+	// Templated convenience function for setting the node's input default value (used if no connection to the given node input is present)
+	template <typename LiteralType>
+	void SetNodeInputDefault(const FMetaSoundNodeHandle& NodeHandle, const FName& InputName, const LiteralType& Literal, EMetaSoundBuilderResult& OutResult)
+	{
+		FMetaSoundBuilderNodeInputHandle In = FindNodeInputByName(NodeHandle, InputName, OutResult);
+		if (OutResult != EMetaSoundBuilderResult::Succeeded)
+		{
+			return;
+		}
+		FMetasoundFrontendLiteral ObjectLiteral;
+		ObjectLiteral.Set(Literal);
+		SetNodeInputDefault(In, ObjectLiteral, OutResult);
+	}
+
 	// Disconnects the given graph input's respective template nodes and sets the graph input's AccessType should it not match the current AccessType.
 	// Result succeeds if the AccessType was successfully changed or if the provided AccessType is already the input's current AccessType.
 	UFUNCTION(BlueprintCallable, Category = "Audio|MetaSound|Builder", meta = (ExpandEnumAsExecs = "OutResult"))
