Attempt fallback to hardcoded type name if type cannot be found (and checks potentially disabled)
#rnx

diff --git a/Engine/Plugins/Runtime/Metasound/Source/MetasoundFrontend/Private/MetasoundFrontendDocument.cpp b/Engine/Plugins/Runtime/Metasound/Source/MetasoundFrontend/Private/MetasoundFrontendDocument.cpp
--- a/Engine/Plugins/Runtime/Metasound/Source/MetasoundFrontend/Private/MetasoundFrontendDocument.cpp
+++ b/Engine/Plugins/Runtime/Metasound/Source/MetasoundFrontend/Private/MetasoundFrontendDocument.cpp
@@ -330,17 +330,19 @@ bool FMetasoundCommentNodeIntVector::Serialize(FStructuredArchive::FSlot Slot)
 bool FMetasoundCommentNodeIntVector::SerializeFromMismatchedTag(const struct FPropertyTag& Tag, FStructuredArchive::FSlot Slot)
 {
 #if WITH_EDITORONLY_DATA
-	const UScriptStruct* DeprecatedSlateType = FDeprecateSlateVector2D::StaticStruct();
-	check(DeprecatedSlateType);
-	if (Tag.GetType().IsStruct(DeprecatedSlateType->GetFName()))
+	if (const UScriptStruct* DeprecatedSlateType = FDeprecateSlateVector2D::StaticStruct())
 	{
-		FDeprecateSlateVector2D OldVector;
-		Slot << OldVector;
-		X = static_cast<int32>(OldVector.X);
-		Y = static_cast<int32>(OldVector.Y);
-		return true;
+		if (Tag.GetType().IsStruct(DeprecatedSlateType->GetFName()))
+		{
+			FDeprecateSlateVector2D OldVector;
+			Slot << OldVector;
+			X = static_cast<int32>(OldVector.X);
+			Y = static_cast<int32>(OldVector.Y);
+			return true;
+		}
 	}
-#else
+#endif // WITH_EDITORONLY_DATA
+
 	if (Tag.GetType().IsStruct("DeprecateSlateVector2D"))
 	{
 		// Missing type, don't care about visualizing this data.
@@ -349,7 +351,6 @@ bool FMetasoundCommentNodeIntVector::SerializeFromMismatchedTag(const struct FPr
 		Y = 0;
 		return true;
 	}
-#endif // !WITH_EDITORONLY_DATA
 
 	if (Tag.GetType().IsStruct(NAME_Vector2d))
 	{
