Add debug logging to help triage AutoTest Cook
#rnx

diff --git a/Engine/Plugins/Runtime/Metasound/Source/MetasoundFrontend/Private/MetasoundFrontendDocument.cpp b/Engine/Plugins/Runtime/Metasound/Source/MetasoundFrontend/Private/MetasoundFrontendDocument.cpp
--- a/Engine/Plugins/Runtime/Metasound/Source/MetasoundFrontend/Private/MetasoundFrontendDocument.cpp
+++ b/Engine/Plugins/Runtime/Metasound/Source/MetasoundFrontend/Private/MetasoundFrontendDocument.cpp
@@ -299,6 +299,13 @@ FMetasoundCommentNodeIntVector::FMetasoundCommentNodeIntVector(const FVector2d&
 {
 }
 
+#if WITH_EDITORONLY_DATA
+FMetasoundCommentNodeIntVector::FMetasoundCommentNodeIntVector(const FDeprecateSlateVector2D& InValue)
+	: FIntVector2(static_cast<int32>(InValue.X), static_cast<int32>(InValue.Y))
+{
+}
+#endif // WITH_EDITORONLY_DATA
+
 FMetasoundCommentNodeIntVector& FMetasoundCommentNodeIntVector::operator=(const FVector2f& InValue)
 {
 	X = static_cast<int32>(InValue.X);
@@ -320,6 +327,14 @@ FMetasoundCommentNodeIntVector& FMetasoundCommentNodeIntVector::operator=(const
 	return *this;
 }
 
+#if WITH_EDITORONLY_DATA
+FMetasoundCommentNodeIntVector& FMetasoundCommentNodeIntVector::operator=(const FDeprecateSlateVector2D& InValue)
+{
+	X = InValue.X;
+	Y = InValue.Y;
+	return *this;
+}
+#endif // WITH_EDITORONLY_DATA
 
 bool FMetasoundCommentNodeIntVector::Serialize(FStructuredArchive::FSlot Slot)
 {
@@ -332,6 +347,7 @@ bool FMetasoundCommentNodeIntVector::SerializeFromMismatchedTag(const struct FPr
 #if WITH_EDITORONLY_DATA
 	if (const UScriptStruct* DeprecatedSlateType = FDeprecateSlateVector2D::StaticStruct())
 	{
+		UE_LOG(LogMetaSound, Display, TEXT("FMetasoundCommentNodeIntVector::SerializeFromMismatchedTag - Found Deprecated Type. Resolving Mismatch"));
 		if (Tag.GetType().IsStruct(DeprecatedSlateType->GetFName()))
 		{
 			FDeprecateSlateVector2D OldVector;
@@ -345,6 +361,7 @@ bool FMetasoundCommentNodeIntVector::SerializeFromMismatchedTag(const struct FPr
 
 	if (Tag.GetType().IsStruct("DeprecateSlateVector2D"))
 	{
+		UE_LOG(LogMetaSound, Display, TEXT("FMetasoundCommentNodeIntVector::SerializeFromMismatchedTag - Missing Deprecated Type. Ignoring Mismatch"));
 		// Missing type, don't care about visualizing this data.
 		// Ignore the serialization.
 		X = 0;
@@ -354,6 +371,7 @@ bool FMetasoundCommentNodeIntVector::SerializeFromMismatchedTag(const struct FPr
 
 	if (Tag.GetType().IsStruct(NAME_Vector2d))
 	{
+		UE_LOG(LogMetaSound, Display, TEXT("FMetasoundCommentNodeIntVector::SerializeFromMismatchedTag - Resolving as FVector2d"));
 		FVector2d OldVector;
 		Slot << OldVector;
 		X = static_cast<int32>(OldVector.X);
@@ -362,6 +380,7 @@ bool FMetasoundCommentNodeIntVector::SerializeFromMismatchedTag(const struct FPr
 	}
 	else if (Tag.GetType().IsStruct(NAME_Vector2f))
 	{
+		UE_LOG(LogMetaSound, Display, TEXT("FMetasoundCommentNodeIntVector::SerializeFromMismatchedTag - Resolving as FVector2f"));
 		FVector2f OldVector;
 		Slot << OldVector;
 		X = static_cast<int32>(OldVector.X);
@@ -372,6 +391,7 @@ bool FMetasoundCommentNodeIntVector::SerializeFromMismatchedTag(const struct FPr
 	{
 		if (Slot.GetUnderlyingArchive().UEVer() < EUnrealEngineObjectUE5Version::LARGE_WORLD_COORDINATES)
 		{
+			UE_LOG(LogMetaSound, Display, TEXT("FMetasoundCommentNodeIntVector::SerializeFromMismatchedTag - Resolving as FVector2f (Pre LWC)"));
 			FVector2f OldVector;
 			Slot << OldVector;
 			X = static_cast<int32>(OldVector.X);
@@ -380,6 +400,7 @@ bool FMetasoundCommentNodeIntVector::SerializeFromMismatchedTag(const struct FPr
 		}
 		else
 		{
+			UE_LOG(LogMetaSound, Display, TEXT("FMetasoundCommentNodeIntVector::SerializeFromMismatchedTag - Resolving as FVector2d (Post LWC)"));
 			FVector2d OldVector;
 			Slot << OldVector;
 			X = static_cast<int32>(OldVector.X);
@@ -388,10 +409,21 @@ bool FMetasoundCommentNodeIntVector::SerializeFromMismatchedTag(const struct FPr
 		}
 	}
 
+	const bool bIsCookCommandlet = IsRunningCookCommandlet();
+
 	// Hack for running cook commandlet, where we don't really care
 	// if it fails as cooked content will never be editable, so don't
 	// bother reporting if old data was not translated.
-	return IsRunningCookCommandlet();
+	if (bIsCookCommandlet)
+	{
+		UE_LOG(LogMetaSound, Display, TEXT("FMetasoundCommentNodeIntVector::SerializeFromMismatchedTag - Did not resolve. Is in cook so ok not serializing."));
+	}
+	else
+	{
+		UE_LOG(LogMetaSound, Display, TEXT("FMetasoundCommentNodeIntVector::SerializeFromMismatchedTag - Did not resolve. Returning failure state."));
+	}
+
+	return bIsCookCommandlet;
 }
 
 const FMetasoundFrontendVersionNumber& FMetasoundFrontendVersionNumber::GetInvalid()
