Missing file
#rnx

diff --git a/Engine/Plugins/Runtime/Metasound/Source/MetasoundFrontend/Public/MetasoundFrontendDocument.h b/Engine/Plugins/Runtime/Metasound/Source/MetasoundFrontend/Public/MetasoundFrontendDocument.h
--- a/Engine/Plugins/Runtime/Metasound/Source/MetasoundFrontend/Public/MetasoundFrontendDocument.h
+++ b/Engine/Plugins/Runtime/Metasound/Source/MetasoundFrontend/Public/MetasoundFrontendDocument.h
@@ -200,10 +200,18 @@ struct FMetasoundCommentNodeIntVector : public FIntVector2
 	UE_API FMetasoundCommentNodeIntVector(const FVector2f& InValue);
 	UE_API FMetasoundCommentNodeIntVector(const FVector2d& InValue);
 
+#if WITH_EDITORONLY_DATA
+	UE_API FMetasoundCommentNodeIntVector(const FDeprecateSlateVector2D& InValue);
+#endif // WITH_EDITORONLY_DATA
+
 	UE_API FMetasoundCommentNodeIntVector& operator=(const FVector2f& InValue);
 	UE_API FMetasoundCommentNodeIntVector& operator=(const FVector2d& InValue);
 	UE_API FMetasoundCommentNodeIntVector& operator=(const FIntVector2& InValue);
 
+#if WITH_EDITORONLY_DATA
+	UE_API FMetasoundCommentNodeIntVector& operator=(const FDeprecateSlateVector2D& InValue);
+#endif // WITH_EDITORONLY_DATA
+
 	UE_API bool Serialize(FStructuredArchive::FSlot Slot);
 	UE_API bool SerializeFromMismatchedTag(const struct FPropertyTag& Tag, FStructuredArchive::FSlot Slot);
 };
