Remove spammier test log lines for resolving AutoTest
#rnx

diff --git a/Engine/Plugins/Runtime/Metasound/Source/MetasoundFrontend/Private/MetasoundFrontendDocument.cpp b/Engine/Plugins/Runtime/Metasound/Source/MetasoundFrontend/Private/MetasoundFrontendDocument.cpp
--- a/Engine/Plugins/Runtime/Metasound/Source/MetasoundFrontend/Private/MetasoundFrontendDocument.cpp
+++ b/Engine/Plugins/Runtime/Metasound/Source/MetasoundFrontend/Private/MetasoundFrontendDocument.cpp
@@ -347,31 +347,33 @@ bool FMetasoundCommentNodeIntVector::SerializeFromMismatchedTag(const struct FPr
 #if WITH_EDITORONLY_DATA
 	if (const UScriptStruct* DeprecatedSlateType = FDeprecateSlateVector2D::StaticStruct())
 	{
-		UE_LOG(LogMetaSound, Display, TEXT("FMetasoundCommentNodeIntVector::SerializeFromMismatchedTag - Found Deprecated Type. Resolving Mismatch"));
 		if (Tag.GetType().IsStruct(DeprecatedSlateType->GetFName()))
 		{
 			FDeprecateSlateVector2D OldVector;
 			Slot << OldVector;
 			X = static_cast<int32>(OldVector.X);
 			Y = static_cast<int32>(OldVector.Y);
+
+			UE_LOG(LogMetaSound, Display, TEXT("FMetasoundCommentNodeIntVector::SerializeFromMismatchedTag - DeprecateSlateVector2D Type found: Resolving Mismatch"));
 			return true;
 		}
 	}
 #endif // WITH_EDITORONLY_DATA
 
+	const bool bIsCookCommandlet = IsRunningCookCommandlet();
 	if (Tag.GetType().IsStruct("DeprecateSlateVector2D"))
 	{
-		UE_LOG(LogMetaSound, Display, TEXT("FMetasoundCommentNodeIntVector::SerializeFromMismatchedTag - Missing Deprecated Type. Ignoring Mismatch"));
 		// Missing type, don't care about visualizing this data.
 		// Ignore the serialization.
 		X = 0;
 		Y = 0;
-		return true;
+
+		UE_LOG(LogMetaSound, Display, TEXT("FMetasoundCommentNodeIntVector::SerializeFromMismatchedTag - DeprecateSlateVector2D Type not loaded: Ignoring Mismatch"));
+		return bIsCookCommandlet;
 	}
 
 	if (Tag.GetType().IsStruct(NAME_Vector2d))
 	{
-		UE_LOG(LogMetaSound, Display, TEXT("FMetasoundCommentNodeIntVector::SerializeFromMismatchedTag - Resolving as FVector2d"));
 		FVector2d OldVector;
 		Slot << OldVector;
 		X = static_cast<int32>(OldVector.X);
@@ -380,7 +382,6 @@ bool FMetasoundCommentNodeIntVector::SerializeFromMismatchedTag(const struct FPr
 	}
 	else if (Tag.GetType().IsStruct(NAME_Vector2f))
 	{
-		UE_LOG(LogMetaSound, Display, TEXT("FMetasoundCommentNodeIntVector::SerializeFromMismatchedTag - Resolving as FVector2f"));
 		FVector2f OldVector;
 		Slot << OldVector;
 		X = static_cast<int32>(OldVector.X);
@@ -391,7 +392,6 @@ bool FMetasoundCommentNodeIntVector::SerializeFromMismatchedTag(const struct FPr
 	{
 		if (Slot.GetUnderlyingArchive().UEVer() < EUnrealEngineObjectUE5Version::LARGE_WORLD_COORDINATES)
 		{
-			UE_LOG(LogMetaSound, Display, TEXT("FMetasoundCommentNodeIntVector::SerializeFromMismatchedTag - Resolving as FVector2f (Pre LWC)"));
 			FVector2f OldVector;
 			Slot << OldVector;
 			X = static_cast<int32>(OldVector.X);
@@ -400,7 +400,6 @@ bool FMetasoundCommentNodeIntVector::SerializeFromMismatchedTag(const struct FPr
 		}
 		else
 		{
-			UE_LOG(LogMetaSound, Display, TEXT("FMetasoundCommentNodeIntVector::SerializeFromMismatchedTag - Resolving as FVector2d (Post LWC)"));
 			FVector2d OldVector;
 			Slot << OldVector;
 			X = static_cast<int32>(OldVector.X);
@@ -409,14 +408,12 @@ bool FMetasoundCommentNodeIntVector::SerializeFromMismatchedTag(const struct FPr
 		}
 	}
 
-	const bool bIsCookCommandlet = IsRunningCookCommandlet();
-
 	// Hack for running cook commandlet, where we don't really care
-	// if it fails as cooked content will never be editable, so don't
-	// bother reporting if old data was not translated.
+	// if it fails as cooked comment content will never be visible,
+	// so don't bother reporting if old data was not translated.
 	if (bIsCookCommandlet)
 	{
-		UE_LOG(LogMetaSound, Display, TEXT("FMetasoundCommentNodeIntVector::SerializeFromMismatchedTag - Did not resolve. Is in cook so ok not serializing."));
+		UE_LOG(LogMetaSound, Display, TEXT("FMetasoundCommentNodeIntVector::SerializeFromMismatchedTag - Did not resolve. Ignoring value (cooking content, value not necessary)."));
 	}
 	else
 	{
