Fix for graph undo sometimes leaving residual connection data when swapping paged graphs in MetaSound Editor

#rnx

diff --git a/Engine/Plugins/Runtime/Metasound/Source/MetasoundEditor/Private/MetasoundEditorSubsystem.cpp b/Engine/Plugins/Runtime/Metasound/Source/MetasoundEditor/Private/MetasoundEditorSubsystem.cpp
--- a/Engine/Plugins/Runtime/Metasound/Source/MetasoundEditor/Private/MetasoundEditorSubsystem.cpp
+++ b/Engine/Plugins/Runtime/Metasound/Source/MetasoundEditor/Private/MetasoundEditorSubsystem.cpp
@@ -499,6 +499,8 @@ bool UMetaSoundEditorSubsystem::SetFocusedPageInternal(FName PageName, const FGu
 	{
 		Builder.Modify();
 		UObject& MetaSound = DocBuilder.CastDocumentObjectChecked<UObject>();
+		MetaSound.Modify();
+		DocBuilder.GetMetasoundAsset().GetGraphChecked().Modify();
 		if (Builder.GetBuilder().SetBuildPageID(InPageID))
 		{
 			// Reregister to ensure all future audible instances are using the new page implementation.
