Fix Metasound Variable or Vertex rename crash if string is too long.

diff --git a/Engine/Plugins/Runtime/Metasound/Source/MetasoundEditor/Private/MetasoundEditorGraph.cpp b/Engine/Plugins/Runtime/Metasound/Source/MetasoundEditor/Private/MetasoundEditorGraph.cpp
--- a/Engine/Plugins/Runtime/Metasound/Source/MetasoundEditor/Private/MetasoundEditorGraph.cpp
+++ b/Engine/Plugins/Runtime/Metasound/Source/MetasoundEditor/Private/MetasoundEditorGraph.cpp
@@ -506,6 +506,12 @@ bool UMetasoundEditorGraphVertex::CanRename(const FText& InNewText, FText& OutEr
 		return false;
 	}
 
+	if (NewNameString.Len() >= NAME_SIZE)
+	{
+		OutError = FText::Format(LOCTEXT("GraphVertexRenameInvalid_NameTooLong", "Name cannot be longer than {0} characters."), NAME_SIZE);
+		return false;
+	}
+
 	if (IsInterfaceMember())
 	{
 		const FText CurrentMemberName = FText::FromName(GetMemberName());
@@ -1701,6 +1707,12 @@ bool UMetasoundEditorGraphVariable::CanRename(const FText& InNewText, FText& Out
 		return false;
 	}
 
+	if (NewNameString.Len() >= NAME_SIZE)
+	{
+		OutError = FText::Format(LOCTEXT("GraphVariableRenameInvalid_NameTooLong", "Name cannot be longer than {0} characters."), NAME_SIZE);
+		return false;
+	}
+
 	const FName NewName(*NewNameString);
 	FName Namespace;
 	FName ParameterName;
