MetaSound Editor Perf Optimizations Part 3:
- Synchronization clean-up
- Misc controller API to builder API clean-up

#tests automated, unit, editor smoke, local cook
#rnx

diff --git a/Engine/Plugins/Runtime/Metasound/Source/MetasoundEditor/Private/FindInMetasound.cpp b/Engine/Plugins/Runtime/Metasound/Source/MetasoundEditor/Private/FindInMetasound.cpp
--- a/Engine/Plugins/Runtime/Metasound/Source/MetasoundEditor/Private/FindInMetasound.cpp
+++ b/Engine/Plugins/Runtime/Metasound/Source/MetasoundEditor/Private/FindInMetasound.cpp
@@ -680,7 +680,7 @@ namespace Metasound::Editor
 					// Get specific node type and display name for variables
 					if (const UMetasoundEditorGraphVariableNode* VariableNode = Cast<UMetasoundEditorGraphVariableNode>(MetaSoundNode))
 					{
-						NodeName = VariableNode->GetClassName().ToString();
+						NodeName = VariableNode->GetBreadcrumb().ClassName.ToString();
 						switch (VariableNode->GetClassType())
 						{
 						case EMetasoundFrontendClassType::VariableMutator:
diff --git a/Engine/Plugins/Runtime/Metasound/Source/MetasoundEditor/Private/MetasoundEditorDocumentClipboardUtils.cpp b/Engine/Plugins/Runtime/Metasound/Source/MetasoundEditor/Private/MetasoundEditorDocumentClipboardUtils.cpp
--- a/Engine/Plugins/Runtime/Metasound/Source/MetasoundEditor/Private/MetasoundEditorDocumentClipboardUtils.cpp
+++ b/Engine/Plugins/Runtime/Metasound/Source/MetasoundEditor/Private/MetasoundEditorDocumentClipboardUtils.cpp
@@ -289,7 +289,7 @@ namespace Metasound::Editor
 			if (ensure(FrontendVariable))
 			{
 				// Can only have one mutator/setter node
-				const FNodeClassName NodeClassName = VariableNode->GetClassName().ToNodeClassName();
+				const FNodeClassName NodeClassName = VariableNode->GetBreadcrumb().ClassName.ToNodeClassName();
 				const bool bMatchesMutatorNodeID = VariableNode->GetNodeID() == FrontendVariable->MutatorNodeID;
 				bool bIsDuplicateMutatorNode = false;
 				if (const FMetasoundFrontendNode* MutatorNode = DocBuilder.FindNode(FrontendVariable->MutatorNodeID))
diff --git a/Engine/Plugins/Runtime/Metasound/Source/MetasoundEditor/Private/MetasoundEditorGraph.cpp b/Engine/Plugins/Runtime/Metasound/Source/MetasoundEditor/Private/MetasoundEditorGraph.cpp
--- a/Engine/Plugins/Runtime/Metasound/Source/MetasoundEditor/Private/MetasoundEditorGraph.cpp
+++ b/Engine/Plugins/Runtime/Metasound/Source/MetasoundEditor/Private/MetasoundEditorGraph.cpp
@@ -2889,6 +2889,7 @@ void UMetasoundEditorGraph::IterateMembers(TFunctionRef<void(UMetasoundEditorGra
 void UMetasoundEditorGraph::ValidateInternal(Metasound::Editor::FGraphValidationResults& OutResults)
 {
 	using namespace Metasound::Editor;
+	using namespace Metasound::Engine;
 	using namespace Metasound::Frontend;
 
 	OutResults = FGraphValidationResults();
@@ -2896,6 +2897,7 @@ void UMetasoundEditorGraph::ValidateInternal(Metasound::Editor::FGraphValidation
 	TArray<UMetasoundEditorGraphNode*> NodesToValidate;
 	GetNodesOfClass<UMetasoundEditorGraphNode>(NodesToValidate);
 	bool bNodeIdFound = false;
+	const FMetaSoundFrontendDocumentBuilder& Builder = FDocumentBuilderRegistry::GetChecked().FindOrBeginBuilding(&GetMetasoundChecked());
 	for (UMetasoundEditorGraphNode* Node : NodesToValidate)
 	{
 		Node->CacheBreadcrumb();
@@ -2913,7 +2915,7 @@ void UMetasoundEditorGraph::ValidateInternal(Metasound::Editor::FGraphValidation
 			}
 		}
 
-		Node->Validate(NodeResult);
+		Node->Validate(Builder, NodeResult);
 
 		OutResults.NodeResults.Add(MoveTemp(NodeResult));
 	}
diff --git a/Engine/Plugins/Runtime/Metasound/Source/MetasoundEditor/Private/MetasoundEditorGraphBuilder.cpp b/Engine/Plugins/Runtime/Metasound/Source/MetasoundEditor/Private/MetasoundEditorGraphBuilder.cpp
--- a/Engine/Plugins/Runtime/Metasound/Source/MetasoundEditor/Private/MetasoundEditorGraphBuilder.cpp
+++ b/Engine/Plugins/Runtime/Metasound/Source/MetasoundEditor/Private/MetasoundEditorGraphBuilder.cpp
@@ -384,75 +384,6 @@ namespace Metasound
 			return DisplayName;
 		}
 
-		FName FGraphBuilder::GetPinName(const Frontend::IOutputController& InFrontendOutput)
-		{
-			using namespace VariableNames; 
-
-			Frontend::FConstNodeHandle OwningNode = InFrontendOutput.GetOwningNode();
-			EMetasoundFrontendClassType OwningNodeClassType = OwningNode->GetClassMetadata().GetType();
-
-			switch (OwningNodeClassType)
-			{
-				case EMetasoundFrontendClassType::Variable:
-				case EMetasoundFrontendClassType::VariableAccessor:
-				case EMetasoundFrontendClassType::VariableDeferredAccessor:
-				case EMetasoundFrontendClassType::VariableMutator:
-				{
-					// All variables nodes use the same pin name for user-modifiable node
-					// inputs and outputs and the editor does not display the pin's name. The
-					// editor instead displays the variable's name in place of the pin name to
-					// maintain a consistent look and behavior to input and output nodes.
-					return METASOUND_GET_PARAM_NAME(OutputData);
-				}
-				case EMetasoundFrontendClassType::Input:
-				case EMetasoundFrontendClassType::Output:
-				{
-					return OwningNode->GetNodeName();
-				}
-
-				case EMetasoundFrontendClassType::Template:
-				default:
-				{
-					return InFrontendOutput.GetName();
-				}
-			}
-		}
-
-		FName FGraphBuilder::GetPinName(const Frontend::IInputController& InFrontendInput)
-		{
-			using namespace VariableNames;
-			
-			Frontend::FConstNodeHandle OwningNode = InFrontendInput.GetOwningNode();
-			EMetasoundFrontendClassType OwningNodeClassType = OwningNode->GetClassMetadata().GetType();
-
-			switch (OwningNodeClassType)
-			{
-				case EMetasoundFrontendClassType::Variable:
-				case EMetasoundFrontendClassType::VariableAccessor:
-				case EMetasoundFrontendClassType::VariableDeferredAccessor:
-				case EMetasoundFrontendClassType::VariableMutator:
-				{
-					// All variables nodes use the same pin name for user-modifiable node
-					// inputs and outputs and the editor does not display the pin's name. The
-					// editor instead displays the variable's name in place of the pin name to
-					// maintain a consistent look and behavior to input and output nodes.
-					return METASOUND_GET_PARAM_NAME(InputData);
-				}
-
-				case EMetasoundFrontendClassType::Input:
-				case EMetasoundFrontendClassType::Output:
-				{
-					return OwningNode->GetNodeName();
-				}
-
-				case EMetasoundFrontendClassType::Template:
-				default:
-				{
-					return InFrontendInput.GetName();
-				}
-			}
-		}
-
 		UMetasoundEditorGraphExternalNode* FGraphBuilder::AddTemplateNode(UObject& InMetaSound, const FGuid& InNodeID, const FMetasoundFrontendClassMetadata& InMetadata, bool bInSelectNewNode)
 		{
 			const EMetasoundFrontendClassType ClassType = InMetadata.GetType();
@@ -1487,19 +1418,25 @@ PRAGMA_ENABLE_DEPRECATION_WARNINGS
 				}
 			}
 
-			FInputHandle InputHandle = GetInputHandleFromPin(&InInputPin);
-			FOutputHandle OutputHandle = GetOutputHandleFromPin(&InOutputPin);
-			if (!InputHandle->IsValid() || !OutputHandle->IsValid())
-			{
-				return false;
-			}
-
-			if (!ensure(InputHandle->Connect(*OutputHandle)))
+			FMetaSoundFrontendDocumentBuilder& Builder = GetBuilderFromPinChecked(InInputPin).GetBuilder();
+			const FMetasoundFrontendVertexHandle InputVertexHandle = GetPinVertexHandle(Builder, &InInputPin);
+			const FMetasoundFrontendVertexHandle OutputVertexHandle = GetPinVertexHandle(Builder, &InOutputPin);
+			if (!InputVertexHandle.IsSet() || !OutputVertexHandle.IsSet())
 			{
 				InInputPin.BreakLinkTo(&InOutputPin);
 				return false;
 			}
 
+			Builder.RemoveEdgeToNodeInput(InputVertexHandle.NodeID, InputVertexHandle.VertexID);
+
+			FMetasoundFrontendEdge NewEdge
+			{
+				OutputVertexHandle.NodeID,
+				OutputVertexHandle.VertexID,
+				InputVertexHandle.NodeID,
+				InputVertexHandle.VertexID,
+			};
+			Builder.AddEdge(MoveTemp(NewEdge));
 			return true;
 		}
 
@@ -1573,41 +1510,6 @@ PRAGMA_ENABLE_DEPRECATION_WARNINGS
 			return bNewGraphBound;
 		}
 
-		void FGraphBuilder::RebuildNodePins(UMetasoundEditorGraphNode& InGraphNode)
-		{
-			using namespace Frontend;
-		
-			for (int32 i = InGraphNode.Pins.Num() - 1; i >= 0; i--)
-			{
-				InGraphNode.RemovePin(InGraphNode.Pins[i]);
-			}
-
-			FConstNodeHandle NodeHandle = InGraphNode.GetConstNodeHandle();
-			TArray<FConstInputHandle> InputHandles = NodeHandle->GetConstInputs();
-			NodeHandle->GetInputStyle().SortDefaults(InputHandles, [](const Frontend::FConstInputHandle& Handle) { return FGraphBuilder::GetDisplayName(*Handle); });
-			for (const FConstInputHandle& InputHandle : InputHandles)
-			{
-				// Only add pins of the node if the connection is user modifiable. 
-				// Connections which the user cannot modify are controlled elsewhere.
-				if (InputHandle->IsConnectionUserModifiable())
-				{
-					AddPinToNode(InGraphNode, InputHandle);
-				}
-			}
-
-			TArray<FConstOutputHandle> OutputHandles = NodeHandle->GetConstOutputs();
-			NodeHandle->GetOutputStyle().SortDefaults(OutputHandles, [](const Frontend::FConstOutputHandle& Handle) { return FGraphBuilder::GetDisplayName(*Handle); });
-			for (const FConstOutputHandle& OutputHandle : OutputHandles)
-			{
-				// Only add pins of the node if the connection is user modifiable. 
-				// Connections which the user cannot modify are controlled elsewhere.
-				if (OutputHandle->IsConnectionUserModifiable())
-				{
-					AddPinToNode(InGraphNode, OutputHandle);
-				}
-			}
-		}
-
 		void FGraphBuilder::RefreshPinMetadata(UEdGraphPin& InPin, bool bAdvancedView)
 		{
 			UMetasoundEditorGraphNode* MetasoundGraphNode = Cast<UMetasoundEditorGraphNode>(InPin.GetOwningNode());
@@ -1680,33 +1582,23 @@ PRAGMA_ENABLE_DEPRECATION_WARNINGS
 			}
 		}
 
-		bool FGraphBuilder::IsMatchingInputHandleAndPin(const Frontend::FConstInputHandle& InInputHandle, const UEdGraphPin& InEditorPin)
+		bool FGraphBuilder::IsMatchingInputAndPin(const FMetaSoundFrontendDocumentBuilder& InBuilder, const FMetasoundFrontendVertex& InVertex, const UEdGraphPin& InEditorPin)
 		{
-			if (InEditorPin.Direction != EGPD_Input)
-			{
-				return false;
-			}
-
-			Frontend::FConstInputHandle PinInputHandle = GetConstInputHandleFromPin(&InEditorPin);
-			if (PinInputHandle->GetID() == InInputHandle->GetID())
+			if (InEditorPin.Direction == EGPD_Input)
 			{
-				return true;
+				const FMetasoundFrontendVertex* PinVertex = GetPinVertex(InBuilder, &InEditorPin);
+				return PinVertex && PinVertex->VertexID == InVertex.VertexID;
 			}
 
 			return false;
 		}
 
-		bool FGraphBuilder::IsMatchingOutputHandleAndPin(const Frontend::FConstOutputHandle& InOutputHandle, const UEdGraphPin& InEditorPin)
+		bool FGraphBuilder::IsMatchingOutputAndPin(const FMetaSoundFrontendDocumentBuilder& InBuilder, const FMetasoundFrontendVertex& InVertex, const UEdGraphPin& InEditorPin)
 		{
-			if (InEditorPin.Direction != EGPD_Output)
+			if (InEditorPin.Direction == EGPD_Output)
 			{
-				return false;
-			}
-
-			Frontend::FConstOutputHandle PinOutputHandle = GetConstOutputHandleFromPin(&InEditorPin);
-			if (PinOutputHandle->GetID() == InOutputHandle->GetID())
-			{
-				return true;
+				const FMetasoundFrontendVertex* PinVertex = GetPinVertex(InBuilder, &InEditorPin);
+				return PinVertex && PinVertex->VertexID == InVertex.VertexID;
 			}
 
 			return false;
@@ -1734,47 +1626,40 @@ PRAGMA_ENABLE_DEPRECATION_WARNINGS
 			}
 		}
 
-		UEdGraphPin* FGraphBuilder::AddPinToNode(UMetasoundEditorGraphNode& InEditorNode, Frontend::FConstInputHandle InInputHandle)
+		UEdGraphPin* FGraphBuilder::AddPinToNode(
+			const IMetasoundEditorModule& EditorModule,
+			const FMetaSoundFrontendDocumentBuilder& InBuilder,
+			UMetasoundEditorGraphNode& InEditorNode,
+			const FMetasoundFrontendClass& InClass,
+			const FMetasoundFrontendNode& InNode,
+			const FMetasoundFrontendVertex& InVertex,
+			EEdGraphPinDirection Direction)
 		{
 			using namespace Frontend;
 
 			FEdGraphPinType PinType;
-			FName DataTypeName = InInputHandle->GetDataType();
-
-			const IMetasoundEditorModule& EditorModule = FModuleManager::GetModuleChecked<IMetasoundEditorModule>("MetaSoundEditor");
-			if (const FEdGraphPinType* RegisteredPinType = EditorModule.FindPinType(DataTypeName))
+			if (const FEdGraphPinType* RegisteredPinType = EditorModule.FindPinType(InVertex.TypeName))
 			{
 				PinType = *RegisteredPinType;
 			}
 
-			FName PinName = FGraphBuilder::GetPinName(*InInputHandle);
-			UEdGraphPin* NewPin = InEditorNode.CreatePin(EGPD_Input, PinType, PinName);
+			UEdGraphPin* NewPin = InEditorNode.CreatePin(Direction, PinType, InVertex.Name);
 			if (ensure(NewPin))
 			{
-				RefreshPinMetadata(*NewPin, InInputHandle->GetMetadata().bIsAdvancedDisplay);
-				const FMetaSoundFrontendDocumentBuilder& Builder = InEditorNode.GetBuilderChecked().GetConstBuilder();
-				SynchronizePinLiteral(Builder, *NewPin);
-			}
-
-			return NewPin;
-		}
-
-		UEdGraphPin* FGraphBuilder::AddPinToNode(UMetasoundEditorGraphNode& InEditorNode, Frontend::FConstOutputHandle InOutputHandle)
-		{
-			FEdGraphPinType PinType;
-			FName DataTypeName = InOutputHandle->GetDataType();
+				auto IsVertexMetadata = [&InVertex](const FMetasoundFrontendClassVertex& ClassVertex) { return ClassVertex.Name == InVertex.Name; };
+				const FMetasoundFrontendClassVertex* ClassVertex = nullptr;
+				if (Direction == EGPD_Input)
+				{
+					ClassVertex = InClass.GetInterfaceForNode(InNode).Inputs.FindByPredicate(IsVertexMetadata);
+					SynchronizePinLiteral(InBuilder, *NewPin);
+				}
+				else
+				{
+					ClassVertex = InClass.GetInterfaceForNode(InNode).Outputs.FindByPredicate(IsVertexMetadata);
+				}
 
-			IMetasoundEditorModule& EditorModule = FModuleManager::GetModuleChecked<IMetasoundEditorModule>("MetaSoundEditor");
-			if (const FEdGraphPinType* RegisteredPinType = EditorModule.FindPinType(DataTypeName))
-			{
-				PinType = *RegisteredPinType;
-			}
+				RefreshPinMetadata(*NewPin, ClassVertex ? ClassVertex->Metadata.bIsAdvancedDisplay : false);
 
-			FName PinName = FGraphBuilder::GetPinName(*InOutputHandle);
-			UEdGraphPin* NewPin = InEditorNode.CreatePin(EGPD_Output, PinType, PinName);
-			if (ensure(NewPin))
-			{
-				RefreshPinMetadata(*NewPin, InOutputHandle->GetMetadata().bIsAdvancedDisplay);
 			}
 
 			return NewPin;
@@ -1814,11 +1699,6 @@ PRAGMA_ENABLE_DEPRECATION_WARNINGS
 
 			if (InPin.PinType != PinType)
 			{
-				if (const UMetasoundEditorGraphNode* Node = Cast<UMetasoundEditorGraphNode>(InPin.GetOwningNodeUnchecked()))
-				{
-					const FString NodeName = Node->GetDisplayName().ToString();
-					UE_LOG(LogMetasoundEditor, Verbose, TEXT("Synchronizing Pin '%s' on Node '%s': Type converted to '%s'"), *NodeName, *InPin.GetName(), *InDataType.ToString());
-				}
 				InPin.PinType = PinType;
 				return true;
 			}
@@ -1921,44 +1801,25 @@ PRAGMA_ENABLE_DEPRECATION_WARNINGS
 			{
 				bool bIsNodeDirty = false;
 
-				FConstNodeHandle Node = EditorNode->GetConstNodeHandle();
-
 				TArray<UEdGraphPin*> Pins = EditorNode->GetAllPins();
-				TArray<FConstInputHandle> NodeInputs = Node->GetConstInputs();
-				TArray<FConstOutputHandle> NodeOutputs = Node->GetConstOutputs();
-
-				// Ignore connections which are not handled by the editor.
-				NodeInputs.RemoveAll([](const FConstInputHandle& FrontendInput) { return !FrontendInput->IsConnectionUserModifiable(); });
-				NodeOutputs.RemoveAll([](const FConstOutputHandle& FrontendOutput) { return !FrontendOutput->IsConnectionUserModifiable(); });
-
-				for (FConstInputHandle& NodeInput : NodeInputs)
+				TArray<const FMetasoundFrontendVertex*> NodeInputs = InBuilder.FindUserModifiableNodeInputs(EditorNode->GetNodeID());
+				for (const FMetasoundFrontendVertex* NodeInput : NodeInputs)
 				{
-					auto IsMatchingInputPin = [&](const UEdGraphPin* Pin) -> bool
-					{
-						return IsMatchingInputHandleAndPin(NodeInput, *Pin);
-					};
+					auto IsMatchingInputPin = [&](const UEdGraphPin* Pin) { return IsMatchingInputAndPin(InBuilder, *NodeInput, *Pin); };
 
 					UEdGraphPin* MatchingPin = nullptr;
-					if (UEdGraphPin** DoublePointer = Pins.FindByPredicate(IsMatchingInputPin))
-					{
-						MatchingPin = *DoublePointer;
-					}
-
-					if (!ensure(MatchingPin))
+					const int32 PinIndex = Pins.IndexOfByPredicate(IsMatchingInputPin);
+					if (!ensure(PinIndex != INDEX_NONE))
 					{
 						continue;
 					}
 
-					// Remove pin so it isn't used twice.
-					Pins.Remove(MatchingPin);
+					MatchingPin = Pins[PinIndex];
+					Pins.RemoveAtSwap(PinIndex, EAllowShrinking::No);
 
-					FConstOutputHandle OutputHandle = NodeInput->GetConnectedOutput();
-					if (OutputHandle->IsValid())
+					const FMetasoundFrontendNode* ConnectedOutputNode = nullptr;
+					if (const FMetasoundFrontendVertex* Output = InBuilder.FindNodeOutputConnectedToNodeInput(EditorNode->GetNodeID(), NodeInput->VertexID, &ConnectedOutputNode))
 					{
-						// Both input and output handles be user modifiable for a
-						// connection to be controlled by the editor.
-						check(OutputHandle->IsConnectionUserModifiable());
-
 						bool bAddLink = false;
 
 						if (MatchingPin->LinkedTo.IsEmpty())
@@ -1966,7 +1827,7 @@ PRAGMA_ENABLE_DEPRECATION_WARNINGS
 							// No link currently exists. Add the appropriate link.
 							bAddLink = true;
 						}
-						else if (!IsMatchingOutputHandleAndPin(OutputHandle, *MatchingPin->LinkedTo[0]))
+						else if (!IsMatchingOutputAndPin(InBuilder, *Output, *MatchingPin->LinkedTo[0]))
 						{
 							// The wrong link exists.
 							constexpr bool bNotifyNodes = false;
@@ -1977,12 +1838,11 @@ PRAGMA_ENABLE_DEPRECATION_WARNINGS
 
 						if (bAddLink)
 						{
-							const FGuid NodeID = OutputHandle->GetOwningNodeID();
 							const FText& OwningNodeName = EditorNode->GetDisplayName();
-							UMetasoundEditorGraphNode* OutputEditorNode = EditorNodesByFrontendID.FindRef(NodeID);
+							UMetasoundEditorGraphNode* OutputEditorNode = EditorNodesByFrontendID.FindRef(ConnectedOutputNode->GetID());
 							if (OutputEditorNode)
 							{
-								UEdGraphPin* OutputPin = OutputEditorNode->FindPinChecked(OutputHandle->GetName(), EEdGraphPinDirection::EGPD_Output);
+								UEdGraphPin* OutputPin = OutputEditorNode->FindPinChecked(Output->Name, EEdGraphPinDirection::EGPD_Output);
 
 								UE_LOG(LogMetasoundEditor, Verbose, TEXT("Synchronizing Node '%s' Connection: Linking Pin '%s' to '%s'"), *OwningNodeName.ToString(), *MatchingPin->GetName(), *OutputPin->GetName());
 
@@ -2004,9 +1864,6 @@ PRAGMA_ENABLE_DEPRECATION_WARNINGS
 							constexpr bool bNotifyNodes = false;
 							constexpr bool bMarkDirty = false;
 							MatchingPin->BreakAllPinLinks(bNotifyNodes, bMarkDirty);
-							const FText OwningNodeName = EditorNode->GetDisplayName();
-							const FText InputName = FGraphBuilder::GetDisplayName(*NodeInput);
-							UE_LOG(LogMetasoundEditor, Verbose, TEXT("Synchronizing Node '%s' Connection: Breaking all pin links to '%s'"), *OwningNodeName.ToString(), *InputName.ToString());
 							bIsNodeDirty = true;
 						}
 					}
@@ -2015,15 +1872,13 @@ PRAGMA_ENABLE_DEPRECATION_WARNINGS
 				}
 
 				// Handle node outputs to break connections for the case 
-				// where the connected node input is on a node that has been deleted, 
+				// where the connected node input is on a node that has been deleted,
 				// so it wasn't handled by the node inputs direction above
-				for (FConstOutputHandle& NodeOutput : NodeOutputs)
+				TArray<const FMetasoundFrontendVertex*> NodeOutputs = InBuilder.FindUserModifiableNodeOutputs(EditorNode->GetNodeID());
+				for (const FMetasoundFrontendVertex* NodeOutput : NodeOutputs)
 				{
 					// Find pin
-					auto IsMatchingOutputPin = [&](const UEdGraphPin* Pin) -> bool
-					{
-						return IsMatchingOutputHandleAndPin(NodeOutput, *Pin);
-					};
+					auto IsMatchingOutputPin = [&](const UEdGraphPin* Pin) { return IsMatchingOutputAndPin(InBuilder, *NodeOutput, *Pin); };
 
 					UEdGraphPin* MatchingPin = nullptr;
 					if (UEdGraphPin** DoublePointer = Pins.FindByPredicate(IsMatchingOutputPin))
@@ -2038,14 +1893,11 @@ PRAGMA_ENABLE_DEPRECATION_WARNINGS
 
 					// Remove connected pins from removal list
 					TArray<UEdGraphPin*> PinsToBreak = MatchingPin->LinkedTo;
-					TArray<FConstInputHandle> InputHandles = NodeOutput->GetConstConnectedInputs();
-					for (FConstInputHandle InputHandle : InputHandles)
+					TArray<const FMetasoundFrontendNode*> ConnectedInputNodes;
+					TArray<const FMetasoundFrontendVertex*> ConnectedInputs = InBuilder.FindNodeInputsConnectedToNodeOutput(EditorNode->GetNodeID(), NodeOutput->VertexID, &ConnectedInputNodes);
+					for (const FMetasoundFrontendVertex* ConnectedInput : ConnectedInputs)
 					{
-						UEdGraphPin** InputPin = PinsToBreak.FindByPredicate(
-							[&](const UEdGraphPin* Pin) -> bool
-							{
-								return IsMatchingInputHandleAndPin(InputHandle, *Pin);
-							});
+						UEdGraphPin** InputPin = PinsToBreak.FindByPredicate([&](const UEdGraphPin* Pin) { return IsMatchingInputAndPin(InBuilder, *ConnectedInput, *Pin); });
 						if (InputPin)
 						{
 							UEdGraphPin* FoundPin = *InputPin;
@@ -2061,9 +1913,6 @@ PRAGMA_ENABLE_DEPRECATION_WARNINGS
 					{
 						constexpr bool bMarkDirty = false;
 						MatchingPin->BreakLinkTo(PinToBreak, bMarkDirty);
-						const FText OwningNodeName = EditorNode->GetDisplayName();
-						const FText OutputName = FGraphBuilder::GetDisplayName(*NodeOutput);
-						UE_LOG(LogMetasoundEditor, Verbose, TEXT("Synchronizing Node '%s' Connection: Breaking a pin link to '%s'"), *OwningNodeName.ToString(), *OutputName.ToString());
 						bIsNodeDirty = true;
 					}
 				}
@@ -2074,7 +1923,7 @@ PRAGMA_ENABLE_DEPRECATION_WARNINGS
 			return bIsGraphDirty;
 		}
 
-		bool FGraphBuilder::SynchronizeGraph(const FMetaSoundFrontendDocumentBuilder& InBuilder, UMetasoundEditorGraph& OutGraph, bool bSkipIfModifyContextUnchanged)
+		bool FGraphBuilder::SynchronizeGraph(FMetaSoundFrontendDocumentBuilder& InBuilder, UMetasoundEditorGraph& OutGraph, bool bSkipIfModifyContextUnchanged)
 		{
 			TRACE_CPUPROFILER_EVENT_SCOPE(Metasound::Editor::FGraphBuilder::SynchronizeGraph);
 
@@ -2107,7 +1956,7 @@ PRAGMA_ENABLE_DEPRECATION_WARNINGS
 
 				if (EditedReferencingMetaSounds.IsEmpty())
 				{
-					MetaSoundAsset->CacheRegistryMetadata();
+					InBuilder.CacheRegistryMetadata();
 					GraphBuilderPrivate::SynchronizeGraphRecursively(InBuilder, OutGraph, bSkipIfModifyContextUnchanged);
 					GraphBuilderPrivate::RecurseClearDocumentModified(*MetaSoundAsset);
 				}
@@ -2118,7 +1967,7 @@ PRAGMA_ENABLE_DEPRECATION_WARNINGS
 						check(EditedMetaSound);
 						UObject* OwningMetaSound = EditedMetaSound->GetOwningAsset();
 						check(OwningMetaSound);
-						const FMetaSoundFrontendDocumentBuilder& EditedBuilder = IDocumentBuilderRegistry::GetChecked().FindOrBeginBuilding(OwningMetaSound);
+						FMetaSoundFrontendDocumentBuilder& EditedBuilder = IDocumentBuilderRegistry::GetChecked().FindOrBeginBuilding(OwningMetaSound);
 						SynchronizeGraph(EditedBuilder, *CastChecked<UMetasoundEditorGraph>(&EditedMetaSound->GetGraphChecked()), bSkipIfModifyContextUnchanged);
 					}
 				}
@@ -2176,71 +2025,80 @@ PRAGMA_ENABLE_DEPRECATION_WARNINGS
 			TRACE_CPUPROFILER_EVENT_SCOPE(Metasound::Editor::FGraphBuilder::SynchronizeNodes);
 
 			using namespace Frontend;
+			using FNodeClassPair = TPair<const FMetasoundFrontendNode*, const FMetasoundFrontendClass*>;
 
 			// Get all external nodes from Frontend graph.  Input and output references will only be added/synchronized
 			// if required when synchronizing connections (as they are not required to inhabit editor graph).
 			UObject& MetaSoundObject = InBuilder.CastDocumentObjectChecked<UObject>();
 			const FMetasoundAssetBase* MetaSoundAsset = IMetasoundUObjectRegistry::Get().GetObjectAsAssetBase(&MetaSoundObject);
 			check(MetaSoundAsset);
-			FConstGraphHandle GraphHandle = MetaSoundAsset->GetRootGraphHandle();
-			TArray<FConstNodeHandle> FrontendNodes = GraphHandle->GetConstNodes();
+			TArray<FNodeClassPair> FrontendNodeAndClass;
 			TArray<UMetasoundEditorGraphNode*> EditorNodes;
-			OutGraph.GetNodesOfClass(EditorNodes);
+
+			{
+				InBuilder.IterateNodes([&FrontendNodeAndClass](const FMetasoundFrontendClass& Class, const FMetasoundFrontendNode& Node)
+				{
+					// Input nodes use template input nodes to visually distinguish location, so ignore them.
+					if (Class.Metadata.GetType() != EMetasoundFrontendClassType::Input)
+					{
+						FrontendNodeAndClass.Add(FNodeClassPair(&Node, &Class));
+					}
+				});
+				OutGraph.GetNodesOfClass(EditorNodes);
+			}
 
 			// Find existing array of editor nodes associated with Frontend node
 			struct FAssociatedNodes
 			{
 				TArray<UMetasoundEditorGraphNode*> EditorNodes;
-				FConstNodeHandle Node = Metasound::Frontend::INodeController::GetInvalidHandle();
+				const FMetasoundFrontendNode* Node = nullptr;
+				const FMetasoundFrontendClass* Class = nullptr;
 			};
 			TMap<FGuid, FAssociatedNodes> AssociatedNodes;
 
-			// Input nodes use template input nodes to visually distinguish location, so ignore them.
-			FrontendNodes.RemoveAllSwap([](const FConstNodeHandle& Node)
-			{
-				return Node->GetClassMetadata().GetType() == EMetasoundFrontendClassType::Input;
-			},
-			EAllowShrinking::No);
-
-			// Reverse iterate so paired nodes can safely be removed from the array.
-			for (int32 i = FrontendNodes.Num() - 1; i >= 0; i--)
 			{
-				const FConstNodeHandle& Node = FrontendNodes[i];
-				bool bFoundEditorNode = false;
-				for (int32 j = EditorNodes.Num() - 1; j >= 0; --j)
+				// Reverse iterate so paired nodes can safely be removed from the array.
+				for (int32 i = FrontendNodeAndClass.Num() - 1; i >= 0; i--)
 				{
-					UMetasoundEditorGraphNode* EditorNode = EditorNodes[j];
-					EditorNode->CacheBreadcrumb();
-					if (EditorNode->GetNodeID() == Node->GetID())
+					bool bFoundEditorNode = false;
+					const FMetasoundFrontendNode& Node = *FrontendNodeAndClass[i].Key;
+					const FMetasoundFrontendClass& Class = *FrontendNodeAndClass[i].Value;
+					for (int32 j = EditorNodes.Num() - 1; j >= 0; --j)
 					{
-						// Editor node may have the same Frontend NodeID as another page,
-						// but may have been assigned a different Editor NodeID, so synchronize
-						// from frontend data here only if node location was able to sync.
-						constexpr bool bUpdateEditorNodeID = true;
-						const bool bLocationFound = EditorNode->SyncLocationFromFrontendNode(bUpdateEditorNodeID);
-						if (bLocationFound)
+						UMetasoundEditorGraphNode* EditorNode = EditorNodes[j];
+						EditorNode->CacheBreadcrumb();
+						if (EditorNode->GetNodeID() == Node.GetID())
 						{
-							bFoundEditorNode = true;
-							FAssociatedNodes& AssociatedNodeData = AssociatedNodes.FindOrAdd(Node->GetID());
-							if (AssociatedNodeData.Node->IsValid())
+							// Editor node may have the same Frontend NodeID as another page,
+							// but may have been assigned a different Editor NodeID, so synchronize
+							// from frontend data here only if node location was able to sync.
+							constexpr bool bUpdateEditorNodeID = true;
+							const bool bLocationFound = EditorNode->SyncLocationFromFrontendNode(bUpdateEditorNodeID);
+							if (bLocationFound)
 							{
-								ensure(AssociatedNodeData.Node == Node);
-							}
-							else
-							{
-								AssociatedNodeData.Node = Node;
-							}
+								bFoundEditorNode = true;
+								FAssociatedNodes& AssociatedNodeData = AssociatedNodes.FindOrAdd(Node.GetID());
+								if (AssociatedNodeData.Node)
+								{
+									ensure(AssociatedNodeData.Node == &Node);
+								}
+								else
+								{
+									AssociatedNodeData.Node = &Node;
+									AssociatedNodeData.Class = &Class;
+								}
 
-							EditorNode->SyncCommentFromFrontendNode();
-							AssociatedNodeData.EditorNodes.Add(EditorNode);
-							EditorNodes.RemoveAtSwap(j, EAllowShrinking::No);
+								EditorNode->SyncCommentFromFrontendNode();
+								AssociatedNodeData.EditorNodes.Add(EditorNode);
+								EditorNodes.RemoveAtSwap(j, EAllowShrinking::No);
+							}
 						}
 					}
-				}
 
-				if (bFoundEditorNode)
-				{
-					FrontendNodes.RemoveAtSwap(i, EAllowShrinking::No);
+					if (bFoundEditorNode)
+					{
+						FrontendNodeAndClass.RemoveAtSwap(i, EAllowShrinking::No);
+					}
 				}
 			}
 
@@ -2258,37 +2116,38 @@ PRAGMA_ENABLE_DEPRECATION_WARNINGS
 			}
 
 			// Add missing editor nodes marked as visible.
-			for (const FConstNodeHandle& Node : FrontendNodes)
+			for (const FNodeClassPair& Pair : FrontendNodeAndClass)
 			{
-				const FMetasoundFrontendNodeStyle& CurrentStyle = Node->GetNodeStyle();
+				const FMetasoundFrontendNodeStyle& CurrentStyle = Pair.Key->Style;
+				const FGuid& NodeID = Pair.Key->GetID();
+				const FMetasoundFrontendClassMetadata& ClassMetadata = Pair.Value->Metadata;
 				for (const TPair<FGuid, FVector2D>& Location : CurrentStyle.Display.Locations)
 				{
 					UMetasoundEditorGraphNode* NewGraphNode = nullptr;
-					const FMetasoundFrontendClassMetadata& ClassMetadata = Node->GetClassMetadata();
 					switch (ClassMetadata.GetType())
 					{
 						case EMetasoundFrontendClassType::External:
 						{
-							NewGraphNode = CastChecked<UMetasoundEditorGraphNode>(AddExternalNode(MetaSoundObject, Node->GetID(), ClassMetadata, false));
+							NewGraphNode = CastChecked<UMetasoundEditorGraphNode>(AddExternalNode(MetaSoundObject, NodeID, ClassMetadata, false));
 						}
 						break;
 
 						case EMetasoundFrontendClassType::Template:
 						{
-							if (Node->GetClassMetadata().GetClassName() == FInputNodeTemplate::ClassName)
+							if (Pair.Value->Metadata.GetClassName() == FInputNodeTemplate::ClassName)
 							{
-								NewGraphNode = CastChecked<UMetasoundEditorGraphNode>(AddInputNode(MetaSoundObject, Node->GetID(), false));
+								NewGraphNode = CastChecked<UMetasoundEditorGraphNode>(AddInputNode(MetaSoundObject, NodeID, false));
 							}
 							else
 							{
-								NewGraphNode = CastChecked<UMetasoundEditorGraphNode>(AddTemplateNode(MetaSoundObject, Node->GetID(), Node->GetClassMetadata(), false));
+								NewGraphNode = CastChecked<UMetasoundEditorGraphNode>(AddTemplateNode(MetaSoundObject, NodeID, ClassMetadata, false));
 							}
 						}
 						break;
 
 						case EMetasoundFrontendClassType::Output:
 						{
-							NewGraphNode = CastChecked<UMetasoundEditorGraphNode>(AddOutputNode(MetaSoundObject, Node->GetID(), false));
+							NewGraphNode = CastChecked<UMetasoundEditorGraphNode>(AddOutputNode(MetaSoundObject, NodeID, false));
 						}
 						break;
 
@@ -2297,7 +2156,7 @@ PRAGMA_ENABLE_DEPRECATION_WARNINGS
 						case EMetasoundFrontendClassType::VariableDeferredAccessor:
 						case EMetasoundFrontendClassType::Variable:
 						{
-							NewGraphNode = CastChecked<UMetasoundEditorGraphNode>(AddVariableNode(MetaSoundObject, Node->GetID(), false));
+							NewGraphNode = CastChecked<UMetasoundEditorGraphNode>(AddVariableNode(MetaSoundObject, NodeID, false));
 						}
 						break;
 
@@ -2324,14 +2183,15 @@ PRAGMA_ENABLE_DEPRECATION_WARNINGS
 						ensureMsgf(NewGraphNode->NodeGuid.IsValid(), TEXT("New editor NodeGuid must be valid."));
 						ensureMsgf(NewGraphNode->NodeGuid == Location.Key, TEXT("New editor NodeGuid must match location key"));
 
-						FAssociatedNodes& AssociatedNodeData = AssociatedNodes.FindOrAdd(Node->GetID());
-						if (AssociatedNodeData.Node->IsValid())
+						FAssociatedNodes& AssociatedNodeData = AssociatedNodes.FindOrAdd(NodeID);
+						if (AssociatedNodeData.Node)
 						{
-							ensure(AssociatedNodeData.Node == Node);
+							ensure(AssociatedNodeData.Node == Pair.Key);
 						}
 						else
 						{
-							AssociatedNodeData.Node = Node;
+							AssociatedNodeData.Node = Pair.Key;
+							AssociatedNodeData.Class = Pair.Value;
 						}
 
 						AssociatedNodeData.EditorNodes.Add(NewGraphNode);
@@ -2344,22 +2204,17 @@ PRAGMA_ENABLE_DEPRECATION_WARNINGS
 			{
 				for (UMetasoundEditorGraphNode* EditorNode : IdNodePair.Value.EditorNodes)
 				{
-					const bool bNodeDirtied = SynchronizeNodePins(*EditorNode, IdNodePair.Value.Node);
-					if (bNodeDirtied)
-					{
-						if (TSharedPtr<FEditor> MetaSoundEditor = FGraphBuilder::GetEditorForMetasound(MetaSoundObject))
-						{
-							if (TSharedPtr<SGraphEditor> GraphEditor = MetaSoundEditor->GetGraphEditor())
-							{
-								GraphEditor->RefreshNode(*EditorNode);
-							}
-						}
-					}
+					SynchronizeNodePins(InBuilder, *EditorNode, *IdNodePair.Value.Node, *IdNodePair.Value.Class);
 				}
 			}
 		}
 
-		bool FGraphBuilder::SynchronizeNodePins(UMetasoundEditorGraphNode& InEditorNode, Frontend::FConstNodeHandle InNode, bool bRemoveUnusedPins, bool bLogChanges)
+		bool FGraphBuilder::SynchronizeNodePins(
+			const FMetaSoundFrontendDocumentBuilder& InBuilder,
+			UMetasoundEditorGraphNode& InEditorNode,
+			const FMetasoundFrontendNode& InNode,
+			const FMetasoundFrontendClass& InClass,
+			bool bRemoveUnusedPins)
 		{
 			using namespace Frontend;
 
@@ -2367,44 +2222,25 @@ PRAGMA_ENABLE_DEPRECATION_WARNINGS
 
 			IMetasoundEditorModule& EditorModule = FModuleManager::GetModuleChecked<IMetasoundEditorModule>("MetaSoundEditor");
 
-			TArray<FConstInputHandle> InputHandles;
-			TArray<FConstOutputHandle> OutputHandles;
-			auto GetUserModifiableHandles = [InNode](TArray<FConstInputHandle>& InHandles, TArray<FConstOutputHandle>& OutHandles)
-			{
-				InHandles = InNode->GetConstInputs();
-				OutHandles = InNode->GetConstOutputs();
-
-				// Remove input and output handles which are not user modifiable
-				InHandles.RemoveAll([](const FConstInputHandle& FrontendInput) { return !FrontendInput->IsConnectionUserModifiable(); });
-				OutHandles.RemoveAll([](const FConstOutputHandle& FrontendOutput) { return !FrontendOutput->IsConnectionUserModifiable(); });
-			};
-			GetUserModifiableHandles(InputHandles, OutputHandles);
+			const FGuid& NodeID = InNode.GetID();
+			TArray<const FMetasoundFrontendVertex*> NodeInputs = InBuilder.FindUserModifiableNodeInputs(NodeID);
+			TArray<const FMetasoundFrontendVertex*> NodeOutputs = InBuilder.FindUserModifiableNodeOutputs(NodeID);
 
 			// Filter out pins which are not paired.
 			TArray<UEdGraphPin*> EditorPins = InEditorNode.Pins;
 			for (int32 i = EditorPins.Num() - 1; i >= 0; i--)
 			{
 				UEdGraphPin* Pin = EditorPins[i];
-
-				auto IsMatchingInputHandle = [&](const FConstInputHandle& InputHandle) -> bool
-				{
-					return IsMatchingInputHandleAndPin(InputHandle, *Pin);
-				};
-
-				auto IsMatchingOutputHandle = [&](const FConstOutputHandle& OutputHandle) -> bool
-				{
-					return IsMatchingOutputHandleAndPin(OutputHandle, *Pin);
-				};
-
 				switch (Pin->Direction)
 				{
 					case EEdGraphPinDirection::EGPD_Input:
 					{
-						int32 MatchingInputIndex = InputHandles.FindLastByPredicate(IsMatchingInputHandle);
+						auto IsMatchingInput = [&](const FMetasoundFrontendVertex* Input) { return IsMatchingInputAndPin(InBuilder, *Input, *Pin); };
+						const int32 MatchingInputIndex = NodeInputs.FindLastByPredicate(IsMatchingInput);
 						if (INDEX_NONE != MatchingInputIndex)
 						{
-							bIsNodeDirty |= SynchronizePinType(EditorModule, *EditorPins[i], InputHandles[MatchingInputIndex]->GetDataType());
-							InputHandles.RemoveAtSwap(MatchingInputIndex);
+							bIsNodeDirty |= SynchronizePinType(EditorModule, *EditorPins[i], NodeInputs[MatchingInputIndex]->TypeName);
+							NodeInputs.RemoveAtSwap(MatchingInputIndex);
 							EditorPins.RemoveAtSwap(i);
 						}
 					}
@@ -2412,11 +2248,12 @@ PRAGMA_ENABLE_DEPRECATION_WARNINGS
 
 					case EEdGraphPinDirection::EGPD_Output:
 					{
-						int32 MatchingOutputIndex = OutputHandles.FindLastByPredicate(IsMatchingOutputHandle);
+						auto IsMatchingOutput = [&](const FMetasoundFrontendVertex* Output) { return IsMatchingOutputAndPin(InBuilder, *Output, *Pin); };
+						const int32 MatchingOutputIndex = NodeOutputs.FindLastByPredicate(IsMatchingOutput);
 						if (INDEX_NONE != MatchingOutputIndex)
 						{
-							bIsNodeDirty |= SynchronizePinType(EditorModule, *EditorPins[i], OutputHandles[MatchingOutputIndex]->GetDataType());
-							OutputHandles.RemoveAtSwap(MatchingOutputIndex);
+							bIsNodeDirty |= SynchronizePinType(EditorModule, *EditorPins[i], NodeOutputs[MatchingOutputIndex]->TypeName);
+							NodeOutputs.RemoveAtSwap(MatchingOutputIndex);
 							EditorPins.RemoveAtSwap(i);
 						}
 					}
@@ -2430,54 +2267,44 @@ PRAGMA_ENABLE_DEPRECATION_WARNINGS
 				bIsNodeDirty |= !EditorPins.IsEmpty();
 				for (UEdGraphPin* Pin : EditorPins)
 				{
-					if (bLogChanges)
-					{
-						constexpr bool bIncludeNamespace = true;
-						const FText NodeDisplayName = FGraphBuilder::GetDisplayName(*InNode, bIncludeNamespace);
-						UE_LOG(LogMetasoundEditor, Verbose, TEXT("Synchronizing Node '%s' Pins: Removing Excess Editor Pin '%s'"), *NodeDisplayName.ToString(), *Pin->GetName());
-					}
 					InEditorNode.RemovePin(Pin);
 				}
 			}
 
 
-			if (!InputHandles.IsEmpty())
+			if (!NodeInputs.IsEmpty())
 			{
 				bIsNodeDirty = true;
-				for (FConstInputHandle& InputHandle : InputHandles)
+				for (const FMetasoundFrontendVertex* NodeInput : NodeInputs)
 				{
-					if (bLogChanges)
-					{
-						constexpr bool bIncludeNamespace = true;
-						const FText NodeDisplayName = FGraphBuilder::GetDisplayName(*InNode, bIncludeNamespace);
-						const FText InputDisplayName = FGraphBuilder::GetDisplayName(*InputHandle);
-						UE_LOG(LogMetasoundEditor, Verbose, TEXT("Synchronizing Node '%s' Pins: Adding missing Editor Input Pin '%s'"), *NodeDisplayName.ToString(), *InputDisplayName.ToString());
-					}
-					AddPinToNode(InEditorNode, InputHandle);
+					AddPinToNode(EditorModule, InBuilder, InEditorNode, InClass, InNode, *NodeInput, EGPD_Input);
 				}
 			}
 
-			if (!OutputHandles.IsEmpty())
+			if (!NodeOutputs.IsEmpty())
 			{
 				bIsNodeDirty = true;
-				for (FConstOutputHandle& OutputHandle : OutputHandles)
+				for (const FMetasoundFrontendVertex* NodeOutput : NodeOutputs)
 				{
-					if (bLogChanges)
-					{
-						constexpr bool bIncludeNamespace = true;
-						const FText NodeDisplayName = FGraphBuilder::GetDisplayName(*InNode, bIncludeNamespace);
-						const FText OutputDisplayName = FGraphBuilder::GetDisplayName(*OutputHandle);
-						UE_LOG(LogMetasoundEditor, Verbose, TEXT("Synchronizing Node '%s' Pins: Adding missing Editor Output Pin '%s'"), *NodeDisplayName.ToString(), *OutputDisplayName.ToString());
-					}
-					AddPinToNode(InEditorNode, OutputHandle);
+					AddPinToNode(EditorModule, InBuilder, InEditorNode, InClass, InNode, *NodeOutput, EGPD_Output);
 				}
 			}
 
 			// Order pins
-			GetUserModifiableHandles(InputHandles, OutputHandles);
+			NodeInputs = InBuilder.FindUserModifiableNodeInputs(NodeID);
+			NodeOutputs = InBuilder.FindUserModifiableNodeOutputs(NodeID);
+
+			auto GetInputDisplayName = [&InBuilder, &NodeID](const FMetasoundFrontendVertex& Vertex)
+			{
+				return InBuilder.GetNodeInputDisplayName(NodeID, Vertex.Name);
+			};
+			InClass.GetInterfaceForNode(InNode).GetInputStyle().SortVertices(NodeInputs, GetInputDisplayName);
 
-			InNode->GetInputStyle().SortDefaults(InputHandles, [](const FConstInputHandle& Handle) { return FGraphBuilder::GetDisplayName(*Handle); });
-			InNode->GetOutputStyle().SortDefaults(OutputHandles, [](const FConstOutputHandle& Handle) { return FGraphBuilder::GetDisplayName(*Handle); });
+			auto GetOutputDisplayName = [&InBuilder, &NodeID](const FMetasoundFrontendVertex& Vertex)
+			{
+				return InBuilder.GetNodeOutputDisplayName(NodeID, Vertex.Name);
+			};
+			InClass.GetInterfaceForNode(InNode).GetOutputStyle().SortVertices(NodeOutputs, GetOutputDisplayName);
 
 			auto SwapAndDirty = [&](int32 IndexA, int32 IndexB)
 			{
@@ -2494,12 +2321,12 @@ PRAGMA_ENABLE_DEPRECATION_WARNINGS
 				UEdGraphPin* Pin = InEditorNode.Pins[i];
 				if (Pin->Direction == EGPD_Input)
 				{
-					if (!InputHandles.IsEmpty())
+					if (!NodeInputs.IsEmpty())
 					{
-						FConstInputHandle InputHandle = InputHandles.Pop(EAllowShrinking::No);
+						const FMetasoundFrontendVertex* Input = NodeInputs.Pop(EAllowShrinking::No);
 						for (int32 j = i; j >= 0; --j)
 						{
-							if (IsMatchingInputHandleAndPin(InputHandle, *InEditorNode.Pins[j]))
+							if (IsMatchingInputAndPin(InBuilder, *Input, *InEditorNode.Pins[j]))
 							{
 								SwapAndDirty(i, j);
 								break;
@@ -2509,12 +2336,12 @@ PRAGMA_ENABLE_DEPRECATION_WARNINGS
 				}
 				else /* Pin->Direction == EGPD_Output */
 				{
-					if (!OutputHandles.IsEmpty())
+					if (!NodeOutputs.IsEmpty())
 					{
-						FConstOutputHandle OutputHandle = OutputHandles.Pop(EAllowShrinking::No);
+						const FMetasoundFrontendVertex* Output = NodeOutputs.Pop(EAllowShrinking::No);
 						for (int32 j = i; j >= 0; --j)
 						{
-							if (IsMatchingOutputHandleAndPin(OutputHandle, *InEditorNode.Pins[j]))
+							if (IsMatchingOutputAndPin(InBuilder, *Output, *InEditorNode.Pins[j]))
 							{
 								SwapAndDirty(i, j);
 								break;
@@ -2524,6 +2351,10 @@ PRAGMA_ENABLE_DEPRECATION_WARNINGS
 				}
 			}
 
+			if (bIsNodeDirty)
+			{
+				InBuilder.GetConstDocumentChecked().Metadata.ModifyContext.AddNodeIDModified(NodeID);
+			}
 			return bIsNodeDirty;
 		}
 
diff --git a/Engine/Plugins/Runtime/Metasound/Source/MetasoundEditor/Private/MetasoundEditorGraphBuilder.h b/Engine/Plugins/Runtime/Metasound/Source/MetasoundEditor/Private/MetasoundEditorGraphBuilder.h
--- a/Engine/Plugins/Runtime/Metasound/Source/MetasoundEditor/Private/MetasoundEditorGraphBuilder.h
+++ b/Engine/Plugins/Runtime/Metasound/Source/MetasoundEditor/Private/MetasoundEditorGraphBuilder.h
@@ -208,9 +208,6 @@ namespace Metasound
 			// Retrieves the proper pin color for the given PinType
 			static FLinearColor GetPinCategoryColor(const FEdGraphPinType& PinType);
 
-			// Rebuilds all editor node pins based on the provided node handle's class definition.
-			static void RebuildNodePins(UMetasoundEditorGraphNode& InGraphNode);
-
 			// Deletes both the editor graph & frontend nodes from respective graphs
 			static bool DeleteNode(UEdGraphNode& InNode, bool bRemoveUnusedDependencies = true);
 
@@ -218,12 +215,6 @@ namespace Metasound
 			// bound to builder's MetaSound object, returns true (false if it already existed).  Sets (optional) pointer to the bound graph.
 			static bool BindEditorGraph(const FMetaSoundFrontendDocumentBuilder& InBuilder, UMetasoundEditorGraph** OutGraph = nullptr);
 
-			// Adds an Input UEdGraphPin to a UMetasoundEditorGraphNode
-			static UEdGraphPin* AddPinToNode(UMetasoundEditorGraphNode& InEditorNode, Frontend::FConstInputHandle InInputHandle);
-
-			// Adds an Output UEdGraphPin to a UMetasoundEditorGraphNode
-			static UEdGraphPin* AddPinToNode(UMetasoundEditorGraphNode& InEditorNode, Frontend::FConstOutputHandle InOutputHandle);
-
 			// Refreshes pin state from class FrontendClassVertexMetadata
 			static void RefreshPinMetadata(UEdGraphPin& InPin, bool bAdvancedView = false);
 
@@ -234,7 +225,7 @@ namespace Metasound
 			// @param OutGraph - Graph to mutate to conform to provided builder's selected build graph
 			// @param bSkipIfModifyContextUnchanged - If true, recursively checks deprecated document modification for builder's document and all referenced documents, and skips EdGraph synchronization if no modifications are logged.
 			// @return whether or not EditorGraph synchronization was performed.
-			static bool SynchronizeGraph(const FMetaSoundFrontendDocumentBuilder& InBuilder, UMetasoundEditorGraph& OutGraph, bool bSkipIfModifyContextUnchanged = true);
+			static bool SynchronizeGraph(FMetaSoundFrontendDocumentBuilder& InBuilder, UMetasoundEditorGraph& OutGraph, bool bSkipIfModifyContextUnchanged = true);
 
 			// Synchronizes editor nodes with frontend nodes, removing editor nodes that are not represented in the frontend, and adding editor nodes to represent missing frontend nodes.
 			static void SynchronizeNodes(const FMetaSoundFrontendDocumentBuilder& InBuilder, UMetasoundEditorGraph& OutGraph);
@@ -247,7 +238,12 @@ namespace Metasound
 			// Adds and removes pins so that the UMetasoundEditorGraphNode matches the InNode.
 			//
 			// @return True if the UMetasoundEditorGraphNode was altered. False otherwise.
-			static bool SynchronizeNodePins(UMetasoundEditorGraphNode& InEditorNode, Frontend::FConstNodeHandle InNode, bool bRemoveUnusedPins = true, bool bLogChanges = true);
+			static bool SynchronizeNodePins(
+				const FMetaSoundFrontendDocumentBuilder& InBuilder,
+				UMetasoundEditorGraphNode& InEditorNode,
+				const FMetasoundFrontendNode& InNode,
+				const FMetasoundFrontendClass& InClass,
+				bool bRemoveUnusedPins = true);
 
 			static bool SynchronizeComments(const FMetaSoundFrontendDocumentBuilder& InBuilder, UMetasoundEditorGraph& OutGraph);
 
@@ -268,11 +264,11 @@ namespace Metasound
 			// @return True if the UEdGraph was altered. False otherwise.
 			static bool SynchronizeGraphMembers(const FMetaSoundFrontendDocumentBuilder& InBuilder, UMetasoundEditorGraph& OutGraph);
 
-			// Returns true if the FInputHandle and UEdGraphPin match each other.
-			static bool IsMatchingInputHandleAndPin(const Frontend::FConstInputHandle& InInputHandle, const UEdGraphPin& InEditorPin);
+			// Returns true if the given input vertex and UEdGraphPin match each other.
+			static bool IsMatchingInputAndPin(const FMetaSoundFrontendDocumentBuilder& InBuilder, const FMetasoundFrontendVertex& InVertex, const UEdGraphPin& InEditorPin);
 
-			// Returns true if the FOutputHandle and UEdGraphPin match each other.
-			static bool IsMatchingOutputHandleAndPin(const Frontend::FConstOutputHandle& InOutputHandle, const UEdGraphPin& InEditorPin);
+			// Returns true if the given output vertex and UEdGraphPin match each other.
+			static bool IsMatchingOutputAndPin(const FMetaSoundFrontendDocumentBuilder& InBuilder, const FMetasoundFrontendVertex& InVertex, const UEdGraphPin& InEditorPin);
 
 			// Function signature for visiting a node doing depth first traversal.
 			//
@@ -286,6 +282,17 @@ namespace Metasound
 			// This implementation avoids recursive function calls to support deep
 			// graphs.
 			static void DepthFirstTraversal(UEdGraphNode* InInitialNode, FDepthFirstVisitFunction InVisitFunction);
+
+		private:
+			// Adds an input or output UEdGraphPin to a UMetasoundEditorGraphNode
+			static UEdGraphPin* AddPinToNode(
+				const IMetasoundEditorModule& EditorModule,
+				const FMetaSoundFrontendDocumentBuilder& InBuilder,
+				UMetasoundEditorGraphNode& InEditorNode,
+				const FMetasoundFrontendClass& InClass,
+				const FMetasoundFrontendNode& InNode,
+				const FMetasoundFrontendVertex& InVertex,
+				EEdGraphPinDirection Direction);
 		};
 	} // namespace Editor
 } // namespace Metasound
diff --git a/Engine/Plugins/Runtime/Metasound/Source/MetasoundEditor/Private/MetasoundEditorGraphInputNode.cpp b/Engine/Plugins/Runtime/Metasound/Source/MetasoundEditor/Private/MetasoundEditorGraphInputNode.cpp
--- a/Engine/Plugins/Runtime/Metasound/Source/MetasoundEditor/Private/MetasoundEditorGraphInputNode.cpp
+++ b/Engine/Plugins/Runtime/Metasound/Source/MetasoundEditor/Private/MetasoundEditorGraphInputNode.cpp
@@ -180,12 +180,12 @@ void UMetasoundEditorGraphInputNode::ReconstructNode()
 	Super::ReconstructNode();
 }
 
-void UMetasoundEditorGraphInputNode::Validate(Metasound::Editor::FGraphNodeValidationResult& OutResult)
+void UMetasoundEditorGraphInputNode::Validate(const FMetaSoundFrontendDocumentBuilder& InBuilder, Metasound::Editor::FGraphNodeValidationResult& OutResult) const
 {
 	using namespace Metasound::Editor;
 	using namespace Metasound::Frontend;
 
-	Super::Validate(OutResult);
+	Super::Validate(InBuilder, OutResult);
 
 	if (const UMetasoundEditorGraphVertex* Vertex = Cast<UMetasoundEditorGraphVertex>(GetMember()))
 	{
@@ -200,8 +200,7 @@ void UMetasoundEditorGraphInputNode::Validate(Metasound::Editor::FGraphNodeValid
 					const TArray<FMetasoundFrontendVertex>& Outputs = Node->Interface.Outputs;
 					if (ensure(!Outputs.IsEmpty()))
 					{
-						const FMetaSoundFrontendDocumentBuilder& Builder = GetBuilderChecked().GetConstBuilder();
-						if (!Builder.IsNodeOutputConnected(Node->GetID(), Outputs.Last().VertexID))
+						if (!InBuilder.IsNodeOutputConnected(Node->GetID(), Outputs.Last().VertexID))
 						{
 							OutResult.SetMessage(EMessageSeverity::Warning, *RequiredText.ToString());
 						}
diff --git a/Engine/Plugins/Runtime/Metasound/Source/MetasoundEditor/Private/MetasoundEditorGraphNode.cpp b/Engine/Plugins/Runtime/Metasound/Source/MetasoundEditor/Private/MetasoundEditorGraphNode.cpp
--- a/Engine/Plugins/Runtime/Metasound/Source/MetasoundEditor/Private/MetasoundEditorGraphNode.cpp
+++ b/Engine/Plugins/Runtime/Metasound/Source/MetasoundEditor/Private/MetasoundEditorGraphNode.cpp
@@ -27,6 +27,7 @@
 #include "MetasoundFrontend.h"
 #include "MetasoundFrontendDataTypeRegistry.h"
 #include "MetasoundFrontendDocument.h"
+#include "MetasoundFrontendDocumentBuilder.h"
 #include "MetasoundFrontendNodeTemplateRegistry.h"
 #include "MetasoundFrontendRegistries.h"
 #include "MetasoundFrontendSearchEngine.h"
@@ -321,7 +322,20 @@ void UMetasoundEditorGraphNode::AllocateDefaultPins()
 	using namespace Metasound;
 
 	ensureAlways(Pins.IsEmpty());
-	Editor::FGraphBuilder::RebuildNodePins(*this);
+
+	if (UObject* Outermost = GetOutermostObject())
+	{
+		const FGuid NodeID = GetNodeID();
+		const FMetaSoundFrontendDocumentBuilder& Builder = Engine::FDocumentBuilderRegistry::GetChecked().FindOrBeginBuilding(Outermost);
+		if (const FMetasoundFrontendNode* Node = Builder.FindNode(NodeID))
+		{
+			if (const FMetasoundFrontendClass* Class = Builder.FindDependency(Node->ClassID))
+			{
+				constexpr bool bRemoveUnusedPins = false;
+				Editor::FGraphBuilder::SynchronizeNodePins(Builder, *this, *Node, *Class, bRemoveUnusedPins);
+			}
+		}
+	}
 }
 
 void UMetasoundEditorGraphNode::SyncChangeIDs()
@@ -343,8 +357,10 @@ void UMetasoundEditorGraphNode::CacheTitle()
 	CachedTitle = NodeHandle->GetDisplayTitle();
 }
 
-void UMetasoundEditorGraphNode::Validate(Metasound::Editor::FGraphNodeValidationResult& OutResult)
+void UMetasoundEditorGraphNode::Validate(const FMetaSoundFrontendDocumentBuilder& InBuilder, Metasound::Editor::FGraphNodeValidationResult& OutResult) const
 {
+	TRACE_CPUPROFILER_EVENT_SCOPE(UMetasoundEditorGraphNode::Validate);
+
 	using namespace Metasound::Editor;
 	using namespace Metasound::Frontend;
 
@@ -366,9 +382,11 @@ void UMetasoundEditorGraphNode::Validate(Metasound::Editor::FGraphNodeValidation
 						{
 							if (ExternalNode->GetBreadcrumb().ClassName == FRerouteNodeTemplate::ClassName)
 							{
-								FConstInputHandle InputHandle = FGraphBuilder::GetConstInputHandleFromPin(Pin);
-								const FString Msg = FString::Format(*GraphNodePrivate::MissingConcreteOutputConnectionFormat, { InputHandle->GetDisplayName().ToString() });
-								OutResult.SetMessage(EMessageSeverity::Warning, Msg);
+								if (const FMetasoundFrontendVertex* Vertex = FGraphBuilder::GetPinVertex(InBuilder, Pin))
+								{
+									const FString Msg = FString::Format(*GraphNodePrivate::MissingConcreteOutputConnectionFormat, { Vertex->Name.ToString() });
+									OutResult.SetMessage(EMessageSeverity::Warning, Msg);
+								}
 							}
 						}
 					}
@@ -391,17 +409,23 @@ bool UMetasoundEditorGraphNode::ContainsClassChange() const
 
 void UMetasoundEditorGraphNode::ReconstructNode()
 {
-	using namespace Metasound::Editor;
-	using namespace Metasound::Frontend;
+	using namespace Metasound;
 
-	// Don't remove unused pins here. Reconstruction can occur while duplicating or pasting nodes,
-	// and subsequent steps clean-up unused pins.  This can be called mid-copy, which means the node
-	// handle may be invalid.  Setting to remove unused causes premature removal and then default values
-	// are lost.
-	FConstNodeHandle NodeHandle = GetConstNodeHandle();
-	if (NodeHandle->IsValid())
+	if (UObject* Outermost = GetOutermostObject())
 	{
-		FGraphBuilder::SynchronizeNodePins(*this, NodeHandle, false /* bRemoveUnusedPins */, false /* bLogChanges */);
+		const FGuid NodeID = GetNodeID();
+		const FMetaSoundFrontendDocumentBuilder& Builder = Engine::FDocumentBuilderRegistry::GetChecked().FindOrBeginBuilding(Outermost);
+		if (const FMetasoundFrontendNode* Node = Builder.FindNode(NodeID))
+		{
+			if (const FMetasoundFrontendClass* Class = Builder.FindDependency(Node->ClassID))
+			{
+				// Don't remove unused pins here. Reconstruction can occur while duplicating or pasting nodes,
+				// and subsequent steps clean-up unused pins.  This can be called mid-copy, which means the node
+				// handle may be invalid.  Setting to remove unused causes premature removal and then default values
+				// are lost.
+				Editor::FGraphBuilder::SynchronizeNodePins(Builder, *this, *Node, *Class, false /* bRemoveUnusedPins */);
+			}
+		}
 	}
 
 	CacheTitle();
@@ -953,12 +977,12 @@ bool UMetasoundEditorGraphOutputNode::EnableInteractWidgets() const
 	return bEnabled;
 }
 
-void UMetasoundEditorGraphOutputNode::Validate(Metasound::Editor::FGraphNodeValidationResult& OutResult)
+void UMetasoundEditorGraphOutputNode::Validate(const FMetaSoundFrontendDocumentBuilder& InBuilder, Metasound::Editor::FGraphNodeValidationResult& OutResult) const
 {
 	using namespace Metasound::Editor;
 	using namespace Metasound::Frontend;
 
-	Super::Validate(OutResult);
+	Super::Validate(InBuilder, OutResult);
 
 	if (const UMetasoundEditorGraphVertex* Vertex = Cast<UMetasoundEditorGraphVertex>(GetMember()))
 	{
@@ -968,7 +992,7 @@ void UMetasoundEditorGraphOutputNode::Validate(Metasound::Editor::FGraphNodeVali
 			FText RequiredText;
 			if (InterfaceToValidate.IsMemberOutputRequired(Vertex->GetMemberName(), RequiredText))
 			{
-				if (const FMetasoundFrontendNode* Node = GetFrontendNode())
+				if (const FMetasoundFrontendNode* Node = InBuilder.FindNode(GetNodeID()))
 				{
 					const TArray<FMetasoundFrontendVertex>& Inputs = Node->Interface.Inputs;
 					if (ensure(!Inputs.IsEmpty()))
@@ -1023,19 +1047,6 @@ void UMetasoundEditorGraphOutputNode::CacheBreadcrumb()
 	}
 }
 
-FMetasoundFrontendClassName UMetasoundEditorGraphOutputNode::GetClassName() const
-{
-	using namespace Metasound::Frontend;
-
-	if (Output)
-	{
-		FConstNodeHandle NodeHandle = Output->GetConstNodeHandle();
-		return NodeHandle->GetClassMetadata().GetClassName();
-	}
-
-	return Breadcrumb.ClassName;
-}
-
 FGuid UMetasoundEditorGraphOutputNode::GetNodeID() const
 {
 	if (Output)
@@ -1192,16 +1203,33 @@ void UMetasoundEditorGraphExternalNode::GetPinHoverText(const UEdGraphPin& Pin,
 	Super::GetPinHoverText(Pin, OutHoverText);
 }
 
-void UMetasoundEditorGraphExternalNode::Validate(Metasound::Editor::FGraphNodeValidationResult& OutResult)
+void UMetasoundEditorGraphExternalNode::Validate(const FMetaSoundFrontendDocumentBuilder& InBuilder, Metasound::Editor::FGraphNodeValidationResult& OutResult) const
 {
 	using namespace Metasound::Editor;
+	using namespace Metasound::Engine;
 	using namespace Metasound::Frontend;
 
 #if WITH_EDITOR
-	Super::Validate(OutResult);
+	Super::Validate(InBuilder, OutResult);
 
-	FConstNodeHandle NodeHandle = GetConstNodeHandle();
-	const FMetasoundFrontendClassMetadata& Metadata = NodeHandle->GetClassMetadata();
+	UObject* Outermost = GetOutermostObject();
+	check(Outermost);
+
+	const FMetasoundFrontendClass* Class = nullptr;
+	const FMetasoundFrontendNode* Node = InBuilder.FindNode(NodeID);
+
+	if (Node)
+	{
+		Class = InBuilder.FindDependency(Node->ClassID);
+	}
+
+	if (!Class)
+	{
+		OutResult.SetMessage(EMessageSeverity::Error, FString::Format(TEXT("Node with ID '{0}' class not found"), { *Node->ClassID.ToString() }));
+		return;
+	}
+
+	const FMetasoundFrontendClassMetadata& Metadata = Class->Metadata;
 
 	// 1. Validate external referenced graph or template node
 	switch(Metadata.GetType())
@@ -1246,7 +1274,7 @@ void UMetasoundEditorGraphExternalNode::Validate(Metasound::Editor::FGraphNodeVa
 			const FNodeRegistryKey Key = FNodeRegistryKey(Metadata);
 			if (const INodeTemplate* Template = INodeTemplateRegistry::Get().FindTemplate(Key))
 			{
-				const bool bIsValidInterface = Template->IsValidNodeInterface(NodeHandle->GetNodeInterface());
+				const bool bIsValidInterface = Template->IsValidNodeInterface(Node->Interface);
 				if (!bIsValidInterface)
 				{
 					OutResult.SetMessage(EMessageSeverity::Error, FString::Format(TEXT("Cannot implement template interface for node class '{0}"), { *Metadata.GetClassName().ToString() }));
@@ -1277,13 +1305,12 @@ void UMetasoundEditorGraphExternalNode::Validate(Metasound::Editor::FGraphNodeVa
 		break;
 	}
 
-PRAGMA_DISABLE_DEPRECATION_WARNINGS
 	// 2. Check if node is invalid, version is missing and cache if interface changes exist between the document's records and the registry
 	FClassInterfaceUpdates InterfaceUpdates;
-	if (!NodeHandle->DiffAgainstRegistryInterface(InterfaceUpdates, false /* bUseHighestMinorVersion */))
+	const bool bUseHighestMinorVersion = false;
+	if (!InBuilder.DiffAgainstRegistryInterface(NodeID, InBuilder.GetBuildPageID(), bUseHighestMinorVersion, InterfaceUpdates))
 	{
-PRAGMA_ENABLE_DEPRECATION_WARNINGS
-		if (NodeHandle->IsValid())
+		if (Node)
 		{
 			const FText* PromptIfMissing = nullptr;
 			FString FormattedClassName;
@@ -1324,26 +1351,14 @@ PRAGMA_ENABLE_DEPRECATION_WARNINGS
 		}
 	}
 
-	// 4. Report if node was nativized
-	const bool bNewIsClassNative = !IMetaSoundAssetManager::GetChecked().IsAssetClass(Metadata);
-	if (bIsClassNative != bNewIsClassNative)
-	{
-		if (bNewIsClassNative)
-		{
-			NodeUpgradeMessage = FText::Format(LOCTEXT("MetaSoundNode_NativizedMessage", "Class '{0}' has been nativized."), Metadata.GetDisplayName());
-		}
-
-		bIsClassNative = bNewIsClassNative;
-	}
-
 	// 5. Report if node was auto-updated
-	FMetasoundFrontendNodeStyle Style = NodeHandle->GetNodeStyle();
+	const FMetasoundFrontendNodeStyle& Style = Node->Style;
 	if (Style.bMessageNodeUpdated)
 	{
-		NodeUpgradeMessage = FText::Format(LOCTEXT("MetaSoundNode_UpgradedMessage", "Node class '{0}' updated to version {1}"),
+		OutResult.SetUpgradeMessage(FText::Format(LOCTEXT("MetaSoundNode_UpgradedMessage", "Node class '{0}' updated to version {1}"),
 			Metadata.GetDisplayName(),
 			FText::FromString(Metadata.GetVersion().ToString())
-		);
+		));
 	}
 
 	// 6. Reset pin state (if pin was orphaned or clear if no longer orphaned)
@@ -1352,20 +1367,25 @@ PRAGMA_ENABLE_DEPRECATION_WARNINGS
 		bool bWasRemoved = false;
 		if (Pin->Direction == EGPD_Input)
 		{
-			FConstInputHandle Input = FGraphBuilder::GetConstInputHandleFromPin(Pin);
-			bWasRemoved |= InterfaceUpdates.RemovedInputs.ContainsByPredicate([&Input](const FMetasoundFrontendClassInput* ClassInput)
+			if (const FMetasoundFrontendVertex* Vertex = FGraphBuilder::GetPinVertex(InBuilder, Pin))
 			{
-				return Input->GetName() == ClassInput->Name && Input->GetDataType() == ClassInput->TypeName;
-			});
+				bWasRemoved |= InterfaceUpdates.RemovedInputs.ContainsByPredicate([&Vertex](const FMetasoundFrontendClassInput* ClassInput)
+				{
+					return Vertex->Name == ClassInput->Name && Vertex->TypeName == ClassInput->TypeName;
+				});
+			}
 		}
 
 		if (Pin->Direction == EGPD_Output)
 		{
-			FConstOutputHandle Output = FGraphBuilder::GetConstOutputHandleFromPin(Pin);
-			bWasRemoved |= InterfaceUpdates.RemovedOutputs.ContainsByPredicate([&Output](const FMetasoundFrontendClassOutput* ClassOutput)
+			if (const FMetasoundFrontendVertex* Vertex = FGraphBuilder::GetPinVertex(InBuilder, Pin))
 			{
-				return Output->GetName() == ClassOutput->Name && Output->GetDataType() == ClassOutput->TypeName;
-			});
+				bWasRemoved |= InterfaceUpdates.RemovedOutputs.ContainsByPredicate([&Vertex](const FMetasoundFrontendClassOutput* ClassOutput)
+				{
+					return Vertex->Name == ClassOutput->Name && Vertex->TypeName == ClassOutput->TypeName;
+				});
+			}
+
 		}
 
 		OutResult.SetPinOrphaned(*Pin, bWasRemoved);
@@ -1620,11 +1640,6 @@ bool UMetasoundEditorGraphVariableNode::EnableInteractWidgets() const
 	return bEnabled;
 }
 
-FMetasoundFrontendClassName UMetasoundEditorGraphVariableNode::GetClassName() const
-{
-	return ClassName;
-}
-
 EMetasoundFrontendClassType UMetasoundEditorGraphVariableNode::GetClassType() const
 {
 	return ClassType;
diff --git a/Engine/Plugins/Runtime/Metasound/Source/MetasoundEditor/Private/MetasoundEditorGraphValidation.cpp b/Engine/Plugins/Runtime/Metasound/Source/MetasoundEditor/Private/MetasoundEditorGraphValidation.cpp
--- a/Engine/Plugins/Runtime/Metasound/Source/MetasoundEditor/Private/MetasoundEditorGraphValidation.cpp
+++ b/Engine/Plugins/Runtime/Metasound/Source/MetasoundEditor/Private/MetasoundEditorGraphValidation.cpp
@@ -90,6 +90,14 @@ namespace Metasound
 			}
 		}
 
+		void FGraphNodeValidationResult::SetUpgradeMessage(const FText& InMessage)
+		{
+			if (Node)
+			{
+				Node->NodeUpgradeMessage = InMessage;
+			}
+		}
+
 		EMessageSeverity::Type FGraphValidationResults::GetHighestMessageSeverity() const
 		{
 			int32 HighestMessageSeverity = static_cast<int32>(EMessageSeverity::Info);
diff --git a/Engine/Plugins/Runtime/Metasound/Source/MetasoundEditor/Public/MetasoundEditorGraphInputNode.h b/Engine/Plugins/Runtime/Metasound/Source/MetasoundEditor/Public/MetasoundEditorGraphInputNode.h
--- a/Engine/Plugins/Runtime/Metasound/Source/MetasoundEditor/Public/MetasoundEditorGraphInputNode.h
+++ b/Engine/Plugins/Runtime/Metasound/Source/MetasoundEditor/Public/MetasoundEditorGraphInputNode.h
@@ -13,6 +13,9 @@
 
 #include "MetasoundEditorGraphInputNode.generated.h"
 
+// Forward Declarations
+struct FMetaSoundFrontendDocumentBuilder;
+
 
 UCLASS()
 class METASOUNDEDITOR_API UMetasoundEditorGraphInputNode : public UMetasoundEditorGraphMemberNode
@@ -45,7 +48,7 @@ public:
 	virtual FLinearColor GetNodeTitleColor() const override;
 	virtual FSlateIcon GetNodeTitleIcon() const override;
 	virtual void GetPinHoverText(const UEdGraphPin& Pin, FString& OutHoverText) const override;
-	virtual void Validate(Metasound::Editor::FGraphNodeValidationResult& OutResult) override;
+	virtual void Validate(const FMetaSoundFrontendDocumentBuilder& InBuilder, Metasound::Editor::FGraphNodeValidationResult& OutResult) const override;
 	virtual bool EnableInteractWidgets() const override;
 	virtual FText GetTooltipText() const override;
 
diff --git a/Engine/Plugins/Runtime/Metasound/Source/MetasoundEditor/Public/MetasoundEditorGraphNode.h b/Engine/Plugins/Runtime/Metasound/Source/MetasoundEditor/Public/MetasoundEditorGraphNode.h
--- a/Engine/Plugins/Runtime/Metasound/Source/MetasoundEditor/Public/MetasoundEditorGraphNode.h
+++ b/Engine/Plugins/Runtime/Metasound/Source/MetasoundEditor/Public/MetasoundEditorGraphNode.h
@@ -26,6 +26,8 @@ class UMetasoundEditorGraphMember;
 class UMetasoundEditorGraphVariable;
 class UMetasoundEditorGraphMemberDefaultFloat;
 
+struct FMetaSoundFrontendDocumentBuilder;
+
 namespace EMessageSeverity { enum Type : int32; }
 
 namespace Metasound::Editor
@@ -181,13 +183,14 @@ public:
 
 	TSet<FString> GetDisallowedPinClassNames(const UEdGraphPin& InPin) const;
 
-	UE_DEPRECATED(5.4, "Use the Frontend node or Breadcrumb directly to get the class name.")
-	virtual FMetasoundFrontendClassName GetClassName() const { return FMetasoundFrontendClassName(); }
-
 	virtual FGuid GetNodeID() const { return FGuid(); }
 	virtual FText GetDisplayName() const;
 	virtual void CacheTitle();
-	virtual void Validate(Metasound::Editor::FGraphNodeValidationResult& OutResult);
+
+	UE_DEPRECATED(5.7, "Use const version that is passed builder")
+	virtual void Validate(Metasound::Editor::FGraphNodeValidationResult& OutResult) { }
+
+	virtual void Validate(const FMetaSoundFrontendDocumentBuilder& Builder, Metasound::Editor::FGraphNodeValidationResult& OutResult) const;
 
 	// Mark node for refresh
 	void SyncChangeIDs();
@@ -279,9 +282,6 @@ public:
 
 	virtual void CacheBreadcrumb() override;	
 
-	UE_DEPRECATED(5.4, "Use the Frontend node or Breadcrumb directly to get the class name.")
-	virtual FMetasoundFrontendClassName GetClassName() const override;
-
 	virtual FGuid GetNodeID() const override;
 	virtual UMetasoundEditorGraphMember* GetMember() const override;
 
@@ -299,7 +299,7 @@ public:
 	// Disables interact widgets (ex. sliders, knobs) when input is connected
 	virtual bool EnableInteractWidgets() const override;
 
-	virtual void Validate(Metasound::Editor::FGraphNodeValidationResult& OutResult) override;
+	virtual void Validate(const FMetaSoundFrontendDocumentBuilder& InBuilder, Metasound::Editor::FGraphNodeValidationResult& OutResult) const override;
 
 protected:
 	// Breadcrumb used if associated FrontendNode cannot be found or has been unlinked
@@ -334,9 +334,6 @@ protected:
 	bool bIsClassNative = true;
 
 public:
-	UE_DEPRECATED(5.4, "Use the Frontend node or Breadcrumb directly to get the class name.")
-	virtual FMetasoundFrontendClassName GetClassName() const override { return Breadcrumb.ClassName; }
-
 	virtual const FMetasoundEditorGraphNodeBreadcrumb& GetBreadcrumb() const;
 	virtual FGuid GetNodeID() const override { return NodeID; }
 	virtual FLinearColor GetNodeTitleColor() const override;
@@ -352,7 +349,7 @@ public:
 	bool CanAutoUpdate() const;
 
 	// Validates node and returns whether or not the node is valid.
-	virtual void Validate(Metasound::Editor::FGraphNodeValidationResult& OutResult) override;
+	virtual void Validate(const FMetaSoundFrontendDocumentBuilder& InBuilder, Metasound::Editor::FGraphNodeValidationResult& OutResult) const override;
 
 	/**Set Unconnected Pins hidden/Unhidden*/
 	void HideUnconnectedPins(const bool InHidePins);
@@ -396,7 +393,6 @@ public:
 	virtual UMetasoundEditorGraphMember* GetMember() const override;
 	virtual bool EnableInteractWidgets() const override;
 
-	virtual FMetasoundFrontendClassName GetClassName() const override;
 	virtual FGuid GetNodeID() const override;
 	virtual FName GetCornerIcon() const override;
 	virtual void GetPinHoverText(const UEdGraphPin& Pin, FString& OutHoverText) const override;
diff --git a/Engine/Plugins/Runtime/Metasound/Source/MetasoundEditor/Public/MetasoundEditorGraphValidation.h b/Engine/Plugins/Runtime/Metasound/Source/MetasoundEditor/Public/MetasoundEditorGraphValidation.h
--- a/Engine/Plugins/Runtime/Metasound/Source/MetasoundEditor/Public/MetasoundEditorGraphValidation.h
+++ b/Engine/Plugins/Runtime/Metasound/Source/MetasoundEditor/Public/MetasoundEditorGraphValidation.h
@@ -31,6 +31,7 @@ namespace Metasound
 
 			void SetMessage(EMessageSeverity::Type InSeverity, const FString& InMessage);
 			void SetPinOrphaned(UEdGraphPin& Pin, bool bIsOrphaned);
+			void SetUpgradeMessage(const FText& InMessage);
 
 		private:
 			// Node associated with validation result
diff --git a/Engine/Plugins/Runtime/Metasound/Source/MetasoundFrontend/Private/MetasoundAssetBase.cpp b/Engine/Plugins/Runtime/Metasound/Source/MetasoundFrontend/Private/MetasoundAssetBase.cpp
--- a/Engine/Plugins/Runtime/Metasound/Source/MetasoundFrontend/Private/MetasoundAssetBase.cpp
+++ b/Engine/Plugins/Runtime/Metasound/Source/MetasoundFrontend/Private/MetasoundAssetBase.cpp
@@ -362,7 +362,10 @@ void FMetasoundAssetBase::UpdateAndRegisterForExecution(Metasound::Frontend::FMe
 #if WITH_EDITOR
 	// Must be completed after auto-update to ensure all non-transient referenced dependency data is up-to-date (ex.
 	// class version), which is required for most accurately caching current registry metadata.
-	CacheRegistryMetadata();
+	if (DocBuilder)
+	{
+		DocBuilder->CacheRegistryMetadata();
+	}
 #endif // WITH_EDITOR
 
 	GraphRegistryKey = FNodeClassRegistry::Get().RegisterGraph(Owner);
@@ -423,7 +426,7 @@ void FMetasoundAssetBase::UpdateAndRegisterForSerialization(FName CookPlatformNa
 #if WITH_EDITOR
 	// Must be completed after auto-update to ensure all non-transient referenced dependency data is up-to-date (ex.
 	// class version), which is required for most accurately caching current registry metadata.
-	CacheRegistryMetadata();
+	DocBuilder.CacheRegistryMetadata();
 #endif // WITH_EDITOR
 
 	{
@@ -589,119 +592,10 @@ void FMetasoundAssetBase::CacheRegistryMetadata()
 
 	using namespace Metasound::Frontend;
 
-
-PRAGMA_DISABLE_DEPRECATION_WARNINGS
-	FMetasoundFrontendDocument* Document = GetDocumentAccessPtr().Get();
-PRAGMA_ENABLE_DEPRECATION_WARNINGS
-
-	if (!ensure(Document))
-	{
-		return;
-	}
-
-	using FNameDataTypePair = TPair<FName, FName>;
-	const TSet<FMetasoundFrontendVersion>& InterfaceVersions = Document->Interfaces;
-	FMetasoundFrontendClassInterface& RootGraphClassInterface = Document->RootGraph.GetDefaultInterface();
-
-	// 1. Gather inputs/outputs managed by interfaces
-	TMap<FNameDataTypePair, FMetasoundFrontendClassInput*> Inputs;
-	for (FMetasoundFrontendClassInput& Input : RootGraphClassInterface.Inputs)
-	{
-		FNameDataTypePair NameDataTypePair = FNameDataTypePair(Input.Name, Input.TypeName);
-		Inputs.Add(MoveTemp(NameDataTypePair), &Input);
-	}
-
-	TMap<FNameDataTypePair, FMetasoundFrontendClassOutput*> Outputs;
-	for (FMetasoundFrontendClassOutput& Output : RootGraphClassInterface.Outputs)
-	{
-		FNameDataTypePair NameDataTypePair = FNameDataTypePair(Output.Name, Output.TypeName);
-		Outputs.Add(MoveTemp(NameDataTypePair), &Output);
-	}
-
-	// 2. Copy metadata for inputs/outputs managed by interfaces, removing them from maps generated
-	auto CacheInterfaceMetadata = [](const FMetasoundFrontendVertexMetadata& InRegistryMetadata, FMetasoundFrontendVertexMetadata& OutMetadata)
-	{
-		const int32 CachedSortOrderIndex = OutMetadata.SortOrderIndex;
-		OutMetadata = InRegistryMetadata;
-		OutMetadata.SortOrderIndex = CachedSortOrderIndex;
-	};
-
-	for (const FMetasoundFrontendVersion& Version : InterfaceVersions)
-	{
-		const FInterfaceRegistryKey InterfaceKey = GetInterfaceRegistryKey(Version);
-		const IInterfaceRegistryEntry* Entry = IInterfaceRegistry::Get().FindInterfaceRegistryEntry(InterfaceKey);
-
-		UE_CLOG(nullptr == Entry, LogMetaSound, Error, TEXT("Failed to find interface (%s) when caching registry data for %s. MetaSound inputs and outputs for asset may not function correctly."), *Version.ToString(), *GetOwningAssetName());
-
-		if (Entry)
-		{
-			for (const FMetasoundFrontendClassInput& InterfaceInput : Entry->GetInterface().Inputs)
-			{
-				const FNameDataTypePair NameDataTypePair = FNameDataTypePair(InterfaceInput.Name, InterfaceInput.TypeName);
-				if (FMetasoundFrontendClassInput* Input = Inputs.FindRef(NameDataTypePair))
-				{
-					CacheInterfaceMetadata(InterfaceInput.Metadata, Input->Metadata);
-					Inputs.Remove(NameDataTypePair);
-				}
-			}
-
-			for (const FMetasoundFrontendClassOutput& InterfaceOutput : Entry->GetInterface().Outputs)
-			{
-				const FNameDataTypePair NameDataTypePair = FNameDataTypePair(InterfaceOutput.Name, InterfaceOutput.TypeName);
-				if (FMetasoundFrontendClassOutput* Output = Outputs.FindRef(NameDataTypePair))
-				{
-					CacheInterfaceMetadata(InterfaceOutput.Metadata, Output->Metadata);
-					Outputs.Remove(NameDataTypePair);
-				}
-			}
-		}
-	}
-
-	// 3. Iterate remaining inputs/outputs not managed by interfaces and set to serialize text
-	// (in case they were orphaned by an interface no longer being implemented).
-	for (TPair<FNameDataTypePair, FMetasoundFrontendClassInput*>& Pair : Inputs)
-	{
-		Pair.Value->Metadata.SetSerializeText(true);
-	}
-
-	for (TPair<FNameDataTypePair, FMetasoundFrontendClassOutput*>& Pair : Outputs)
-	{
-		Pair.Value->Metadata.SetSerializeText(true);
-	}
-
-	// 4. Refresh style as order of members could've changed
-	{
-		FMetasoundFrontendInterfaceStyle InputStyle;
-		Algo::ForEach(RootGraphClassInterface.Inputs, [&InputStyle](const FMetasoundFrontendClassInput& Input)
-		{
-			InputStyle.DefaultSortOrder.Add(Input.Metadata.SortOrderIndex);
-		});
-		RootGraphClassInterface.SetInputStyle(InputStyle);
-	}
-
-	{
-		FMetasoundFrontendInterfaceStyle OutputStyle;
-		Algo::ForEach(RootGraphClassInterface.Outputs, [&OutputStyle](const FMetasoundFrontendClassOutput& Output)
-		{
-			OutputStyle.DefaultSortOrder.Add(Output.Metadata.SortOrderIndex);
-		});
-		RootGraphClassInterface.SetOutputStyle(OutputStyle);
-	}
-
-	// 5. Cache registry data on document dependencies
-	for (FMetasoundFrontendClass& Dependency : Document->Dependencies)
-	{
-		if (!FMetasoundFrontendClass::CacheGraphDependencyMetadataFromRegistry(Dependency))
-		{
-			UE_LOG(LogMetaSound, Warning,
-				TEXT("'%s' failed to cache dependency registry data: Registry missing class with key '%s'"),
-				*GetOwningAssetName(),
-				*Dependency.Metadata.GetClassName().ToString());
-			UE_LOG(LogMetaSound, Warning,
-				TEXT("Asset '%s' may fail to build runtime graph unless re-registered after dependency with given key is loaded."),
-				*GetOwningAssetName());
-		}
-	}
+	UObject* Owner = GetOwningAsset();
+	check(Owner);
+	FMetaSoundFrontendDocumentBuilder& DocBuilder = IDocumentBuilderRegistry::GetChecked().FindOrBeginBuilding(Owner);
+	DocBuilder.CacheRegistryMetadata();
 }
 
 FMetasoundFrontendDocumentModifyContext& FMetasoundAssetBase::GetModifyContext()
diff --git a/Engine/Plugins/Runtime/Metasound/Source/MetasoundFrontend/Private/MetasoundFrontendDocument.cpp b/Engine/Plugins/Runtime/Metasound/Source/MetasoundFrontend/Private/MetasoundFrontendDocument.cpp
--- a/Engine/Plugins/Runtime/Metasound/Source/MetasoundFrontend/Private/MetasoundFrontendDocument.cpp
+++ b/Engine/Plugins/Runtime/Metasound/Source/MetasoundFrontend/Private/MetasoundFrontendDocument.cpp
@@ -1117,7 +1117,7 @@ bool FMetasoundFrontendClass::CacheGraphDependencyMetadataFromRegistry(FMetasoun
 			{
 				AddRegistryVertexMetadata(InputMembers, Input, InputStyle);
 			});
-			InOutDependency.GetDefaultInterface().SetInputStyle(InputStyle);
+			InOutDependency.GetDefaultInterface().SetInputStyle(MoveTemp(InputStyle));
 
 			FMetasoundFrontendInterfaceStyle OutputStyle;
 			FVertexMetadataMap OutputMembers;
@@ -1126,7 +1126,7 @@ bool FMetasoundFrontendClass::CacheGraphDependencyMetadataFromRegistry(FMetasoun
 			{
 				AddRegistryVertexMetadata(OutputMembers, Output, OutputStyle);
 			});
-			InOutDependency.GetDefaultInterface().SetOutputStyle(OutputStyle);
+			InOutDependency.GetDefaultInterface().SetOutputStyle(MoveTemp(OutputStyle));
 
 			return true;
 		}
diff --git a/Engine/Plugins/Runtime/Metasound/Source/MetasoundFrontend/Private/MetasoundFrontendDocumentBuilder.cpp b/Engine/Plugins/Runtime/Metasound/Source/MetasoundFrontend/Private/MetasoundFrontendDocumentBuilder.cpp
--- a/Engine/Plugins/Runtime/Metasound/Source/MetasoundFrontend/Private/MetasoundFrontendDocumentBuilder.cpp
+++ b/Engine/Plugins/Runtime/Metasound/Source/MetasoundFrontend/Private/MetasoundFrontendDocumentBuilder.cpp
@@ -1893,7 +1893,7 @@ const FMetasoundFrontendGraph& FMetaSoundFrontendDocumentBuilder::AddGraphPage(c
 #endif // WITH_EDITORONLY_DATA
 
 #if WITH_EDITOR
-void FMetaSoundFrontendDocumentBuilder::CacheRegistryMetadata()
+void FMetaSoundFrontendDocumentBuilder::CacheRegistryMetadata() const
 {
 	using namespace Metasound::Frontend;
 	using FNameDataTypePair = TPair<FName, FName>;
diff --git a/Engine/Plugins/Runtime/Metasound/Source/MetasoundFrontend/Private/MetasoundFrontendTransform.cpp b/Engine/Plugins/Runtime/Metasound/Source/MetasoundFrontend/Private/MetasoundFrontendTransform.cpp
--- a/Engine/Plugins/Runtime/Metasound/Source/MetasoundFrontend/Private/MetasoundFrontendTransform.cpp
+++ b/Engine/Plugins/Runtime/Metasound/Source/MetasoundFrontend/Private/MetasoundFrontendTransform.cpp
@@ -45,7 +45,7 @@ namespace Metasound
 			}
 
 #if WITH_EDITOR
-			FGetNodeDisplayNameProjection NodeDisplayNameProjection;
+			FGetNodeDisplayNameProjection NodeDisplayNameProjection = [] (const FNodeHandle&) { return FText(); };
 
 			bool GetVersioningLoggingEnabled()
 			{
@@ -823,7 +823,9 @@ namespace Metasound
 			FMetasoundFrontendInterfaceStyle Style = ReferencedGraphHandle->GetInputStyle();
 			InPresetGraphHandle->SetInputStyle(MoveTemp(Style));
 
+PRAGMA_DISABLE_DEPRECATION_WARNINGS
 			Style.SortDefaults(NodeHandles, DocumentTransform::GetNodeDisplayNameProjection());
+PRAGMA_ENABLE_DEPRECATION_WARNINGS
 
 			InputNodeLocation = FVector2D::ZeroVector;
 			for (const FNodeHandle& NodeHandle : NodeHandles)
@@ -868,7 +870,9 @@ namespace Metasound
 			FMetasoundFrontendInterfaceStyle Style = ReferencedGraphHandle->GetOutputStyle();
 			InPresetGraphHandle->SetOutputStyle(MoveTemp(Style));
 
+PRAGMA_DISABLE_DEPRECATION_WARNINGS
 			Style.SortDefaults(NodeHandles, DocumentTransform::GetNodeDisplayNameProjection());
+PRAGMA_ENABLE_DEPRECATION_WARNINGS
 
 			// Set output node location
 			for (const FNodeHandle& OutputNode : NodeHandles)
diff --git a/Engine/Plugins/Runtime/Metasound/Source/MetasoundFrontend/Public/MetasoundAssetBase.h b/Engine/Plugins/Runtime/Metasound/Source/MetasoundFrontend/Public/MetasoundAssetBase.h
--- a/Engine/Plugins/Runtime/Metasound/Source/MetasoundFrontend/Public/MetasoundAssetBase.h
+++ b/Engine/Plugins/Runtime/Metasound/Source/MetasoundFrontend/Public/MetasoundAssetBase.h
@@ -168,12 +168,7 @@ public:
 #endif // WITH_EDITORONLY_DATA
 
 #if WITH_EDITOR
-	/*
-	 * Caches transient metadata (class & vertex) found in the registry
-	 * that is not necessary for serialization or core graph generation.
-	 *
-	 * @return - Whether class was found in the registry & data was cached successfully.
-	 */
+	UE_DEPRECATED(5.7, "Moved to FMetaSoundFrontendDocumentBuilder")
 	UE_API void CacheRegistryMetadata();
 
 	UE_API FMetasoundFrontendDocumentModifyContext& GetModifyContext();
diff --git a/Engine/Plugins/Runtime/Metasound/Source/MetasoundFrontend/Public/MetasoundFrontendDocument.h b/Engine/Plugins/Runtime/Metasound/Source/MetasoundFrontend/Public/MetasoundFrontendDocument.h
--- a/Engine/Plugins/Runtime/Metasound/Source/MetasoundFrontend/Public/MetasoundFrontendDocument.h
+++ b/Engine/Plugins/Runtime/Metasound/Source/MetasoundFrontend/Public/MetasoundFrontendDocument.h
@@ -1081,6 +1081,7 @@ struct FMetasoundFrontendInterfaceStyle
 	UE_API void SortVertices(TArray<const FMetasoundFrontendVertex*>& OutVertices, TFunctionRef<FText(const FMetasoundFrontendVertex&)> InGetDisplayNamePredicate) const;
 
 	template <typename HandleType, typename NamePredicateType>
+	UE_DEPRECATED(5.7, "Use SortVertices and the builder API instead of this template (which is built primarily for use with the controller API, currently being deprecated)")
 	void SortDefaults(TArray<HandleType>& OutHandles, NamePredicateType InGetDisplayNamePredicate) const
 	{
 		TMap<FGuid, int32> NodeIDToSortIndex;
diff --git a/Engine/Plugins/Runtime/Metasound/Source/MetasoundFrontend/Public/MetasoundFrontendDocumentBuilder.h b/Engine/Plugins/Runtime/Metasound/Source/MetasoundFrontend/Public/MetasoundFrontendDocumentBuilder.h
--- a/Engine/Plugins/Runtime/Metasound/Source/MetasoundFrontend/Public/MetasoundFrontendDocumentBuilder.h
+++ b/Engine/Plugins/Runtime/Metasound/Source/MetasoundFrontend/Public/MetasoundFrontendDocumentBuilder.h
@@ -199,7 +199,7 @@ public:
 #if WITH_EDITOR
 	// Caches transient dependency metadata & style (class & vertex) found in the
 	// registry that is not necessary for serialization or core graph generation.
-	UE_API void CacheRegistryMetadata();
+	UE_API void CacheRegistryMetadata() const;
 #endif // WITH_EDITOR
 
 	// Returns whether or not the given edge can be added, which requires that its input
