Fix MetaSound nodes using node configuration not being copy pasted properly

#rnx

diff --git a/Engine/Plugins/Runtime/Metasound/Source/MetasoundEditor/Private/MetasoundEditorDocumentClipboardUtils.cpp b/Engine/Plugins/Runtime/Metasound/Source/MetasoundEditor/Private/MetasoundEditorDocumentClipboardUtils.cpp
--- a/Engine/Plugins/Runtime/Metasound/Source/MetasoundEditor/Private/MetasoundEditorDocumentClipboardUtils.cpp
+++ b/Engine/Plugins/Runtime/Metasound/Source/MetasoundEditor/Private/MetasoundEditorDocumentClipboardUtils.cpp
@@ -373,9 +373,15 @@ namespace Metasound::Editor
 				}
 				else
 				{
-					FNodeHandle NodeHandle = FGraphBuilder::AddExternalNodeHandle(MetaSound, Breadcrumb.ClassName);
-					ExternalNode->NodeID = NodeHandle->GetID();
-					if (!NodeHandle->IsValid())
+					FMetaSoundFrontendDocumentBuilder& Builder = IDocumentBuilderRegistry::GetChecked().FindOrBeginBuilding(OutAsset.GetOwningAsset());
+					const FMetasoundFrontendNode* NewNode = Builder.AddNodeByClassName(Breadcrumb.ClassName);
+					if (NewNode)
+					{
+						const FGuid& NewNodeID = NewNode->GetID();
+						ExternalNode->NodeID = NewNodeID;
+						Builder.SetNodeConfiguration(NewNodeID, Breadcrumb.NodeConfiguration);
+					}
+					else
 					{
 						OutPastedNodes.RemoveAtSwap(Index, 1, EAllowShrinking::No);
 						Graph.RemoveNode(ExternalNode);
@@ -404,9 +410,15 @@ namespace Metasound::Editor
 					FMetasoundFrontendClass ExternalClass;
 					if (ISearchEngine::Get().FindClassWithHighestVersion(Breadcrumb.ClassName, ExternalClass))
 					{
-						FConstNodeHandle NodeHandle = FGraphBuilder::AddExternalNodeHandle(MetaSound, Breadcrumb.ClassName);
-						ExternalNode->NodeID = NodeHandle->GetID();
-						if (!NodeHandle->IsValid())
+						FMetaSoundFrontendDocumentBuilder& Builder = IDocumentBuilderRegistry::GetChecked().FindOrBeginBuilding(OutAsset.GetOwningAsset());
+						const FMetasoundFrontendNode* NewNode = Builder.AddNodeByClassName(Breadcrumb.ClassName);
+						if (NewNode)
+						{
+							const FGuid& NewNodeID = NewNode->GetID();
+							ExternalNode->NodeID = NewNodeID;
+							Builder.SetNodeConfiguration(NewNodeID, Breadcrumb.NodeConfiguration);
+						}
+						else
 						{
 							OutPastedNodes.RemoveAtSwap(Index, 1, EAllowShrinking::No);
 							Graph.RemoveNode(ExternalNode);
diff --git a/Engine/Plugins/Runtime/Metasound/Source/MetasoundEditor/Private/MetasoundEditorGraphNode.cpp b/Engine/Plugins/Runtime/Metasound/Source/MetasoundEditor/Private/MetasoundEditorGraphNode.cpp
--- a/Engine/Plugins/Runtime/Metasound/Source/MetasoundEditor/Private/MetasoundEditorGraphNode.cpp
+++ b/Engine/Plugins/Runtime/Metasound/Source/MetasoundEditor/Private/MetasoundEditorGraphNode.cpp
@@ -1144,6 +1144,9 @@ void UMetasoundEditorGraphExternalNode::CacheBreadcrumb()
 
 		Breadcrumb.bIsClassNative = !bIsAssetClass;
 		Breadcrumb.ClassName = Metadata.GetClassName();
+
+		const FMetasoundFrontendNode& Node = GetFrontendNodeChecked();
+		Breadcrumb.NodeConfiguration = Node.Configuration;
 	}
 
 	// Cache template node generation parameters
diff --git a/Engine/Plugins/Runtime/Metasound/Source/MetasoundEditor/Public/MetasoundEditorGraphNode.h b/Engine/Plugins/Runtime/Metasound/Source/MetasoundEditor/Public/MetasoundEditorGraphNode.h
--- a/Engine/Plugins/Runtime/Metasound/Source/MetasoundEditor/Public/MetasoundEditorGraphNode.h
+++ b/Engine/Plugins/Runtime/Metasound/Source/MetasoundEditor/Public/MetasoundEditorGraphNode.h
@@ -53,6 +53,9 @@ struct FMetasoundEditorGraphNodeBreadcrumb
 	UPROPERTY()
 	bool bIsClassNative = true;
 
+	UPROPERTY()
+	TInstancedStruct<FMetaSoundFrontendNodeConfiguration> NodeConfiguration;
+
 	// For use with template nodes only
 	UPROPERTY()
 	TOptional<FNodeTemplateGenerateInterfaceParams> TemplateParams;
diff --git a/Engine/Plugins/Runtime/Metasound/Source/MetasoundFrontend/Private/MetasoundFrontendDocumentBuilder.cpp b/Engine/Plugins/Runtime/Metasound/Source/MetasoundFrontend/Private/MetasoundFrontendDocumentBuilder.cpp
--- a/Engine/Plugins/Runtime/Metasound/Source/MetasoundFrontend/Private/MetasoundFrontendDocumentBuilder.cpp
+++ b/Engine/Plugins/Runtime/Metasound/Source/MetasoundFrontend/Private/MetasoundFrontendDocumentBuilder.cpp
@@ -5923,6 +5923,23 @@ bool FMetaSoundFrontendDocumentBuilder::SetNodeCommentVisible(const FGuid& InNod
 }
 #endif // WITH_EDITOR
 
+bool FMetaSoundFrontendDocumentBuilder::SetNodeConfiguration(const FGuid& InNodeID, TInstancedStruct<FMetaSoundFrontendNodeConfiguration> InNodeConfiguration, const FGuid* InPageID)
+{
+	using namespace Metasound::Frontend;
+
+	const FGuid& PageID = InPageID ? *InPageID : BuildPageID;
+	FMetasoundFrontendGraph& Graph = GetDocumentChecked().RootGraph.FindGraphChecked(PageID);
+	const IDocumentGraphNodeCache& NodeCache = DocumentCache->GetNodeCache(PageID);
+	if (const int32* NodeIndex = NodeCache.FindNodeIndex(InNodeID))
+	{
+		FMetasoundFrontendNode& Node = Graph.Nodes[*NodeIndex];
+		Node.Configuration = MoveTemp(InNodeConfiguration);
+		UpdateNodeInterfaceFromConfiguration(InNodeID, InPageID);
+		return true;
+	}
+	return false;
+}
+
 bool FMetaSoundFrontendDocumentBuilder::SetNodeInputDefault(const FGuid& InNodeID, const FGuid& InVertexID, const FMetasoundFrontendLiteral& InLiteral, const FGuid* InPageID)
 {
 	using namespace Metasound::Frontend;
@@ -6341,34 +6358,41 @@ bool FMetaSoundFrontendDocumentBuilder::UpdateNodeInterfaceFromConfiguration(con
 {
 	if (FMetasoundFrontendNode* Node = FindNodeInternal(InNodeID, InPageID))
 	{
+		const FMetasoundFrontendClass* Class = FindDependency(Node->ClassID);
+		check(Class);
+
+		// Update class interface override
 		if (const FMetaSoundFrontendNodeConfiguration* ConfigurationPtr = Node->Configuration.GetPtr())
 		{
-			const FMetasoundFrontendClass* Class = FindDependency(Node->ClassID);
-			check(Class);
-
 			Node->ClassInterfaceOverride = ConfigurationPtr->OverrideDefaultInterface(*Class);
-			const FMetasoundFrontendClassInterface& ClassInterface = Class->GetInterfaceForNode(*Node);
+		}
+		else
+		{
+			// Set class interface override back to default if no node configuration
+			Node->ClassInterfaceOverride = TInstancedStruct<FMetasoundFrontendClassInterface>();
+		}
 
-			auto DisconnectInput = [this, &InNodeID, &InPageID](const FMetasoundFrontendVertex& InNodeInput)
-			{
-				this->RemoveEdgeToNodeInput(InNodeID, InNodeInput.VertexID, InPageID);
-			};
+		// Update node interface
+		const FMetasoundFrontendClassInterface& ClassInterface = Class->GetInterfaceForNode(*Node);
 
-			auto DisconnectOutput = [this, &InNodeID, &InPageID](const FMetasoundFrontendVertex& InNodeOutput)
-			{
-				this->RemoveEdgesFromNodeOutput(InNodeID, InNodeOutput.VertexID, InPageID);
-			};
+		auto DisconnectInput = [this, &InNodeID, &InPageID](const FMetasoundFrontendVertex& InNodeInput)
+		{
+			this->RemoveEdgeToNodeInput(InNodeID, InNodeInput.VertexID, InPageID);
+		};
+
+		auto DisconnectOutput = [this, &InNodeID, &InPageID](const FMetasoundFrontendVertex& InNodeOutput)
+		{
+			this->RemoveEdgesFromNodeOutput(InNodeID, InNodeOutput.VertexID, InPageID);
+		};
 #if WITH_EDITORONLY_DATA
-			const bool bInterfaceUpdated = Node->Interface.Update(ClassInterface, DisconnectInput, DisconnectOutput);
-			if (bInterfaceUpdated)
-			{
-				GetDocumentChecked().Metadata.ModifyContext.AddNodeIDModified(InNodeID);
-			}
+		const bool bInterfaceUpdated = Node->Interface.Update(ClassInterface, DisconnectInput, DisconnectOutput);
+		if (bInterfaceUpdated)
+		{
+			GetDocumentChecked().Metadata.ModifyContext.AddNodeIDModified(InNodeID);
+		}
 #else
-			Node->Interface.Update(ClassInterface, DisconnectInput, DisconnectOutput);
+		Node->Interface.Update(ClassInterface, DisconnectInput, DisconnectOutput);
 #endif
-		}
-
 		return true;
 	}
 
diff --git a/Engine/Plugins/Runtime/Metasound/Source/MetasoundFrontend/Public/MetasoundFrontendDocumentBuilder.h b/Engine/Plugins/Runtime/Metasound/Source/MetasoundFrontend/Public/MetasoundFrontendDocumentBuilder.h
--- a/Engine/Plugins/Runtime/Metasound/Source/MetasoundFrontend/Public/MetasoundFrontendDocumentBuilder.h
+++ b/Engine/Plugins/Runtime/Metasound/Source/MetasoundFrontend/Public/MetasoundFrontendDocumentBuilder.h
@@ -622,6 +622,10 @@ public:
 	// Sets the editor-only Unconnected Pins Hidden for a node with the given ID.
 	UE_API bool SetNodeUnconnectedPinsHidden(const FGuid& InNodeID, const bool bUnconnectedPinsHidden, const FGuid* InPageID = nullptr);
 #endif // WITH_EDITOR
+	
+	// Sets the node configuration for the given node and updates the interface
+	// returns true if the node configuration is set
+	UE_API bool SetNodeConfiguration(const FGuid& InNodeID, TInstancedStruct<FMetaSoundFrontendNodeConfiguration> InNodeConfiguration, const FGuid* InPageID = nullptr);
 
 	UE_API bool SetNodeInputDefault(const FGuid& InNodeID, const FGuid& InVertexID, const FMetasoundFrontendLiteral& InLiteral, const FGuid* InPageID = nullptr);
 
