Fix for static analysis
#rnx

diff --git a/Engine/Plugins/Runtime/Metasound/Source/MetasoundEditor/Private/MetasoundEditorGraphNode.cpp b/Engine/Plugins/Runtime/Metasound/Source/MetasoundEditor/Private/MetasoundEditorGraphNode.cpp
--- a/Engine/Plugins/Runtime/Metasound/Source/MetasoundEditor/Private/MetasoundEditorGraphNode.cpp
+++ b/Engine/Plugins/Runtime/Metasound/Source/MetasoundEditor/Private/MetasoundEditorGraphNode.cpp
@@ -1218,14 +1218,15 @@ void UMetasoundEditorGraphExternalNode::Validate(const FMetaSoundFrontendDocumen
 	UObject* Outermost = GetOutermostObject();
 	check(Outermost);
 
-	const FMetasoundFrontendClass* Class = nullptr;
 	const FMetasoundFrontendNode* Node = InBuilder.FindNode(NodeID);
 
-	if (Node)
+	if (!Node)
 	{
-		Class = InBuilder.FindDependency(Node->ClassID);
+		OutResult.SetMessage(EMessageSeverity::Error, TEXT("Stale node not found in given MetaSound"));
+		return;
 	}
 
+	const FMetasoundFrontendClass* Class = InBuilder.FindDependency(Node->ClassID);
 	if (!Class)
 	{
 		OutResult.SetMessage(EMessageSeverity::Error, FString::Format(TEXT("Node with ID '{0}' class not found"), { *Node->ClassID.ToString() }));
