Fix unclear MetaSound editor node behavior when promoting node input to an input that already exists in the graph

#rnx

diff --git a/Engine/Plugins/Runtime/Metasound/Source/MetasoundEditor/Private/MetasoundEditorGraphSchema.cpp b/Engine/Plugins/Runtime/Metasound/Source/MetasoundEditor/Private/MetasoundEditorGraphSchema.cpp
--- a/Engine/Plugins/Runtime/Metasound/Source/MetasoundEditor/Private/MetasoundEditorGraphSchema.cpp
+++ b/Engine/Plugins/Runtime/Metasound/Source/MetasoundEditor/Private/MetasoundEditorGraphSchema.cpp
@@ -749,8 +749,12 @@ PRAGMA_ENABLE_DEPRECATION_WARNINGS
 			};
 
 			const FMetasoundFrontendNode* NewNode = nullptr;
+			// Name may be different than previous name because it may already exist in the graph,
+			// and CreateUniqueClassInput will add an identifier to the end to make it unique
+			FName NewName; 
 			{
 				FMetasoundFrontendClassInput ClassInput = FGraphBuilder::CreateUniqueClassInput(ParentMetasound, VertexParams, DefaultLiterals, &InputVertex->Name);
+				NewName = ClassInput.Name;
 				NewNode = Builder.AddGraphInput(MoveTemp(ClassInput));
 			}
 
@@ -759,7 +763,7 @@ PRAGMA_ENABLE_DEPRECATION_WARNINGS
 				UMetasoundEditorGraphInput* Input = MetasoundGraph->FindOrAddInput(NewNode->GetID());
 				if (ensure(Input))
 				{
-					if (const FMetasoundFrontendNode* NewTemplateNode = FInputNodeTemplate::CreateNode(Builder, InputVertex->Name))
+					if (const FMetasoundFrontendNode* NewTemplateNode = FInputNodeTemplate::CreateNode(Builder, NewName))
 					{
 						if (UMetasoundEditorGraphNode* NewGraphNode = FGraphBuilder::AddInputNode(ParentMetasound, NewTemplateNode->GetID()))
 						{
