Minor category clean-up
#rnx

diff --git a/Engine/Plugins/Runtime/Metasound/Source/MetasoundEngine/Public/MetasoundBuilderSubsystem.h b/Engine/Plugins/Runtime/Metasound/Source/MetasoundEngine/Public/MetasoundBuilderSubsystem.h
--- a/Engine/Plugins/Runtime/Metasound/Source/MetasoundEngine/Public/MetasoundBuilderSubsystem.h
+++ b/Engine/Plugins/Runtime/Metasound/Source/MetasoundEngine/Public/MetasoundBuilderSubsystem.h
@@ -276,7 +276,7 @@ public:
 	// If target page is not implemented (or cooked in a runtime build) for the active platform,
 	// uses order of cooked pages(see 'Page Settings' for order) falling back to lower index - ordered page
 	// implemented in MetaSound asset. If no fallback is found, uses default graph/input default.
-	UFUNCTION(BlueprintCallable, Category = "MetaSounds|Pages")
+	UFUNCTION(BlueprintCallable, Category = "Audio|MetaSound|Pages")
 	UE_API UPARAM(DisplayName = "TargetPageChanged") bool SetTargetPage(FName PageName);
 
 	UFUNCTION(BlueprintCallable, Category = "Audio|MetaSound|Builder")
