- Add editor-only function to MetaSound Asset Subsystem to find all referencing MetaSounds of a given asset and optionally filter by preset status.

#rnx

diff --git a/Engine/Plugins/Runtime/Metasound/Source/MetasoundEngine/Private/MetasoundAssetSubsystem.cpp b/Engine/Plugins/Runtime/Metasound/Source/MetasoundEngine/Private/MetasoundAssetSubsystem.cpp
--- a/Engine/Plugins/Runtime/Metasound/Source/MetasoundEngine/Private/MetasoundAssetSubsystem.cpp
+++ b/Engine/Plugins/Runtime/Metasound/Source/MetasoundEngine/Private/MetasoundAssetSubsystem.cpp
@@ -1537,6 +1537,79 @@ bool UMetaSoundAssetSubsystem::FindAssetClassInfo(const FTopLevelAssetPath& InPa
 
 	return false;
 }
+
+bool UMetaSoundAssetSubsystem::FindReferencingAssetClassInfo(TScriptInterface<IMetaSoundDocumentInterface> MetaSound, TArray<FTopLevelAssetPath>& OutPaths, TArray<FMetaSoundDocumentInfo>& OutDocInfo, TArray<FMetaSoundClassInterfaceInfo>& OutInterfaceInfo, bool bOnlyPresets, bool bForceLoad) const
+{
+	using namespace Metasound;
+	using namespace Metasound::Engine;
+	using namespace Metasound::Frontend;
+
+	IMetasoundUObjectRegistry& ObjectRegistry = IMetasoundUObjectRegistry::Get();
+
+	const IAssetRegistry& AssetRegistry = IAssetRegistry::GetChecked();
+
+	const FTopLevelAssetPath MetaSoundPath(MetaSound.GetObject());
+	TArray<FName> OnDiskReferencers;
+	AssetRegistry.GetReferencers(MetaSoundPath.GetPackageName(), OnDiskReferencers);
+
+	bool bSuccess = true;
+
+	OutDocInfo.Empty();
+	OutInterfaceInfo.Empty();
+	OutPaths.Empty();
+
+	for (const FName& OnDiskReferencer : OnDiskReferencers)
+	{
+		FAssetData RefAssetData;
+
+		FString AssetName = *OnDiskReferencer.ToString();
+		{
+			int32 Index = INDEX_NONE;
+			if (AssetName.FindLastChar('/', Index))
+			{
+				AssetName.RightChopInline(Index + 1);
+			}
+		}
+		const FTopLevelAssetPath RefAssetPath(FName(*OnDiskReferencer.ToString()), FName(*AssetName));
+		UE::AssetRegistry::EExists RefExists = AssetRegistry.TryGetAssetByObjectPath(FSoftObjectPath(RefAssetPath), RefAssetData);
+		if (RefExists != UE::AssetRegistry::EExists::Exists)
+		{
+			bSuccess = false;
+			continue;
+		}
+
+		const UClass* AssetClass = RefAssetData.GetClass();
+		if (!AssetClass)
+		{
+			bSuccess = false;
+			continue;
+		}
+		 
+		if (!ObjectRegistry.IsRegisteredClass(*AssetClass))
+		{
+			// Could be non-MetaSound references (Blueprint, redirect, etc.), so isn't
+			// a failure to succeed if referencing class isn't strictly MetaSound.
+			continue;
+		}
+
+		FMetaSoundDocumentInfo DocInfo;
+		FMetaSoundClassInterfaceInfo InterfaceInfo;
+
+		const bool bRefFound = FindAssetClassInfo(RefAssetPath, DocInfo, InterfaceInfo, bForceLoad);
+		bSuccess &= bRefFound;
+		if (bRefFound)
+		{
+			if (!bOnlyPresets || (bOnlyPresets && DocInfo.bIsPreset != 0))
+			{
+				OutPaths.Add(RefAssetPath);
+				OutDocInfo.Add(MoveTemp(DocInfo));
+				OutInterfaceInfo.Add(MoveTemp(InterfaceInfo));
+			}
+		}
+	}
+
+	return bSuccess;
+}
 #endif // WITH_EDITOR
 
 FMetasoundAssetBase* UMetaSoundAssetSubsystem::GetAsAsset(UObject& InObject) const
diff --git a/Engine/Plugins/Runtime/Metasound/Source/MetasoundEngine/Public/MetasoundAssetSubsystem.h b/Engine/Plugins/Runtime/Metasound/Source/MetasoundEngine/Public/MetasoundAssetSubsystem.h
--- a/Engine/Plugins/Runtime/Metasound/Source/MetasoundEngine/Public/MetasoundAssetSubsystem.h
+++ b/Engine/Plugins/Runtime/Metasound/Source/MetasoundEngine/Public/MetasoundAssetSubsystem.h
@@ -315,20 +315,33 @@ PRAGMA_ENABLE_DEPRECATION_WARNINGS
 	// info without loading the asset if its tag data is up-to-date, or if set to force load, will load
 	// otherwise (synchronously and can be slow).  Returns true if asset is found, was MetaSound, and all data
 	// was retrieved successfully, false if not.
-	UFUNCTION(BlueprintCallable, Category = "MetaSounds|Utilities", meta = (AdvancedDisplay = "3"))
+	UFUNCTION(BlueprintCallable, Category = "Audio|MetaSound|Utilities", meta = (AdvancedDisplay = "3"))
 	UE_API UPARAM(DisplayName = "Info Found") bool FindAssetClassInfo(
 		UPARAM(DisplayName = "Path") const FTopLevelAssetPath& InPath,
 		UPARAM(DisplayName = "Document Info") FMetaSoundDocumentInfo& OutDocInfo,
 		UPARAM(DisplayName = "Interface Info") FMetaSoundClassInterfaceInfo& OutInterfaceInfo,
 		UPARAM(DisplayName = "Force Load") bool bForceLoad = true) const;
 
-	UFUNCTION(BlueprintCallable, Category = "MetaSounds|Utilities")
+	// Returns info for all MetaSounds assets referencing the given MetaSound as a dependency.  Only returns assets on disk
+	// (i.e. not in-memory, transient MetaSounds created by the builder API).  (Optional) Only returns references that are
+	// presets of the given MetaSound if bOnlyPresets is true.
+	UFUNCTION(BlueprintCallable, Category = "Audio|MetaSound|Utilities", meta = (DisplayName = "Find Referencing Asset Class Info", AdvancedDisplay = "5"))
+	UE_API UPARAM(DisplayName = "Reference Info Found") bool FindReferencingAssetClassInfo(
+		UPARAM(DisplayName = "MetaSound") TScriptInterface<IMetaSoundDocumentInterface> MetaSound,
+		UPARAM(DisplayName = "Reference Paths") TArray<FTopLevelAssetPath>& OutPaths,
+		UPARAM(DisplayName = "Document Info") TArray<FMetaSoundDocumentInfo>& OutDocInfo,
+		UPARAM(DisplayName = "Interface Info") TArray<FMetaSoundClassInterfaceInfo>& OutInterfaceInfo,
+		UPARAM(DisplayName = "Only Presets") bool bOnlyPresets = false,
+		UPARAM(DisplayName = "Force Load") bool bForceLoad = true) const;
+
+	// Reassigns a new class name for the given MetaSound object, invalidating all references to the given MetaSound Asset.
+	UFUNCTION(BlueprintCallable, Category = "Audio|MetaSound|Utilities")
 	UE_API UPARAM(DisplayName = "Reassigned") bool ReassignClassName(TScriptInterface<IMetaSoundDocumentInterface> DocInterface);
 
 	// Replaces dependencies in a MetaSound with the given class name and version with another MetaSound with the given
 	// class name and version.  Can be asset or code-defined.  It is up to the caller to validate the two classes have
 	// matching interfaces (Swapping with classes of unmatched interfaces can leave MetaSound in non-executable state).
-	UFUNCTION(BlueprintCallable, Category = "MetaSounds|Utilities", meta = (AdvancedDisplay = "3"))
+	UFUNCTION(BlueprintCallable, Category = "Audio|MetaSound|Utilities", meta = (AdvancedDisplay = "3"))
 	UE_API UPARAM(DisplayName = "References Replaced") bool ReplaceReferencesInDirectory(
 		const TArray<FMetaSoundAssetDirectory>& InDirectories,
 		const FMetasoundFrontendClassName& OldClassName,
@@ -337,10 +350,10 @@ PRAGMA_ENABLE_DEPRECATION_WARNINGS
 		const FMetasoundFrontendVersionNumber NewVersion = FMetasoundFrontendVersionNumber());
 #endif // WITH_EDITOR
 
-	UFUNCTION(BlueprintCallable, Category = "MetaSounds|Registration")
+	UFUNCTION(BlueprintCallable, Category = "Audio|MetaSound|Registration")
 	UE_API void RegisterAssetClassesInDirectories(const TArray<FMetaSoundAssetDirectory>& Directories);
 
-	UFUNCTION(BlueprintCallable, Category = "MetaSounds|Registration")
+	UFUNCTION(BlueprintCallable, Category = "Audio|MetaSound|Registration")
 	UE_API void UnregisterAssetClassesInDirectories(const TArray<FMetaSoundAssetDirectory>& Directories);
 
 protected:
