- Add asset actions to browse to/open MetaSound Preset Parent
- Clean-up right-click options for faster previewing of MetaSounds, better sub-categorization, & add abilities to browse to and open MetaSound preset parents

#rnx

diff --git a/Engine/Plugins/Runtime/Metasound/Source/MetasoundEditor/Private/MetasoundAssetDefinitions.cpp b/Engine/Plugins/Runtime/Metasound/Source/MetasoundEditor/Private/MetasoundAssetDefinitions.cpp
--- a/Engine/Plugins/Runtime/Metasound/Source/MetasoundEditor/Private/MetasoundAssetDefinitions.cpp
+++ b/Engine/Plugins/Runtime/Metasound/Source/MetasoundEditor/Private/MetasoundAssetDefinitions.cpp
@@ -15,13 +15,14 @@
 #include "MetasoundAssetManager.h"
 #include "MetasoundAssetSubsystem.h"
 #include "MetasoundBuilderSubsystem.h"
-#include "MetasoundSource.h"
+#include "MetasoundDocumentInterface.h"
 #include "MetasoundEditor.h"
 #include "MetasoundEditorGraphBuilder.h"
 #include "MetasoundEditorModule.h"
 #include "MetasoundEditorSettings.h"
 #include "MetasoundEditorSubsystem.h"
 #include "MetasoundFactory.h"
+#include "MetasoundSource.h"
 #include "MetasoundUObjectRegistry.h"
 #include "Misc/AssetRegistryInterface.h"
 #include "Modules/ModuleManager.h"
@@ -65,6 +66,162 @@ namespace Metasound::Editor
 
 			return &Metasound::Editor::Style::GetSlateBrushSafe(FName(*BrushName));
 		}
+
+		void ExecuteBrowseToPresetReferencedAsset(const FToolMenuContext& InContext)
+		{
+			using namespace Metasound;
+
+			if (UObject* MetaSound = UContentBrowserAssetContextMenuContext::LoadSingleSelectedAsset<UObject>(InContext))
+			{
+				if (UAssetEditorSubsystem* AssetSubsystem = GEditor->GetEditorSubsystem<UAssetEditorSubsystem>())
+				{
+					FMetaSoundFrontendDocumentBuilder& DocumentBuilder = Frontend::IDocumentBuilderRegistry::GetChecked().FindOrBeginBuilding(MetaSound);
+					if (const FMetasoundAssetBase* MetaSoundAsset = DocumentBuilder.GetReferencedPresetAsset())
+					{
+						if (GEditor)
+						{
+							GEditor->SyncBrowserToObjects({ MetaSoundAsset->GetOwningAsset() });
+						}
+					}
+				}
+			}
+		}
+
+		void ExecuteOpenPresetReferencedAsset(const FToolMenuContext& InContext)
+		{
+			using namespace Metasound;
+
+			if (UObject* MetaSound = UContentBrowserAssetContextMenuContext::LoadSingleSelectedAsset<UObject>(InContext))
+			{
+				if (UAssetEditorSubsystem* AssetSubsystem = GEditor->GetEditorSubsystem<UAssetEditorSubsystem>())
+				{
+					FMetaSoundFrontendDocumentBuilder& DocumentBuilder = Frontend::IDocumentBuilderRegistry::GetChecked().FindOrBeginBuilding(MetaSound);
+					if (const FMetasoundAssetBase* MetaSoundAsset = DocumentBuilder.GetReferencedPresetAsset())
+					{
+						AssetSubsystem->OpenEditorForAsset(MetaSoundAsset->GetOwningAsset());
+					}
+				}
+			}
+		}
+		
+		template <typename TClass, typename TFactoryClass>
+		void ExecuteCreateMetaSoundPreset(const FToolMenuContext& MenuContext)
+		{
+			if (const UContentBrowserAssetContextMenuContext* Context = UContentBrowserAssetContextMenuContext::FindContextWithAssets(MenuContext))
+			{
+				using namespace Metasound::Editor;
+				TArray<UObject*> ObjectsToSync;
+
+				for (TClass* ReferencedMetaSound : Context->LoadSelectedObjects<TClass>())
+				{
+					FString PackagePath;
+					FString AssetName;
+					UObject* NewMetaSound = nullptr;
+
+					IAssetTools::Get().CreateUniqueAssetName(ReferencedMetaSound->GetOutermost()->GetName(), TEXT("_Preset"), PackagePath, AssetName);
+
+					Metasound::Editor::IMetasoundEditorModule& MetaSoundEditorModule = FModuleManager::GetModuleChecked<Metasound::Editor::IMetasoundEditorModule>("MetaSoundEditor");
+					if (MetaSoundEditorModule.IsRestrictedMode())
+					{
+						// Cannot duplicate cooked assets in restricted mode, so create new object and copy over properties
+						// for sources in InitAsset below. Since copying properties is done manually, 
+						// SetSoundWaveSettingsFromTemplate may need to be updated with properties to be copied. 
+						UMetaSoundBaseFactory* Factory = NewObject<TFactoryClass>();
+						Factory->ReferencedMetaSoundObject = ReferencedMetaSound;
+
+						NewMetaSound = IAssetTools::Get().CreateAssetWithDialog(AssetName, FPackageName::GetLongPackagePath(PackagePath), Factory->GetSupportedClass(), Factory);
+					}
+					else
+					{
+						// Duplicate asset to preserve properties from referenced asset (ex. quality settings, soundwave properties)
+						NewMetaSound = IAssetTools::Get().DuplicateAssetWithDialogAndTitle(AssetName, FPackageName::GetLongPackagePath(PackagePath), ReferencedMetaSound, LOCTEXT("CreateMetaSoundPresetTitle", "Create MetaSound Preset"));
+					}
+				
+					if (NewMetaSound)
+					{
+						UMetaSoundEditorSubsystem::GetChecked().InitAsset(*NewMetaSound, ReferencedMetaSound, /*bClearDocument=*/true);
+
+						FGraphBuilder::RegisterGraphWithFrontend(*NewMetaSound);
+						ObjectsToSync.Add(NewMetaSound);
+					}
+					else
+					{
+						UE_LOG(LogMetaSound, Display, TEXT("Error creating new asset when creating preset '%s' or asset creation was canceled by user."), *AssetName);
+					}
+				}
+
+				// Sync content browser to newly created valid assets
+				// Assets can be invalid if multiple assets are created with the same name 
+				// then force overwritten within the same operation
+				ObjectsToSync.RemoveAllSwap([](const UObject* InObject)
+				{
+					return !InObject || !InObject->IsValidLowLevelFast();
+				});
+
+				if (ObjectsToSync.Num() > 0)
+				{
+					FContentBrowserModule& ContentBrowserModule = FModuleManager::LoadModuleChecked<FContentBrowserModule>("ContentBrowser");
+					ContentBrowserModule.Get().SyncBrowserToAssets(ObjectsToSync);
+				}
+			}
+		}
+
+		bool IsPreset(const FToolMenuContext& InContext)
+		{
+			using namespace Metasound::Frontend;
+
+			if (const UContentBrowserAssetContextMenuContext* CBContext = UContentBrowserAssetContextMenuContext::FindContextWithAssets(InContext))
+			{
+				if (CBContext->SelectedAssets.Num() == 1)
+				{
+					return FMetaSoundAssetClassInfo(CBContext->SelectedAssets.Last()).DocInfo.bIsPreset;
+				}
+			}
+
+			return false;
+		}
+
+		void AddMetaSoundActions(UToolMenu& Menu, const FSlateIcon& AssetIcon, const FText& ClassName, FToolUIAction CreatePresetAction)
+		{
+			FToolMenuSection& Section = Menu.FindOrAddSection("MetaSoundPresetActions", LOCTEXT("PresetActions_Label", "MetaSound"));
+			{
+				Section.AddDynamicEntry(NAME_None, FNewToolMenuSectionDelegate::CreateLambda([CreateAction = MoveTemp(CreatePresetAction), ClassName, AssetIcon](FToolMenuSection& InSection)
+				{
+					using namespace Metasound::Editor;
+
+					{
+						const TAttribute<FText> Label = FText::Format(LOCTEXT("MetaSound_CreatePreset", "Create Preset..."), ClassName);
+						const TAttribute<FText> ToolTip = FText::Format(LOCTEXT("MetaSound_CreatePresetToolTipFormat", "Creates a {0} Preset using the selected MetaSound as a reference."), ClassName);
+
+						InSection.AddMenuEntry("MetaSound_CreatePreset", Label, ToolTip, AssetIcon, CreateAction);
+					}
+					{
+						const TAttribute<FText> Label = LOCTEXT("MetaSound_BrowseToPresetParentFormat", "Browse To Preset Parent");
+						const TAttribute<FText> ToolTip = FText::Format(LOCTEXT("MetaSound_BrowseToPresetParentToolTipFormat", "Browses to the selected {0} preset's referenced parent asset in the content browser."), ClassName);
+						const FSlateIcon FindInContentBrowserIcon(FAppStyle::GetAppStyleSetName(), "SystemWideCommands.FindInContentBrowser");
+
+						FToolUIAction UIAction;
+						UIAction.ExecuteAction = FToolMenuExecuteAction::CreateStatic(&AssetDefinitionsPrivate::ExecuteBrowseToPresetReferencedAsset);
+						UIAction.CanExecuteAction = FToolMenuCanExecuteAction::CreateStatic(&AssetDefinitionsPrivate::IsPreset);
+						UIAction.IsActionVisibleDelegate = FToolMenuIsActionButtonVisible::CreateStatic(&AssetDefinitionsPrivate::IsPreset);
+						InSection.AddMenuEntry("MetaSoundSource_BrowseToPresetParent", Label, ToolTip, FindInContentBrowserIcon, UIAction);
+					}
+
+					IMetasoundEditorModule& EditorModule = FModuleManager::GetModuleChecked<Metasound::Editor::IMetasoundEditorModule>("MetaSoundEditor");
+					if (!EditorModule.IsRestrictedMode())
+					{
+						const TAttribute<FText> Label = LOCTEXT("MetaSound_EditPresetParentFormat", "Edit Preset Parent...");
+						const TAttribute<FText> ToolTip = FText::Format(LOCTEXT("MetaSound_OpenPresetParentToolTipFormat", "Opens the selected {0} preset's parent MetaSound."), ClassName);
+
+						FToolUIAction UIAction;
+						UIAction.ExecuteAction = FToolMenuExecuteAction::CreateStatic(&AssetDefinitionsPrivate::ExecuteOpenPresetReferencedAsset);
+						UIAction.CanExecuteAction = FToolMenuCanExecuteAction::CreateStatic(&AssetDefinitionsPrivate::IsPreset);
+						UIAction.IsActionVisibleDelegate = FToolMenuIsActionButtonVisible::CreateStatic(&AssetDefinitionsPrivate::IsPreset);
+						InSection.AddMenuEntry("MetaSoundSource_OpenToPresetParent", Label, ToolTip, FSlateIcon(FAppStyle::GetAppStyleSetName(), "Icons.Edit"), UIAction);
+					}
+				}));
+			}
+		}
 	} // namespace AssetDefinitionsPrivate
 } // namespace Metasound::Editor
 
@@ -411,155 +568,135 @@ void UAssetDefinition_MetaSoundSource::GetAssetActionButtonExtensions(const FAss
 
 namespace MenuExtension_MetaSoundSourceTemplate
 {
-	template <typename TClass, typename TFactoryClass>
-	void ExecuteCreateMetaSoundPreset(const FToolMenuContext& MenuContext)
-	{
-		if (const UContentBrowserAssetContextMenuContext* Context = UContentBrowserAssetContextMenuContext::FindContextWithAssets(MenuContext))
-		{
-			using namespace Metasound::Editor;
-			TArray<UObject*> ObjectsToSync;
-
-			for (TClass* ReferencedMetaSound : Context->LoadSelectedObjects<TClass>())
-			{
-				FString PackagePath;
-				FString AssetName;
-				UObject* NewMetaSound;
-
-				IAssetTools::Get().CreateUniqueAssetName(ReferencedMetaSound->GetOutermost()->GetName(), TEXT("_Preset"), PackagePath, AssetName);
-
-				Metasound::Editor::IMetasoundEditorModule& MetaSoundEditorModule = FModuleManager::GetModuleChecked<Metasound::Editor::IMetasoundEditorModule>("MetaSoundEditor");
-				if (MetaSoundEditorModule.IsRestrictedMode())
-				{
-					// Cannot duplicate cooked assets in restricted mode, so create new object and copy over properties
-					// for sources in InitAsset below. Since copying properties is done manually, 
-					// SetSoundWaveSettingsFromTemplate may need to be updated with properties to be copied. 
-					UMetaSoundBaseFactory* Factory = NewObject<TFactoryClass>();
-					Factory->ReferencedMetaSoundObject = ReferencedMetaSound;
-
-					NewMetaSound = IAssetTools::Get().CreateAssetWithDialog(AssetName, FPackageName::GetLongPackagePath(PackagePath), Factory->GetSupportedClass(), Factory);
-				}
-				else
-				{
-					// Duplicate asset to preserve properties from referenced asset (ex. quality settings, soundwave properties)
-					NewMetaSound = IAssetTools::Get().DuplicateAssetWithDialogAndTitle(AssetName, FPackageName::GetLongPackagePath(PackagePath), ReferencedMetaSound, LOCTEXT("CreateMetaSoundPresetTitle", "Create MetaSound Preset"));
-				}
-				
-				if (NewMetaSound)
-				{
-					UMetaSoundEditorSubsystem::GetChecked().InitAsset(*NewMetaSound, ReferencedMetaSound, /*bClearDocument=*/true);
-
-					FGraphBuilder::RegisterGraphWithFrontend(*NewMetaSound);
-					ObjectsToSync.Add(NewMetaSound);
-				}
-				else
-				{
-					UE_LOG(LogMetaSound, Display, TEXT("Error creating new asset when creating preset '%s' or asset creation was canceled by user."), *AssetName);
-				}
-			}
-
-			// Sync content browser to newly created valid assets
-			// Assets can be invalid if multiple assets are created with the same name 
-			// then force overwritten within the same operation
-			ObjectsToSync.RemoveAllSwap([](const UObject* InObject)
-			{
-				return !InObject || !InObject->IsValidLowLevelFast();
-			});
-
-			if (ObjectsToSync.Num() > 0)
-			{
-				FContentBrowserModule& ContentBrowserModule = FModuleManager::LoadModuleChecked<FContentBrowserModule>("ContentBrowser");
-				ContentBrowserModule.Get().SyncBrowserToAssets(ObjectsToSync);
-			}
-		}
-	}
-
  	static FDelayedAutoRegisterHelper DelayedAutoRegister(EDelayedRegisterRunPhase::EndOfEngineInit, []
 	{
  		UToolMenus::RegisterStartupCallback(FSimpleMulticastDelegate::FDelegate::CreateLambda([]()
  		{
+			using namespace Metasound::Editor;
+
  			FToolMenuOwnerScoped OwnerScoped(UE_MODULE_NAME);
- 
- 			{
-				UToolMenu* Menu = UE::ContentBrowser::ExtendToolMenu_AssetContextMenu(UMetaSoundSource::StaticClass());
 
-				FToolMenuSection& Section = Menu->FindOrAddSection("GetAssetActions");
-				{
+			UClass* MetaSoundClass = UMetaSoundSource::StaticClass();
+			check(MetaSoundClass);
+			if (UToolMenu* Menu = UE::ContentBrowser::ExtendToolMenu_AssetContextMenu(MetaSoundClass))
+			{
+				FToolUIAction CreateAction;
+				CreateAction.ExecuteAction = FToolMenuExecuteAction::CreateStatic(&AssetDefinitionsPrivate::ExecuteCreateMetaSoundPreset<UMetaSoundSource, UMetaSoundSourceFactory>);
+				const FSlateIcon AssetIcon = Metasound::Editor::Style::CreateSlateIcon("ClassIcon.MetasoundSource");
 
-					Section.AddDynamicEntry(NAME_None, FNewToolMenuSectionDelegate::CreateLambda([](FToolMenuSection& InSection)
+				AssetDefinitionsPrivate::AddMetaSoundActions(*Menu, AssetIcon, MetaSoundClass->GetDisplayNameText(), MoveTemp(CreateAction));
+				{
+					FToolMenuSection& Section = Menu->FindOrAddSection("AudioPlaybackActions", LOCTEXT("PlaybackActions_Label", "Playback"), FToolMenuInsert(FName { }, EToolMenuInsertType::Last));
 					{
+						Section.AddDynamicEntry(NAME_None, FNewToolMenuSectionDelegate::CreateLambda([AssetIcon](FToolMenuSection& InSection)
 						{
-							const TAttribute<FText> Label = LOCTEXT("Sound_PlaySound", "Play");
-							const TAttribute<FText> ToolTip = LOCTEXT("Sound_PlaySoundTooltip", "Plays the selected sound.");
-							const FSlateIcon Icon = FSlateIcon(FAppStyle::GetAppStyleSetName(), "MediaAsset.AssetActions.Play.Small");
-
-							FToolUIAction UIAction;
-							UIAction.ExecuteAction = FToolMenuExecuteAction::CreateStatic(&UAssetDefinition_MetaSoundSource::ExecutePlaySound);
-							UIAction.CanExecuteAction = FToolMenuCanExecuteAction::CreateStatic(&UAssetDefinition_MetaSoundSource::CanExecutePlayCommand);
-							InSection.AddMenuEntry("Sound_PlaySound", Label, ToolTip, Icon, UIAction);
-						}
-						{
-							const TAttribute<FText> Label = LOCTEXT("Sound_StopSound", "Stop");
-							const TAttribute<FText> ToolTip = LOCTEXT("Sound_StopSoundTooltip", "Stops the selected sounds.");
-							const FSlateIcon Icon = FSlateIcon(FAppStyle::GetAppStyleSetName(), "MediaAsset.AssetActions.Stop.Small");
+							auto IsPlayingThis = [](const FToolMenuContext& InContext) { return UAssetDefinition_SoundBase::IsPlayingContextAsset(InContext, true /* bMustMatchContext */); };
+							auto IsNotPlayingThis = [](const FToolMenuContext& InContext) { return !UAssetDefinition_SoundBase::IsPlayingContextAsset(InContext, true /* bMustMatchContext */); };
+							auto IsNotPlayingAny = [](const FToolMenuContext& InContext) { return !UAssetDefinition_SoundBase::IsPlayingContextAsset(InContext, false /* bMustMatchContext */); };
+							auto IsPlayingAny = [](const FToolMenuContext& InContext) { return UAssetDefinition_SoundBase::IsPlayingContextAsset(InContext, false /* bMustMatchContext */); };
+							auto IsPlayingOther = [](const FToolMenuContext& InContext) { return UAssetDefinition_SoundBase::IsPlayingContextAsset(InContext, false /* bMustMatchContext */) && !UAssetDefinition_SoundBase::IsPlayingContextAsset(InContext, true /* bMustMatchContext */); };
 
-							FToolUIAction UIAction;
-							UIAction.ExecuteAction = FToolMenuExecuteAction::CreateStatic(&UAssetDefinition_MetaSoundSource::ExecuteStopSound);
-							InSection.AddMenuEntry("Sound_StopSound", Label, ToolTip, Icon, UIAction);
-						}
-						{
-							const TAttribute<FText> Label = LOCTEXT("Sound_MuteSound", "Mute");
-							const TAttribute<FText> ToolTip = LOCTEXT("Sound_MuteSoundTooltip", "Mutes the selected sounds.");
-							const FSlateIcon Icon = FSlateIcon(FAppStyle::GetAppStyleSetName(), "MediaAsset.AssetActions.Mute.Small");
-
-							FToolUIAction UIAction;
-							UIAction.ExecuteAction = FToolMenuExecuteAction::CreateStatic(&UAssetDefinition_MetaSoundSource::ExecuteMuteSound);
-							UIAction.CanExecuteAction = FToolMenuCanExecuteAction::CreateStatic(&UAssetDefinition_MetaSoundSource::CanExecuteMuteCommand);
-							UIAction.GetActionCheckState = FToolMenuGetActionCheckState::CreateStatic(&UAssetDefinition_MetaSoundSource::IsActionCheckedMute);
-							InSection.AddMenuEntry("Sound_SoundMute", Label, ToolTip, Icon, UIAction, EUserInterfaceActionType::ToggleButton);
-						}
-						{
-							const TAttribute<FText> Label = LOCTEXT("Sound_SoloSound", "Solo");
-							const TAttribute<FText> ToolTip = LOCTEXT("Sound_SoloSoundTooltip", "Solos the selected sounds.");
-							const FSlateIcon Icon = FSlateIcon(FAppStyle::GetAppStyleSetName(), "MediaAsset.AssetActions.Solo.Small");
-
-							FToolUIAction UIAction;
-							UIAction.ExecuteAction = FToolMenuExecuteAction::CreateStatic(&UAssetDefinition_MetaSoundSource::ExecuteSoloSound);
-							UIAction.CanExecuteAction = FToolMenuCanExecuteAction::CreateStatic(&UAssetDefinition_MetaSoundSource::CanExecuteSoloCommand);
-							UIAction.GetActionCheckState = FToolMenuGetActionCheckState::CreateStatic(&UAssetDefinition_MetaSoundSource::IsActionCheckedSolo);
-							InSection.AddMenuEntry("Sound_StopSolo", Label, ToolTip, Icon, UIAction, EUserInterfaceActionType::ToggleButton);
-						}
-						{
-							const TAttribute<FText> Label = LOCTEXT("MetaSoundSource_CreatePreset", "Create MetaSound Source Preset");
-							const TAttribute<FText> ToolTip = LOCTEXT("MetaSoundSource_CreatePresetToolTip", "Creates a MetaSoundSource Preset using the selected MetaSound's root graph as a reference.");
-							const FSlateIcon Icon = Metasound::Editor::Style::CreateSlateIcon("ClassIcon.MetasoundSource");
+							{
+								const TAttribute<FText> Label = LOCTEXT("Sound_PlaySound" , "Play");
+								const TAttribute<FText> ToolTip = LOCTEXT("Sound_PlaySoundTooltip", "Plays the selected sound.");
+								const FSlateIcon Icon = FSlateIcon(FAppStyle::GetAppStyleSetName(), "MediaAsset.AssetActions.Play.Small");
+
+								FToolUIAction UIAction;
+								UIAction.ExecuteAction = FToolMenuExecuteAction::CreateStatic(&UAssetDefinition_MetaSoundSource::ExecutePlaySound);
+								UIAction.CanExecuteAction = FToolMenuCanExecuteAction::CreateStatic(&UAssetDefinition_MetaSoundSource::CanExecutePlayCommand);
+								UIAction.IsActionVisibleDelegate = FToolMenuIsActionButtonVisible::CreateLambda(IsNotPlayingThis);
+								InSection.AddMenuEntry("Sound_PlaySound", Label, ToolTip, Icon, UIAction);
+							}
+							{
+								const TAttribute<FText> Label = LOCTEXT("Sound_RestartSound", "Restart");
+								const TAttribute<FText> ToolTip = LOCTEXT("Sound_RestartSoundTooltip", "Restarts the selected sound.");
+								const FSlateIcon Icon = FSlateIcon(FAppStyle::GetAppStyleSetName(), "Cascade.RestartInLevel.Small");
+
+								FToolUIAction UIAction;
+								UIAction.ExecuteAction = FToolMenuExecuteAction::CreateStatic(&UAssetDefinition_MetaSoundSource::ExecutePlaySound);
+								UIAction.CanExecuteAction = FToolMenuCanExecuteAction::CreateStatic(&UAssetDefinition_MetaSoundSource::CanExecutePlayCommand);
+								UIAction.IsActionVisibleDelegate = FToolMenuIsActionButtonVisible::CreateLambda(IsPlayingThis);
+								InSection.AddMenuEntry("Sound_RestartSound", Label, ToolTip, Icon, UIAction);
+							}
 
-							FToolUIAction UIAction;
-							UIAction.ExecuteAction = FToolMenuExecuteAction::CreateStatic(&ExecuteCreateMetaSoundPreset<UMetaSoundSource, UMetaSoundSourceFactory>);
-							InSection.AddMenuEntry("MetaSoundSource_CreatePreset", Label, ToolTip, Icon, UIAction);
-						}
-					}));
+							{ // Stop
+								const FSlateIcon Icon = FSlateIcon(FAppStyle::GetAppStyleSetName(), "MediaAsset.AssetActions.Stop.Small");
+								{ // Selected
+									const TAttribute<FText> StopSelectedToolTip = LOCTEXT("Sound_StopSoundTooltip", "Stops the selected sound.");
+									{
+										const TAttribute<FText> Label = LOCTEXT("Sound_StopSoundDisabled", "Stop");
+
+										FToolUIAction UIAction;
+										UIAction.CanExecuteAction = FToolMenuCanExecuteAction::CreateLambda(IsPlayingThis);
+										UIAction.ExecuteAction = FToolMenuExecuteAction::CreateStatic(&UAssetDefinition_MetaSoundSource::ExecuteStopSound);
+										UIAction.IsActionVisibleDelegate = FToolMenuIsActionButtonVisible::CreateLambda([IsNotPlayingAny, IsPlayingThis](const FToolMenuContext& Context) { return IsNotPlayingAny(Context) || IsPlayingThis(Context); });
+										InSection.AddMenuEntry("Sound_StopSound", Label, StopSelectedToolTip, Icon, UIAction);
+									}
+								}
+								{ // Other
+									const TAttribute<FText> Label = TAttribute<FText>::CreateLambda([]()
+									{
+										if (const USoundBase* OtherSound = UAssetDefinition_SoundBase::GetPlayingSound())
+										{
+											return FText::Format(LOCTEXT("Sound_StopSoundOtherFormat", "Stop ({0})"), FText::FromName(OtherSound->GetFName()));
+										}
+										return LOCTEXT("Sound_StopSoundOther", "Stop (Other)");
+									});
+									const TAttribute<FText> ToolTip = LOCTEXT("Sound_StopOtherSoundTooltip", "Stops the currently previewing (other) sound.");
+									FToolUIAction UIAction;
+									UIAction.ExecuteAction = FToolMenuExecuteAction::CreateStatic(&UAssetDefinition_MetaSoundSource::ExecuteStopSound);
+									UIAction.IsActionVisibleDelegate = FToolMenuIsActionButtonVisible::CreateLambda(IsPlayingOther);
+									InSection.AddMenuEntry("Sound_StopOtherSound", Label, ToolTip, Icon, UIAction);
+								}
+							}
+
+							{
+								const TAttribute<FText> Label = LOCTEXT("Sound_MuteSound", "Mute");
+								const TAttribute<FText> ToolTip = LOCTEXT("Sound_MuteSoundTooltip", "Mutes the selected sounds.");
+								const FSlateIcon Icon = FSlateIcon(FAppStyle::GetAppStyleSetName(), "MediaAsset.AssetActions.Mute.Small");
+				
+								FToolUIAction UIAction;
+								UIAction.ExecuteAction = FToolMenuExecuteAction::CreateStatic(&UAssetDefinition_MetaSoundSource::ExecuteMuteSound);
+								UIAction.CanExecuteAction = FToolMenuCanExecuteAction::CreateStatic(&UAssetDefinition_MetaSoundSource::CanExecuteMuteCommand);
+								UIAction.GetActionCheckState = FToolMenuGetActionCheckState::CreateStatic(&UAssetDefinition_MetaSoundSource::IsActionCheckedMute);
+								InSection.AddMenuEntry("Sound_SoundMute", Label, ToolTip, Icon, UIAction, EUserInterfaceActionType::ToggleButton);
+							}
+							{
+								const TAttribute<FText> Label = LOCTEXT("Sound_SoloSound", "Solo");
+								const TAttribute<FText> ToolTip = LOCTEXT("Sound_SoloSoundTooltip", "Solos the selected sounds.");
+								const FSlateIcon Icon = FSlateIcon(FAppStyle::GetAppStyleSetName(), "MediaAsset.AssetActions.Solo.Small");
+				
+								FToolUIAction UIAction;
+								UIAction.ExecuteAction = FToolMenuExecuteAction::CreateStatic(&UAssetDefinition_MetaSoundSource::ExecuteSoloSound);
+								UIAction.CanExecuteAction = FToolMenuCanExecuteAction::CreateStatic(&UAssetDefinition_MetaSoundSource::CanExecuteSoloCommand);
+								UIAction.GetActionCheckState = FToolMenuGetActionCheckState::CreateStatic(&UAssetDefinition_MetaSoundSource::IsActionCheckedSolo);
+								InSection.AddMenuEntry("Sound_StopSolo", Label, ToolTip, Icon, UIAction, EUserInterfaceActionType::ToggleButton);
+							}
+							{
+								const TAttribute<FText> Label = LOCTEXT("Sound_ClearMutedSoloed", "Clear Muted/Soloed");
+								const TAttribute<FText> ToolTip = LOCTEXT("Sound_ClearMutedSoloedTooltip", "Clear all flags to mute/solo specific assets.");
+								const FSlateIcon Icon = FSlateIcon(FAppStyle::GetAppStyleSetName(), "PropertyWindow.DiffersFromDefault");
+
+								FToolUIAction UIAction;
+								UIAction.ExecuteAction = FToolMenuExecuteAction::CreateStatic(&UAssetDefinition_SoundBase::ExecuteClearMutesAndSolos);
+								UIAction.IsActionVisibleDelegate = FToolMenuCanExecuteAction::CreateStatic(&UAssetDefinition_SoundBase::CanExecuteClearMutesAndSolos);
+								InSection.AddMenuEntry("Sound_ClearMuteSoloSettings", Label, ToolTip, Icon, UIAction);
+							}
+						}));
+					}
 				}
  			}
 	
+			MetaSoundClass = UMetaSoundPatch::StaticClass();
+			check(MetaSoundClass);
+			if (UToolMenu* Menu = UE::ContentBrowser::ExtendToolMenu_AssetContextMenu(MetaSoundClass))
 			{
-				UToolMenu* Menu = UE::ContentBrowser::ExtendToolMenu_AssetContextMenu(UMetaSoundPatch::StaticClass());
-
-				FToolMenuSection& Section = Menu->FindOrAddSection("GetAssetActions");
-				Section.AddDynamicEntry(NAME_None, FNewToolMenuSectionDelegate::CreateLambda([](FToolMenuSection& InSection)
-				{
-					const TAttribute<FText> Label = LOCTEXT("MetaSoundPatch_CreatePreset", "Create MetaSound Patch Preset");
-					const TAttribute<FText> ToolTip = LOCTEXT("MetaSoundPatch_CreatePresetToolTip", "Creates a MetaSoundSource Patch Preset using the selected MetaSound Patch's root graph as a reference.");
-
-					const FSlateIcon Icon = Metasound::Editor::Style::CreateSlateIcon("ClassIcon.MetasoundPatch");
-
-					FToolUIAction UIAction;
-					UIAction.ExecuteAction = FToolMenuExecuteAction::CreateStatic(&ExecuteCreateMetaSoundPreset<UMetaSoundPatch, UMetaSoundFactory>);
-					InSection.AddMenuEntry("MetaSoundPatch_CreatePreset", Label, ToolTip, Icon, UIAction);
-				}));
+				FToolUIAction CreateAction;
+				CreateAction.ExecuteAction = FToolMenuExecuteAction::CreateStatic(&AssetDefinitionsPrivate::ExecuteCreateMetaSoundPreset<UMetaSoundPatch, UMetaSoundFactory>);
+				const FSlateIcon AssetIcon = Metasound::Editor::Style::CreateSlateIcon("ClassIcon.MetasoundPatch");
+				AssetDefinitionsPrivate::AddMetaSoundActions(*Menu, AssetIcon, MetaSoundClass->GetDisplayNameText(), MoveTemp(CreateAction));
 			}
  		}));
 	});
 }
-
-
 #undef LOCTEXT_NAMESPACE //MetaSoundEditor
diff --git a/Engine/Source/Editor/AudioEditor/Private/AssetTypeActions/AssetDefinition_SoundBase.cpp b/Engine/Source/Editor/AudioEditor/Private/AssetTypeActions/AssetDefinition_SoundBase.cpp
--- a/Engine/Source/Editor/AudioEditor/Private/AssetTypeActions/AssetDefinition_SoundBase.cpp
+++ b/Engine/Source/Editor/AudioEditor/Private/AssetTypeActions/AssetDefinition_SoundBase.cpp
@@ -361,6 +361,36 @@ bool UAssetDefinition_SoundBase::GetThumbnailActionOverlay(const FAssetData& InA
 #endif
 	}
 
+	bool UAssetDefinition_SoundBase::CanExecuteClearMutesAndSolos(const FToolMenuContext& InContext)
+	{
+#if ENABLE_AUDIO_DEBUG
+		if (FAudioDeviceManager* ADM = GEditor->GetAudioDeviceManager())
+		{
+			if (const UContentBrowserAssetContextMenuContext* CBContext = UContentBrowserAssetContextMenuContext::FindContextWithAssets(InContext))
+			{
+				Audio::FAudioDebugger& Debugger = ADM->GetDebugger();
+				return Debugger.IsMuteOrSoloActive();
+			}
+		}
+#endif // ENABLE_AUDIO_DEBUG
+
+		return false;
+	}
+
+	void UAssetDefinition_SoundBase::ExecuteClearMutesAndSolos(const FToolMenuContext& InContext)
+	{
+#if ENABLE_AUDIO_DEBUG
+		if (FAudioDeviceManager* ADM = GEditor->GetAudioDeviceManager())
+		{
+			if (const UContentBrowserAssetContextMenuContext* CBContext = UContentBrowserAssetContextMenuContext::FindContextWithAssets(InContext))
+			{
+				Audio::FAudioDebugger& Debugger = ADM->GetDebugger();
+				Debugger.ClearMutesAndSolos();
+			}
+		}
+#endif // ENABLE_AUDIO_DEBUG
+	}
+
 	void UAssetDefinition_SoundBase::ExecuteSoloSound(const FToolMenuContext& InContext)
 	{
 #if ENABLE_AUDIO_DEBUG
@@ -440,7 +470,7 @@ bool UAssetDefinition_SoundBase::GetThumbnailActionOverlay(const FAssetData& InA
 						return false;
 					}
 				}
-            
+
 				for (const FAssetData& SoundWaveAsset : CBContext->GetSelectedAssetsOfType(USoundWave::StaticClass()))
 				{
 					if (Debugger.IsMuteSoundWave(SoundWaveAsset.AssetName))
@@ -457,6 +487,51 @@ bool UAssetDefinition_SoundBase::GetThumbnailActionOverlay(const FAssetData& InA
 		return false;
 	}
 
+	bool UAssetDefinition_SoundBase::IsPlayingContextAsset(const FToolMenuContext& InContext, bool bMustMatchContext)
+	{
+		const UAudioComponent* PreviewComponent = GEditor->GetPreviewAudioComponent();
+		if (!PreviewComponent || !PreviewComponent->IsPlaying())
+		{
+			return false;
+		}
+
+		if (!bMustMatchContext)
+		{
+			return true;
+		}
+
+		if (const UContentBrowserAssetContextMenuContext* CBContext = UContentBrowserAssetContextMenuContext::FindContextWithAssets(InContext))
+		{
+			if (CBContext->SelectedAssets.Num() == 1)
+			{
+				if (const USoundBase* PreviewSound = PreviewComponent->Sound)
+				{
+					const FAssetData AssetData = CBContext->SelectedAssets.Last();
+					if (AssetData.IsAssetLoaded())
+					{
+						if (UObject* ContextObject = AssetData.GetAsset())
+						{
+							return PreviewSound->GetUniqueID() == ContextObject->GetUniqueID();
+						}
+					}
+				}
+			}
+		}
+
+		return false;
+	}
+
+	const USoundBase* UAssetDefinition_SoundBase::GetPlayingSound()
+	{
+		const UAudioComponent* PreviewComponent = GEditor->GetPreviewAudioComponent();
+		if (!PreviewComponent || !PreviewComponent->IsPlaying())
+		{
+			return nullptr;
+		}
+
+		return PreviewComponent->Sound.Get();
+	}
+
 	void UAssetDefinition_SoundBase::GetSoundBaseAssetActionButtonExtensions(const FAssetData& InAssetData, TArray<FAssetButtonActionExtension>& OutExtensions)
 	{
 		FAssetButtonActionExtension AssetButtonActionExtension
@@ -504,61 +579,110 @@ namespace MenuExtension_SoundBase
 	{
 		UToolMenu* Menu = UE::ContentBrowser::ExtendToolMenu_AssetContextMenu(InClass);
 
-		FToolMenuSection& Section = Menu->FindOrAddSection("GetAssetActions");
+		FToolMenuSection& Section = Menu->FindOrAddSection("AudioPlaybackActions", LOCTEXT("PlaybackActions_Label", "Playback"), FToolMenuInsert(FName{ }, EToolMenuInsertType::Last));
 		Section.AddDynamicEntry(NAME_None, FNewToolMenuSectionDelegate::CreateLambda([](FToolMenuSection& InSection)
+		{
+			if (const UContentBrowserAssetContextMenuContext* Context = UContentBrowserAssetContextMenuContext::FindContextWithAssets(InSection))
 			{
-				if (const UContentBrowserAssetContextMenuContext* Context = UContentBrowserAssetContextMenuContext::FindContextWithAssets(InSection))
+				if (Context->SelectedAssets.Num() > 0)
 				{
-					if (Context->SelectedAssets.Num() > 0)
+					const FAssetData& AssetData = Context->SelectedAssets[0];
+					if (AssetData.AssetClassPath == USoundWave::StaticClass()->GetClassPathName() || AssetData.AssetClassPath == USoundCue::StaticClass()->GetClassPathName())
 					{
-						const FAssetData& AssetData = Context->SelectedAssets[0];
-						if (AssetData.AssetClassPath == USoundWave::StaticClass()->GetClassPathName() || AssetData.AssetClassPath == USoundCue::StaticClass()->GetClassPathName())
-						{
-							{
-								const TAttribute<FText> Label = LOCTEXT("Sound_PlaySound", "Play");
-								const TAttribute<FText> ToolTip = LOCTEXT("Sound_PlaySoundTooltip", "Plays the selected sound.");
-								const FSlateIcon Icon = FSlateIcon(FAppStyle::GetAppStyleSetName(), "MediaAsset.AssetActions.Play.Small");
+						auto IsPlayingThis = [](const FToolMenuContext& InContext) { return UAssetDefinition_SoundBase::IsPlayingContextAsset(InContext, true /* bMustMatchContext */); };
+						auto IsNotPlayingThis = [](const FToolMenuContext& InContext) { return !UAssetDefinition_SoundBase::IsPlayingContextAsset(InContext, true /* bMustMatchContext */); };
+						auto IsNotPlayingAny = [](const FToolMenuContext& InContext) { return !UAssetDefinition_SoundBase::IsPlayingContextAsset(InContext, false /* bMustMatchContext */); };
+						auto IsPlayingAny = [](const FToolMenuContext& InContext) { return UAssetDefinition_SoundBase::IsPlayingContextAsset(InContext, false /* bMustMatchContext */); };
+						auto IsPlayingOther = [](const FToolMenuContext& InContext) { return UAssetDefinition_SoundBase::IsPlayingContextAsset(InContext, false /* bMustMatchContext */) && !UAssetDefinition_SoundBase::IsPlayingContextAsset(InContext, true /* bMustMatchContext */); };
 
-								FToolUIAction UIAction;
-								UIAction.ExecuteAction = FToolMenuExecuteAction::CreateStatic(&UAssetDefinition_SoundBase::ExecutePlaySound);
-								UIAction.CanExecuteAction = FToolMenuCanExecuteAction::CreateStatic(&UAssetDefinition_SoundBase::CanExecutePlayCommand);
-								InSection.AddMenuEntry("Sound_PlaySound", Label, ToolTip, Icon, UIAction);
+						{
+							const TAttribute<FText> Label = LOCTEXT("Sound_PlaySound", "Play");
+							const TAttribute<FText> ToolTip = LOCTEXT("Sound_PlaySoundTooltip", "Plays the selected sound.");
+							const FSlateIcon Icon = FSlateIcon(FAppStyle::GetAppStyleSetName(), "MediaAsset.AssetActions.Play.Small");
+
+							FToolUIAction UIAction;
+							UIAction.ExecuteAction = FToolMenuExecuteAction::CreateStatic(&UAssetDefinition_SoundBase::ExecutePlaySound);
+							UIAction.CanExecuteAction = FToolMenuCanExecuteAction::CreateStatic(&UAssetDefinition_SoundBase::CanExecutePlayCommand);
+							UIAction.IsActionVisibleDelegate = FToolMenuIsActionButtonVisible::CreateLambda(IsNotPlayingThis);
+							InSection.AddMenuEntry("Sound_PlaySound", Label, ToolTip, Icon, UIAction);
+						}
+						{
+							const TAttribute<FText> Label = LOCTEXT("Sound_RestartSound", "Restart");
+							const TAttribute<FText> ToolTip = LOCTEXT("Sound_RestartSoundTooltip", "Restarts the selected sound.");
+							const FSlateIcon Icon = FSlateIcon(FAppStyle::GetAppStyleSetName(), "Cascade.RestartInLevel.Small");
+
+							FToolUIAction UIAction;
+							UIAction.ExecuteAction = FToolMenuExecuteAction::CreateStatic(&UAssetDefinition_SoundBase::ExecutePlaySound);
+							UIAction.CanExecuteAction = FToolMenuCanExecuteAction::CreateStatic(&UAssetDefinition_SoundBase::CanExecutePlayCommand);
+							UIAction.IsActionVisibleDelegate = FToolMenuIsActionButtonVisible::CreateLambda(IsPlayingThis);
+							InSection.AddMenuEntry("Sound_RestartSound", Label, ToolTip, Icon, UIAction);
+						}
+						{ // Stop
+							const FSlateIcon Icon = FSlateIcon(FAppStyle::GetAppStyleSetName(), "MediaAsset.AssetActions.Stop.Small");
+							{ // Selected
+								const TAttribute<FText> StopSelectedToolTip = LOCTEXT("Sound_StopSoundTooltip", "Stops the selected sound.");
+								{
+									const TAttribute<FText> Label = LOCTEXT("Sound_StopSoundDisabled", "Stop");
+
+									FToolUIAction UIAction;
+									UIAction.CanExecuteAction = FToolMenuCanExecuteAction::CreateLambda(IsPlayingThis);
+									UIAction.ExecuteAction = FToolMenuExecuteAction::CreateStatic(&UAssetDefinition_SoundBase::ExecuteStopSound);
+									UIAction.IsActionVisibleDelegate = FToolMenuIsActionButtonVisible::CreateLambda([IsNotPlayingAny, IsPlayingThis](const FToolMenuContext& Context) { return IsNotPlayingAny(Context) || IsPlayingThis(Context); });
+									InSection.AddMenuEntry("Sound_StopSound", Label, StopSelectedToolTip, Icon, UIAction);
+								}
 							}
-							{
-								const TAttribute<FText> Label = LOCTEXT("Sound_StopSound", "Stop");
-								const TAttribute<FText> ToolTip = LOCTEXT("Sound_StopSoundTooltip", "Stops the selected sounds.");
-								const FSlateIcon Icon = FSlateIcon(FAppStyle::GetAppStyleSetName(), "MediaAsset.AssetActions.Stop.Small");
-
+							{ // Other
+								const TAttribute<FText> Label = TAttribute<FText>::CreateLambda([]()
+								{
+									if (const USoundBase* OtherSound = UAssetDefinition_SoundBase::GetPlayingSound())
+									{
+										return FText::Format(LOCTEXT("Sound_StopSoundOtherFormat", "Stop ({0})"), FText::FromName(OtherSound->GetFName()));
+									}
+									return LOCTEXT("Sound_StopSoundOther", "Stop (Other)");
+								});
+								const TAttribute<FText> ToolTip = LOCTEXT("Sound_StopOtherSoundTooltip", "Stops the currently previewing (other) sound.");
 								FToolUIAction UIAction;
 								UIAction.ExecuteAction = FToolMenuExecuteAction::CreateStatic(&UAssetDefinition_SoundBase::ExecuteStopSound);
-								InSection.AddMenuEntry("Sound_StopSound", Label, ToolTip, Icon, UIAction);
-							}
-							{
-								const TAttribute<FText> Label = LOCTEXT("Sound_MuteSound", "Mute");
-								const TAttribute<FText> ToolTip = LOCTEXT("Sound_MuteSoundTooltip", "Mutes the selected sounds.");
-								const FSlateIcon Icon = FSlateIcon(FAppStyle::GetAppStyleSetName(), "MediaAsset.AssetActions.Mute.Small");
-
-								FToolUIAction UIAction;
-								UIAction.ExecuteAction = FToolMenuExecuteAction::CreateStatic(&UAssetDefinition_SoundBase::ExecuteMuteSound);
-								UIAction.CanExecuteAction = FToolMenuCanExecuteAction::CreateStatic(&UAssetDefinition_SoundBase::CanExecuteMuteCommand);
-								UIAction.GetActionCheckState = FToolMenuGetActionCheckState::CreateStatic(&UAssetDefinition_SoundBase::IsActionCheckedMute);
-								InSection.AddMenuEntry("Sound_SoundMute", Label, ToolTip, Icon, UIAction, EUserInterfaceActionType::ToggleButton);
-							}
-							{
-								const TAttribute<FText> Label = LOCTEXT("Sound_SoloSound", "Solo");
-								const TAttribute<FText> ToolTip = LOCTEXT("Sound_SoloSoundTooltip", "Solos the selected sounds.");
-								const FSlateIcon Icon = FSlateIcon(FAppStyle::GetAppStyleSetName(), "MediaAsset.AssetActions.Solo.Small");
-
-								FToolUIAction UIAction;
-								UIAction.ExecuteAction = FToolMenuExecuteAction::CreateStatic(&UAssetDefinition_SoundBase::ExecuteSoloSound);
-								UIAction.CanExecuteAction = FToolMenuCanExecuteAction::CreateStatic(&UAssetDefinition_SoundBase::CanExecuteSoloCommand);
-								UIAction.GetActionCheckState = FToolMenuGetActionCheckState::CreateStatic(&UAssetDefinition_SoundBase::IsActionCheckedSolo);
-								InSection.AddMenuEntry("Sound_StopSolo", Label, ToolTip, Icon, UIAction, EUserInterfaceActionType::ToggleButton);
+								UIAction.IsActionVisibleDelegate = FToolMenuIsActionButtonVisible::CreateLambda(IsPlayingOther);
+								InSection.AddMenuEntry("Sound_StopOtherSound", Label, ToolTip, Icon, UIAction);
 							}
 						}
+						{
+							const TAttribute<FText> Label = LOCTEXT("Sound_MuteSound", "Mute");
+							const TAttribute<FText> ToolTip = LOCTEXT("Sound_MuteSoundTooltip", "Mutes the selected sounds.");
+							const FSlateIcon Icon = FSlateIcon(FAppStyle::GetAppStyleSetName(), "MediaAsset.AssetActions.Mute.Small");
+
+							FToolUIAction UIAction;
+							UIAction.ExecuteAction = FToolMenuExecuteAction::CreateStatic(&UAssetDefinition_SoundBase::ExecuteMuteSound);
+							UIAction.CanExecuteAction = FToolMenuCanExecuteAction::CreateStatic(&UAssetDefinition_SoundBase::CanExecuteMuteCommand);
+							UIAction.GetActionCheckState = FToolMenuGetActionCheckState::CreateStatic(&UAssetDefinition_SoundBase::IsActionCheckedMute);
+							InSection.AddMenuEntry("Sound_SoundMute", Label, ToolTip, Icon, UIAction, EUserInterfaceActionType::ToggleButton);
+						}
+						{
+							const TAttribute<FText> Label = LOCTEXT("Sound_SoloSound", "Solo");
+							const TAttribute<FText> ToolTip = LOCTEXT("Sound_SoloSoundTooltip", "Solos the selected sounds.");
+							const FSlateIcon Icon = FSlateIcon(FAppStyle::GetAppStyleSetName(), "MediaAsset.AssetActions.Solo.Small");
+
+							FToolUIAction UIAction;
+							UIAction.ExecuteAction = FToolMenuExecuteAction::CreateStatic(&UAssetDefinition_SoundBase::ExecuteSoloSound);
+							UIAction.CanExecuteAction = FToolMenuCanExecuteAction::CreateStatic(&UAssetDefinition_SoundBase::CanExecuteSoloCommand);
+							UIAction.GetActionCheckState = FToolMenuGetActionCheckState::CreateStatic(&UAssetDefinition_SoundBase::IsActionCheckedSolo);
+							InSection.AddMenuEntry("Sound_StopSolo", Label, ToolTip, Icon, UIAction, EUserInterfaceActionType::ToggleButton);
+						}
+						{
+							const TAttribute<FText> Label = LOCTEXT("Sound_ClearMutedSoloed", "Clear Muted/Soloed");
+							const TAttribute<FText> ToolTip = LOCTEXT("Sound_ClearMutedSoloedTooltip", "Clear all flags to mute/solo specific assets.");
+							const FSlateIcon Icon = FSlateIcon(FAppStyle::GetAppStyleSetName(), "PropertyWindow.DiffersFromDefault");
+
+							FToolUIAction UIAction;
+							UIAction.ExecuteAction = FToolMenuExecuteAction::CreateStatic(&UAssetDefinition_SoundBase::ExecuteClearMutesAndSolos);
+							UIAction.IsActionVisibleDelegate = FToolMenuCanExecuteAction::CreateStatic(&UAssetDefinition_SoundBase::CanExecuteClearMutesAndSolos);
+							InSection.AddMenuEntry("Sound_ClearMuteSoloSettings", Label, ToolTip, Icon, UIAction);
+						}
 					}
 				}
-			}));
+			}
+		}));
 	}
 
 	
diff --git a/Engine/Source/Editor/AudioEditor/Private/AssetTypeActions/AssetDefinition_SoundWave.cpp b/Engine/Source/Editor/AudioEditor/Private/AssetTypeActions/AssetDefinition_SoundWave.cpp
--- a/Engine/Source/Editor/AudioEditor/Private/AssetTypeActions/AssetDefinition_SoundWave.cpp
+++ b/Engine/Source/Editor/AudioEditor/Private/AssetTypeActions/AssetDefinition_SoundWave.cpp
@@ -244,61 +244,63 @@ namespace MenuExtension_SoundWave
 		{
 			FToolMenuOwnerScoped OwnerScoped(UE_MODULE_NAME);
 			{
-				UToolMenu* Menu = UE::ContentBrowser::ExtendToolMenu_AssetContextMenu(USoundWave::StaticClass());		
-				FToolMenuSection& Section = Menu->FindOrAddSection("GetAssetActions");
-				Section.AddDynamicEntry(NAME_None, FNewToolMenuSectionDelegate::CreateLambda([](FToolMenuSection& InSection)
+				if (UToolMenu* Menu = UE::ContentBrowser::ExtendToolMenu_AssetContextMenu(USoundWave::StaticClass()))
 				{
-					if (const UContentBrowserAssetContextMenuContext* Context = UContentBrowserAssetContextMenuContext::FindContextWithAssets(InSection))
+					FToolMenuSection& Section = Menu->FindOrAddSection("SoundActions", LOCTEXT("SoundActions_Label", "Sound"));
+					Section.AddDynamicEntry(NAME_None, FNewToolMenuSectionDelegate::CreateLambda([](FToolMenuSection& InSection)
 					{
-						TArray<TSoftObjectPtr<USoundWave>> SoundNodes = Context->GetSelectedAssetSoftObjects<USoundWave>();
-						bool bCreateCueForEachSoundWave = true;
-
-						if (Context->SelectedAssets.Num() == 1)
-						{
-							const TAttribute<FText> Label = LOCTEXT("SoundWave_CreateCue", "Create Cue");
-							const TAttribute<FText> ToolTip = LOCTEXT("SoundWave_CreateCueTooltip", "Creates a sound cue using this sound wave.");
-							const FSlateIcon Icon = FSlateIcon(FAppStyle::GetAppStyleSetName(), "ClassIcon.SoundCue");
-
-							FToolUIAction UIAction;
-							UIAction.ExecuteAction = FToolMenuExecuteAction::CreateStatic(&ExecuteCreateSoundCue, SoundNodes, bCreateCueForEachSoundWave);
-							InSection.AddMenuEntry("SoundWave_CreateCue", Label, ToolTip, Icon, UIAction);
-						}
-						else
+						if (const UContentBrowserAssetContextMenuContext* Context = UContentBrowserAssetContextMenuContext::FindContextWithAssets(InSection))
 						{
-							bCreateCueForEachSoundWave = false;
-
-							const TAttribute<FText> Label_Single = LOCTEXT("SoundWave_CreateSingleCue", "Create Single Cue");
-							const TAttribute<FText> ToolTip_Single = LOCTEXT("SoundWave_CreateSingleCueTooltip", "Creates a single sound cue using these sound waves.");
-							const FSlateIcon Icon_Single = FSlateIcon(FAppStyle::GetAppStyleSetName(), "ClassIcon.SoundCue");
-
-							FToolUIAction UIAction_Single;
-							UIAction_Single.ExecuteAction = FToolMenuExecuteAction::CreateStatic(&ExecuteCreateSoundCue, SoundNodes, bCreateCueForEachSoundWave);
-							InSection.AddMenuEntry("SoundWave_CreateSingleCue", Label_Single, ToolTip_Single, Icon_Single, UIAction_Single);
-
-							bCreateCueForEachSoundWave = true;
-
-							const TAttribute<FText> Label_Multi = LOCTEXT("SoundWave_CreateMultiCue", "Create Multiple Cues");
-							const TAttribute<FText> ToolTip_Multi = LOCTEXT("SoundWave_CreateMultiCueTooltip", "Creates multiple sound cues, one from each selected sound wave.");
-							const FSlateIcon Icon_Multi = FSlateIcon(FAppStyle::GetAppStyleSetName(), "ClassIcon.SoundCue");
-
-							FToolUIAction UIAction_Multi;
-							UIAction_Multi.ExecuteAction = FToolMenuExecuteAction::CreateStatic(&ExecuteCreateSoundCue, SoundNodes, bCreateCueForEachSoundWave);
-							InSection.AddMenuEntry("SoundWave_CreateMultiCue", Label_Multi, ToolTip_Multi, Icon_Multi, UIAction_Multi);
+							TArray<TSoftObjectPtr<USoundWave>> SoundNodes = Context->GetSelectedAssetSoftObjects<USoundWave>();
+							bool bCreateCueForEachSoundWave = true;
+				
+							if (Context->SelectedAssets.Num() == 1)
+							{
+								const TAttribute<FText> Label = LOCTEXT("SoundWave_CreateSoundCue", "Create SoundCue");
+								const TAttribute<FText> ToolTip = LOCTEXT("SoundWave_CreateCueTooltip", "Creates a SoundCue referencing the selected SoundWave.");
+								const FSlateIcon Icon = FSlateIcon(FAppStyle::GetAppStyleSetName(), "ClassIcon.SoundCue");
+				
+								FToolUIAction UIAction;
+								UIAction.ExecuteAction = FToolMenuExecuteAction::CreateStatic(&ExecuteCreateSoundCue, SoundNodes, bCreateCueForEachSoundWave);
+								InSection.AddMenuEntry("SoundWave_CreateCue", Label, ToolTip, Icon, UIAction);
+							}
+							else
+							{
+								bCreateCueForEachSoundWave = false;
+				
+								const TAttribute<FText> Label_Single = LOCTEXT("SoundWave_CreateSingleSoundCue", "Create SoundCue (Single)");
+								const TAttribute<FText> ToolTip_Single = LOCTEXT("SoundWave_CreateSingleCueTooltip", "Creates a single SoundCue Referencing the selected SoundWaves.");
+								const FSlateIcon Icon_Single = FSlateIcon(FAppStyle::GetAppStyleSetName(), "ClassIcon.SoundCue");
+				
+								FToolUIAction UIAction_Single;
+								UIAction_Single.ExecuteAction = FToolMenuExecuteAction::CreateStatic(&ExecuteCreateSoundCue, SoundNodes, bCreateCueForEachSoundWave);
+								InSection.AddMenuEntry("SoundWave_CreateSingleCue", Label_Single, ToolTip_Single, Icon_Single, UIAction_Single);
+				
+								bCreateCueForEachSoundWave = true;
+				
+								const TAttribute<FText> Label_Multi = LOCTEXT("SoundWave_CreateMultiCue", "Create SoundCues (Multiple)");
+								const TAttribute<FText> ToolTip_Multi = LOCTEXT("SoundWave_CreateMultiCueTooltip", "Creates multiple SoundCues, one referencing each selected SoundWave.");
+								const FSlateIcon Icon_Multi = FSlateIcon(FAppStyle::GetAppStyleSetName(), "ClassIcon.SoundCue");
+				
+								FToolUIAction UIAction_Multi;
+								UIAction_Multi.ExecuteAction = FToolMenuExecuteAction::CreateStatic(&ExecuteCreateSoundCue, SoundNodes, bCreateCueForEachSoundWave);
+								InSection.AddMenuEntry("SoundWave_CreateMultiCue", Label_Multi, ToolTip_Multi, Icon_Multi, UIAction_Multi);
+							}
+				
+							const TAttribute<FText> Label_Dialogue = LOCTEXT("SoundWave_CreateDialogue", "Create Dialogue");
+							const TAttribute<FText> ToolTip_Dialogue = LOCTEXT("SoundWave_CreateDialogueTooltip", "Creates a DialogueWave referencing the selected SoundWave.");
+				
+							InSection.AddSubMenu(
+								"SoundWave_CreateDialogue",
+								Label_Dialogue,
+								ToolTip_Dialogue,
+								FNewToolMenuDelegate::CreateStatic(&FillVoiceMenu, SoundNodes),
+								false,
+								FSlateIcon(FAppStyle::GetAppStyleSetName(), "ClassIcon.DialogueWave")
+							);
 						}
-
-						const TAttribute<FText> Label_Dialogue = LOCTEXT("SoundWave_CreateDialogue", "Create Dialogue");
-						const TAttribute<FText> ToolTip_Dialogue = LOCTEXT("SoundWave_CreateDialogueTooltip", "Creates a dialogue wave using this sound wave.");
-
-						InSection.AddSubMenu(
-							"SoundWave_CreateDialogue",
-							Label_Dialogue,
-							ToolTip_Dialogue,
-							FNewToolMenuDelegate::CreateStatic(&FillVoiceMenu, SoundNodes),
-							false,
-							FSlateIcon()
-						);
-					}
-				}));
+					}));
+				}
 			}
 		}));
 	});
diff --git a/Engine/Source/Editor/AudioEditor/Public/AssetTypeActions/AssetDefinition_SoundBase.h b/Engine/Source/Editor/AudioEditor/Public/AssetTypeActions/AssetDefinition_SoundBase.h
--- a/Engine/Source/Editor/AudioEditor/Public/AssetTypeActions/AssetDefinition_SoundBase.h
+++ b/Engine/Source/Editor/AudioEditor/Public/AssetTypeActions/AssetDefinition_SoundBase.h
@@ -48,8 +48,12 @@ public:
 	static ECheckBoxState IsActionCheckedSolo(const FToolMenuContext& InContext);
 	static void ExecuteMuteSound(const FToolMenuContext& InContext);
 	static void ExecuteSoloSound(const FToolMenuContext& InContext);
+	static void ExecuteClearMutesAndSolos(const FToolMenuContext& InContext);
 	static bool CanExecuteMuteCommand(const FToolMenuContext& InContext);
 	static bool CanExecuteSoloCommand(const FToolMenuContext& InContext);
+	static bool CanExecuteClearMutesAndSolos(const FToolMenuContext& InContext);
+	static bool IsPlayingContextAsset(const FToolMenuContext& InContext, bool bMustMatchContext);
+	static const USoundBase* GetPlayingSound();
 
 	// Asset definition static utilities
 	static TSharedPtr<SWidget> GetSoundBaseThumbnailOverlay(const FAssetData& InAssetData, TFunction<FReply()>&& OnClicked);
diff --git a/Engine/Source/Runtime/Engine/Classes/Sound/AudioSettings.h b/Engine/Source/Runtime/Engine/Classes/Sound/AudioSettings.h
--- a/Engine/Source/Runtime/Engine/Classes/Sound/AudioSettings.h
+++ b/Engine/Source/Runtime/Engine/Classes/Sound/AudioSettings.h
@@ -260,6 +260,12 @@ class UAudioSettings : public UDeveloperSettings
 	UPROPERTY(config, EditAnywhere, Category="Mix")
 	TArray<FDefaultAudioBusSettings> DefaultAudioBuses;
 
+	/** Whether or not project supports the creation of legacy audio asset types. Currently includes:
+	  * - Simple Sounds
+	  */
+	UPROPERTY(config, EditAnywhere, Category="Audio", AdvancedDisplay)
+	bool bEnableLegacyAssetTypes = false;
+
 #if WITH_EDITOR
 	FAudioSettingsChanged AudioSettingsChanged;
 #endif // WITH_EDITOR
diff --git a/Engine/Source/Runtime/Engine/Private/AudioDebug.cpp b/Engine/Source/Runtime/Engine/Private/AudioDebug.cpp
--- a/Engine/Source/Runtime/Engine/Private/AudioDebug.cpp
+++ b/Engine/Source/Runtime/Engine/Private/AudioDebug.cpp
@@ -2282,10 +2282,10 @@ namespace Audio
 		
 		Stats.RequestedStats.Append(StatsToSet);
 	}
-		
 
 
-		
+
+
 
 	bool FAudioDebugger::PostStatModulatorHelp(UWorld* World, FCommonViewportClient* ViewportClient, const TCHAR* Stream)
 	{
@@ -2474,6 +2474,16 @@ namespace Audio
 		DebugNames.SoloSoundWave.Empty();
 	}
 
+	bool FAudioDebugger::IsMuteOrSoloActive() const
+	{
+		return !DebugNames.MuteSoundClass.IsEmpty()
+			|| !DebugNames.MuteSoundCue.IsEmpty()
+			|| !DebugNames.MuteSoundWave.IsEmpty()
+			|| !DebugNames.SoloSoundClass.IsEmpty()
+			|| !DebugNames.SoloSoundCue.IsEmpty()
+			|| !DebugNames.SoloSoundWave.IsEmpty();
+	}
+
 	void FAudioDebugger::LogSubtitle(const TCHAR* InCmd, USoundWave& InSoundWave)
 	{
 		const bool bLogSubtitle = FParse::Param(InCmd, TEXT("LogSubtitle"));
diff --git a/Engine/Source/Runtime/Engine/Public/Audio/AudioDebug.h b/Engine/Source/Runtime/Engine/Public/Audio/AudioDebug.h
--- a/Engine/Source/Runtime/Engine/Public/Audio/AudioDebug.h
+++ b/Engine/Source/Runtime/Engine/Public/Audio/AudioDebug.h
@@ -114,6 +114,7 @@ namespace Audio
 		static ENGINE_API bool IsVirtualLoopVisualizeEnabled();
 
 		ENGINE_API void ClearMutesAndSolos();
+		ENGINE_API bool IsMuteOrSoloActive() const;
 		ENGINE_API void DumpActiveSounds() const;
 
 		ENGINE_API bool IsVisualizeDebug3dEnabled() const;
