Truncate a Metasound variable's Display name to the size limit required of an FName (1024 chars), fixing an editor crash.

diff --git a/Engine/Plugins/Runtime/Metasound/Source/MetasoundEditor/Private/MetasoundNodeDetailCustomization.h b/Engine/Plugins/Runtime/Metasound/Source/MetasoundEditor/Private/MetasoundNodeDetailCustomization.h
--- a/Engine/Plugins/Runtime/Metasound/Source/MetasoundEditor/Private/MetasoundNodeDetailCustomization.h
+++ b/Engine/Plugins/Runtime/Metasound/Source/MetasoundEditor/Private/MetasoundNodeDetailCustomization.h
@@ -154,7 +154,10 @@ namespace Metasound
 
 				if (GraphMember.IsValid())
 				{
-					GraphMember->SetDisplayName(InText, true);
+					// At future points in processing, the display name may be used as a member name and needs to match the length requirements of an FName.
+					// The text edit box used does not have a public function to display an error and it would be difficult to notify to the user, thus 
+					// we are just silently truncating the DisplayName text here. 
+					GraphMember->SetDisplayName(FText::FromString(InText.ToString().Left(NAME_SIZE - 1)), true);
 				}
 			}
 		};
