Early out if test node fails to create operator to avoid null unique ptr deref
#rnx

diff --git a/Engine/Plugins/Runtime/Metasound/Source/MetasoundEngineTest/Private/EngineTestMetaSoundAutomatedNodeTest_ValidOutputs.cpp b/Engine/Plugins/Runtime/Metasound/Source/MetasoundEngineTest/Private/EngineTestMetaSoundAutomatedNodeTest_ValidOutputs.cpp
--- a/Engine/Plugins/Runtime/Metasound/Source/MetasoundEngineTest/Private/EngineTestMetaSoundAutomatedNodeTest_ValidOutputs.cpp
+++ b/Engine/Plugins/Runtime/Metasound/Source/MetasoundEngineTest/Private/EngineTestMetaSoundAutomatedNodeTest_ValidOutputs.cpp
@@ -84,6 +84,7 @@ bool FMetasoundAutomatedNodeTest_ValidOutputs::RunTest(const FString& InRegistry
 		if (!Operator.IsValid())
 		{
 			AddError(FString::Printf(TEXT("Failed to create operator from node %s - %s."), *InRegistryKeyString, *GetPrettyName(RegistryKey)));
+			return;
 		}
 
 
