Fix for crash in some assets with potentially corrupt connection data when migrating editor graph data from older MetaSound 1.12 doc versions

#rnx

diff --git a/Engine/Plugins/Runtime/Metasound/Source/MetasoundEditor/Private/MetasoundEditorGraph.cpp b/Engine/Plugins/Runtime/Metasound/Source/MetasoundEditor/Private/MetasoundEditorGraph.cpp
--- a/Engine/Plugins/Runtime/Metasound/Source/MetasoundEditor/Private/MetasoundEditorGraph.cpp
+++ b/Engine/Plugins/Runtime/Metasound/Source/MetasoundEditor/Private/MetasoundEditorGraph.cpp
@@ -2712,19 +2712,38 @@ void UMetasoundEditorGraph::MigrateEditorDocumentData(FMetaSoundFrontendDocument
 								// Swap connection from input node to connected node to now be from template node to connected node
 								EdgeToRemove.ToNodeID = ConnectedNodeID,
 								EdgeToRemove.ToVertexID = ConnectedInput->VertexID;
-								OutBuilder.RemoveEdge(EdgeToRemove);
-
-								OutBuilder.AddEdge(FMetasoundFrontendEdge
+								bool bEdgeRemoved = OutBuilder.RemoveEdge(EdgeToRemove);
+								if (!bEdgeRemoved)
+								{
+									UE_LOG(LogMetaSound, Display, TEXT("Editor graph '%s' migration failed to remove exact node '%s' connection to class output '%s': Removing any existing connections."),
+										*ConnectedNode->GetDisplayName().ToString(),
+										*Linked->GetName(),
+										*InputName.ToString());
+									bEdgeRemoved = OutBuilder.RemoveEdgeToNodeInput(ConnectedNodeID, ConnectedInput->VertexID);
+								}
+
+								if (bEdgeRemoved)
+								{
+									OutBuilder.AddEdge(FMetasoundFrontendEdge
+									{
+										TemplateNodeID,
+										TemplateOutputID,
+										ConnectedNodeID,
+										ConnectedInput->VertexID
+									});
+								}
+								else
 								{
-									TemplateNodeID,
-									TemplateOutputID,
-									ConnectedNodeID,
-									ConnectedInput->VertexID
-								});
+									UE_LOG(LogMetaSound, Display, TEXT("Editor graph '%s' migration failed to remove connected for node '%s' class output '%s': Ignoring connection upgrade from input '%s' "),
+										*OutBuilder.GetDebugName(),
+										*ConnectedNode->GetDisplayName().ToString(),
+										*Linked->GetName(),
+										*InputName.ToString());
+								}
 							}
 							else
 							{
-								UE_LOG(LogMetaSound, Warning, TEXT("Editor graph '%s' migration failed to find node '%s' class output '%s': Ignoring connection upgrade from input '%s'"),
+								UE_LOG(LogMetaSound, Display, TEXT("Editor graph '%s' migration failed to find node '%s' class output '%s': Ignoring connection upgrade from input '%s'"),
 									*OutBuilder.GetDebugName(),
 									*ConnectedNode->GetDisplayName().ToString(),
 									*Linked->GetName(),
