Ensure that the MetasoundGraph is registered fore execution when creating the ProxyData so that we can properly retrieve the Graph. Also cache the Proxy Data after creation.

#rnx

diff --git a/Engine/Plugins/Runtime/Metasound/Source/MetasoundFrontend/Private/MetasoundAssetBase.cpp b/Engine/Plugins/Runtime/Metasound/Source/MetasoundFrontend/Private/MetasoundAssetBase.cpp
--- a/Engine/Plugins/Runtime/Metasound/Source/MetasoundFrontend/Private/MetasoundAssetBase.cpp
+++ b/Engine/Plugins/Runtime/Metasound/Source/MetasoundFrontend/Private/MetasoundAssetBase.cpp
@@ -268,17 +268,34 @@ bool FMetasoundAssetBase::ConformObjectDataToInterfaces()
 
 TSharedPtr<Audio::IProxyData> FMetasoundAssetBase::CreateProxyData(const Audio::FProxyDataInitParams& InitParams)
 {
-	using namespace Metasound::Frontend;
-	TScriptInterface<const IMetaSoundDocumentInterface> DocInterface = GetOwningAsset();
-	const FGraphRegistryKey& Key = GetGraphRegistryKey();
-	FMetasoundAssetProxy::FParameters Args;
-	Args.Interfaces = DocInterface->GetConstDocument().Interfaces;
-	Args.Graph = FMetasoundFrontendRegistryContainer::Get()->GetGraph(Key);
-	if (Args.Graph.IsValid())
+	if (!Proxy)
 	{
-		return MakeShared<FMetasoundAssetProxy>(Args);
+		using namespace Metasound::Frontend;
+		TScriptInterface<const IMetaSoundDocumentInterface> DocInterface = GetOwningAsset();
+		if (!IsRegistered())
+		{
+			UpdateAndRegisterForExecution();
+		}
+		const FGraphRegistryKey& Key = GetGraphRegistryKey();
+		if (ensure(Key.IsValid()))
+		{
+			FMetasoundAssetProxy::FParameters Args;
+			Args.Interfaces = DocInterface->GetConstDocument().Interfaces;
+			Args.GraphRegistryKey = Key;
+			Proxy = MakeShared<FMetasoundAssetProxy>(Args);
+		}
+		else
+		{
+			UObject* Owner = GetOwningAsset();
+			check(Owner);
+			const UClass* Class = Owner->GetClass();
+			check(Class);
+			const FString ClassName = Class->GetName();
+			UE_LOG(LogMetaSound, Error, TEXT("Unable to create proxy data. Failed for MetaSound node class '%s' of UObject class '%s'"), *GetOwningAssetName(), *ClassName);
+		}
 	}
-	return nullptr;
+
+	return Proxy;
 }
 
 void FMetasoundAssetBase::RegisterGraphWithFrontend(Metasound::Frontend::FMetaSoundAssetRegistrationOptions InRegistrationOptions)
@@ -982,13 +999,30 @@ void FMetasoundAssetBase::SetVersionedOnLoad()
 FMetasoundAssetProxy::FMetasoundAssetProxy(const FParameters& InParams)
 {
 	Interfaces = InParams.Interfaces;
-	Graph = InParams.Graph;
+	GraphRegistryKey = InParams.GraphRegistryKey;
 }
 
 FMetasoundAssetProxy::FMetasoundAssetProxy(const FMetasoundAssetProxy& Other)
 {
 	Interfaces = Other.Interfaces;
+	GraphRegistryKey = Other.GraphRegistryKey;
 	Graph = Other.Graph;
 }
 
+const Metasound::IGraph* FMetasoundAssetProxy::GetGraph() const
+{
+	using namespace Metasound::Frontend;
+	if (!Graph.IsValid())
+	{
+		check(GraphRegistryKey.IsValid());
+		Graph = INodeClassRegistry::Get()->GetGraph(GraphRegistryKey);
+	}
+	return Graph.Get();
+}
+
+const TSet<FMetasoundFrontendVersion>& FMetasoundAssetProxy::GetInterfaces() const
+{
+	return Interfaces;
+}
+
 #undef LOCTEXT_NAMESPACE // "MetaSound"
diff --git a/Engine/Plugins/Runtime/Metasound/Source/MetasoundFrontend/Public/MetasoundAssetBase.h b/Engine/Plugins/Runtime/Metasound/Source/MetasoundFrontend/Public/MetasoundAssetBase.h
--- a/Engine/Plugins/Runtime/Metasound/Source/MetasoundFrontend/Public/MetasoundAssetBase.h
+++ b/Engine/Plugins/Runtime/Metasound/Source/MetasoundFrontend/Public/MetasoundAssetBase.h
@@ -25,6 +25,9 @@ struct FMetaSoundFrontendDocumentBuilder;
 class IConsoleVariable;
 typedef TMulticastDelegate<void(IConsoleVariable*), FDefaultDelegateUserPolicy> FConsoleVariableMulticastDelegate;
 
+class FMetasoundAssetProxy;
+using FMetasoundAssetProxyPtr = TSharedPtr<FMetasoundAssetProxy>;
+
 namespace Metasound
 {
 	class IGraph;
@@ -263,6 +266,8 @@ private:
 #endif // WITH_EDITORONLY_DATA
 
 	Metasound::Frontend::FGraphRegistryKey GraphRegistryKey;
+	
+	FMetasoundAssetProxyPtr Proxy;
 };
 
 class FMetasoundAssetProxy final : public Audio::TProxyData<FMetasoundAssetProxy>
@@ -273,29 +278,23 @@ public:
 	struct FParameters
 	{
 		TSet<FMetasoundFrontendVersion> Interfaces;
-		TSharedPtr<const Metasound::IGraph> Graph;
-
+		Metasound::Frontend::FGraphRegistryKey GraphRegistryKey;
 	};
 
 	UE_API explicit FMetasoundAssetProxy(const FParameters& InParams);
 	
 	UE_API FMetasoundAssetProxy(const FMetasoundAssetProxy& Other);
 
-	const Metasound::IGraph* GetGraph() const
-	{
-		return Graph.Get();
-	}
-
-	const TSet<FMetasoundFrontendVersion>& GetInterfaces() const
-	{
-		return Interfaces;
-	}
+	UE_API const Metasound::IGraph* GetGraph() const;
+	
+	UE_API const TSet<FMetasoundFrontendVersion>& GetInterfaces() const;
 
 private:
 	
 	TSet<FMetasoundFrontendVersion> Interfaces;
-	TSharedPtr<const Metasound::IGraph> Graph;
+	Metasound::Frontend::FGraphRegistryKey GraphRegistryKey;
+	// mutable so we can cache it when calling GetGraph
+	mutable TSharedPtr<const Metasound::IGraph> Graph;
 };
-using FMetasoundAssetProxyPtr = TSharedPtr<FMetasoundAssetProxy, ESPMode::ThreadSafe>;
 
 #undef UE_API
