Caching the Graph is a mistake since it can change after creating the proxy data

#rnx

diff --git a/Engine/Plugins/Runtime/Metasound/Source/MetasoundFrontend/Private/MetasoundAssetBase.cpp b/Engine/Plugins/Runtime/Metasound/Source/MetasoundFrontend/Private/MetasoundAssetBase.cpp
--- a/Engine/Plugins/Runtime/Metasound/Source/MetasoundFrontend/Private/MetasoundAssetBase.cpp
+++ b/Engine/Plugins/Runtime/Metasound/Source/MetasoundFrontend/Private/MetasoundAssetBase.cpp
@@ -268,15 +268,37 @@ bool FMetasoundAssetBase::ConformObjectDataToInterfaces()
 
 TSharedPtr<Audio::IProxyData> FMetasoundAssetBase::CreateProxyData(const Audio::FProxyDataInitParams& InitParams)
 {
-	if (!Proxy)
+	using namespace Metasound::Frontend;
+	
+	if (!IsRegistered())
+	{
+		UpdateAndRegisterForExecution();
+	}
+	const FGraphRegistryKey& Key = GetGraphRegistryKey();
+	TScriptInterface<const IMetaSoundDocumentInterface> DocInterface = GetOwningAsset();
+
+	// check if graph key has changed
+	if (Proxy)
+	{
+		if (Proxy->GetGraphRegistryKey() != Key)
+		{
+			Proxy = nullptr;
+		}
+	}
+
+	// check interfaces have changed
+	if (Proxy)
 	{
-		using namespace Metasound::Frontend;
-		TScriptInterface<const IMetaSoundDocumentInterface> DocInterface = GetOwningAsset();
-		if (!IsRegistered())
+		int32 Num = Proxy->GetInterfaces().Num();
+		if (Num != Proxy->GetInterfaces().Union(DocInterface->GetConstDocument().Interfaces).Num())
 		{
-			UpdateAndRegisterForExecution();
+			Proxy = nullptr;
 		}
-		const FGraphRegistryKey& Key = GetGraphRegistryKey();
+	}
+
+	// recreate the proxy if necessary
+	if (!Proxy)
+	{
 		if (ensure(Key.IsValid()))
 		{
 			FMetasoundAssetProxy::FParameters Args;
@@ -1006,18 +1028,21 @@ FMetasoundAssetProxy::FMetasoundAssetProxy(const FMetasoundAssetProxy& Other)
 {
 	Interfaces = Other.Interfaces;
 	GraphRegistryKey = Other.GraphRegistryKey;
-	Graph = Other.Graph;
+}
+
+const Metasound::Frontend::FGraphRegistryKey& FMetasoundAssetProxy::GetGraphRegistryKey() const
+{
+	return GraphRegistryKey;
 }
 
 const Metasound::IGraph* FMetasoundAssetProxy::GetGraph() const
 {
 	using namespace Metasound::Frontend;
-	if (!Graph.IsValid())
+	if (GraphRegistryKey.IsValid())
 	{
-		check(GraphRegistryKey.IsValid());
-		Graph = INodeClassRegistry::Get()->GetGraph(GraphRegistryKey);
+		return INodeClassRegistry::Get()->GetGraph(GraphRegistryKey).Get();
 	}
-	return Graph.Get();
+	return nullptr;
 }
 
 const TSet<FMetasoundFrontendVersion>& FMetasoundAssetProxy::GetInterfaces() const
diff --git a/Engine/Plugins/Runtime/Metasound/Source/MetasoundFrontend/Public/MetasoundAssetBase.h b/Engine/Plugins/Runtime/Metasound/Source/MetasoundFrontend/Public/MetasoundAssetBase.h
--- a/Engine/Plugins/Runtime/Metasound/Source/MetasoundFrontend/Public/MetasoundAssetBase.h
+++ b/Engine/Plugins/Runtime/Metasound/Source/MetasoundFrontend/Public/MetasoundAssetBase.h
@@ -285,6 +285,8 @@ public:
 	
 	UE_API FMetasoundAssetProxy(const FMetasoundAssetProxy& Other);
 
+	UE_API const Metasound::Frontend::FGraphRegistryKey& GetGraphRegistryKey() const;
+
 	UE_API const Metasound::IGraph* GetGraph() const;
 	
 	UE_API const TSet<FMetasoundFrontendVersion>& GetInterfaces() const;
@@ -293,8 +295,6 @@ private:
 	
 	TSet<FMetasoundFrontendVersion> Interfaces;
 	Metasound::Frontend::FGraphRegistryKey GraphRegistryKey;
-	// mutable so we can cache it when calling GetGraph
-	mutable TSharedPtr<const Metasound::IGraph> Graph;
 };
 
 #undef UE_API
