Add cvar to continue using Change ID comparison when auto-updating/diffing registry interfaces, but disable now by default.

Change IDs are to eventually be deprecated in the future in favor of leveraging MetaSound Builder transaction "tracking"/count. ChangeID is not reliable in all cases, causing editor bugs at times where an interface wasn't getting updated when it should on load of the given asset.

#rnx

diff --git a/Engine/Plugins/Runtime/Metasound/Source/MetasoundFrontend/Private/MetasoundFrontendDocumentVersioning.cpp b/Engine/Plugins/Runtime/Metasound/Source/MetasoundFrontend/Private/MetasoundFrontendDocumentVersioning.cpp
--- a/Engine/Plugins/Runtime/Metasound/Source/MetasoundFrontend/Private/MetasoundFrontendDocumentVersioning.cpp
+++ b/Engine/Plugins/Runtime/Metasound/Source/MetasoundFrontend/Private/MetasoundFrontendDocumentVersioning.cpp
@@ -4,6 +4,7 @@
 #if WITH_EDITORONLY_DATA
 #include "Algo/Transform.h"
 #include "CoreGlobals.h"
+#include "HAL/IConsoleManager.h"
 #include "Interfaces/MetasoundFrontendInterface.h"
 #include "Interfaces/MetasoundFrontendInterfaceRegistry.h"
 #include "MetasoundAccessPtr.h"
@@ -24,6 +25,15 @@ namespace Metasound::Frontend
 
 	namespace VersioningPrivate
 	{
+		static int32 MetaSoundAutoUpdateUseAssetChangeIDsCVar = 0;
+
+		FAutoConsoleVariableRef CVarMetaSoundAutoUpdateUseAssetChangeIDs(
+			TEXT("au.MetaSound.AutoUpdate.UseChangeIDs"),
+			MetaSoundAutoUpdateUseAssetChangeIDsCVar,
+			TEXT("If true, use soft-deprecated Change ID system to speed up diffing interface or metadata changes during auto-update. If false, ignore change IDs and always do \"deep check\" for changes. \n")
+			TEXT("0: Don't use Change Ids (default), !0: Use ChangeIDs"),
+			ECVF_Default);
+
 		class FMigratePagePropertiesTransform : public FMetaSoundFrontendDocumentBuilder::IPropertyVersionTransform
 		{
 		public:
@@ -766,6 +776,11 @@ PRAGMA_ENABLE_DEPRECATION_WARNINGS
 		}
 	} // namespace VersioningPrivate
 
+	bool ChangeIDComparisonEnabledInAutoUpdate()
+	{
+		return VersioningPrivate::MetaSoundAutoUpdateUseAssetChangeIDsCVar != 0;
+	}
+
 	bool VersionDocument(FMetaSoundFrontendDocumentBuilder& Builder)
 	{
 		using namespace VersioningPrivate;
diff --git a/Engine/Plugins/Runtime/Metasound/Source/MetasoundFrontend/Private/MetasoundFrontendDocumentVersioning.h b/Engine/Plugins/Runtime/Metasound/Source/MetasoundFrontend/Private/MetasoundFrontendDocumentVersioning.h
--- a/Engine/Plugins/Runtime/Metasound/Source/MetasoundFrontend/Private/MetasoundFrontendDocumentVersioning.h
+++ b/Engine/Plugins/Runtime/Metasound/Source/MetasoundFrontend/Private/MetasoundFrontendDocumentVersioning.h
@@ -19,6 +19,8 @@ namespace Metasound::Frontend
 		return FMetasoundFrontendVersionNumber { 1, 14 };
 	}
 
+	bool ChangeIDComparisonEnabledInAutoUpdate();
+
 	// Versions Frontend Document. Passed as AssetBase for backward compat to
 	// version asset documents predating the IMetaSoundDocumentInterface
 	bool VersionDocument(FMetaSoundFrontendDocumentBuilder& Builder);
diff --git a/Engine/Plugins/Runtime/Metasound/Source/MetasoundFrontend/Private/MetasoundFrontendNodeController.cpp b/Engine/Plugins/Runtime/Metasound/Source/MetasoundFrontend/Private/MetasoundFrontendNodeController.cpp
--- a/Engine/Plugins/Runtime/Metasound/Source/MetasoundFrontend/Private/MetasoundFrontendNodeController.cpp
+++ b/Engine/Plugins/Runtime/Metasound/Source/MetasoundFrontend/Private/MetasoundFrontendNodeController.cpp
@@ -6,6 +6,7 @@
 #include "Algo/Transform.h"
 #include "HAL/IConsoleManager.h"
 #include "MetasoundFrontendDocumentAccessPtr.h"
+#include "MetasoundFrontendDocumentVersioning.h"
 #include "MetasoundFrontendInputController.h"
 #include "MetasoundFrontendInvalidController.h"
 #include "MetasoundFrontendOutputController.h"
@@ -1226,6 +1227,7 @@ namespace Metasound
 		{
 			METASOUND_TRACE_CPUPROFILER_EVENT_SCOPE(BaseNodeController::CanAutoUpdate);
 
+#if WITH_EDITORONLY_DATA
 			OutInterfaceUpdates = { };
 
 			const FMetasoundFrontendClassMetadata& NodeClassMetadata = GetClassMetadata();
@@ -1251,28 +1253,31 @@ namespace Metasound
 			// to ensure a change wasn't created that didn't contain a version change.
 			if (RegistryClass.Metadata.GetVersion() == NodeClassMetadata.GetVersion())
 			{
-				// TODO: Merge these paths.  Shouldn't use different logic to
-				// define changes in native vs asset class definitions. Need to
-				// hash a changeID from natively defined node classes in order to
-				// merge branches.
-				
-				const bool bIsAssetClass = AssetManager.IsAssetClass(RegistryClass.Metadata);
-				if (bIsAssetClass)
+				if (Frontend::ChangeIDComparisonEnabledInAutoUpdate())
 				{
-					if (RegistryClass.Metadata.GetChangeID() == NodeClassMetadata.GetChangeID())
+					// WIP: Merging these paths.  Should no longer require using different
+					// logic to determine changes in native vs asset class definitions.
+					// Its less performant to ignore ChangeIDs, but auto-update is no longer
+					// called at runtime and the editor now avoids running on mass numbers of
+					// assets at once due to the utilization of AssetTag data for querying class info.
+					const bool bIsAssetClass = AssetManager.IsAssetClass(RegistryClass.Metadata);
+					if (bIsAssetClass)
 					{
-						const FGuid& NodeClassInterfaceChangeID = GetClassInterface().GetChangeID();
-						if (RegistryClass.GetDefaultInterface().GetChangeID() == NodeClassInterfaceChangeID)
+						if (RegistryClass.Metadata.GetChangeID() == NodeClassMetadata.GetChangeID())
 						{
-							return false;
+							const FGuid& NodeClassInterfaceChangeID = GetClassInterface().GetChangeID();
+							if (RegistryClass.GetDefaultInterface().GetChangeID() == NodeClassInterfaceChangeID)
+							{
+								return false;
+							}
 						}
 					}
-				}
-				else
-				{
-					if (!MetaSoundAutoUpdateNativeClassesOfEqualVersionCVar)
+					else
 					{
-						return false;
+						if (!MetaSoundAutoUpdateNativeClassesOfEqualVersionCVar)
+						{
+							return false;
+						}
 					}
 				}
 
@@ -1295,6 +1300,11 @@ PRAGMA_DISABLE_DEPRECATION_WARNINGS
 			DiffAgainstRegistryInterface(OutInterfaceUpdates, bUseHighestMinorVersion, bForceRegenerateClassInterfaceOverride);
 PRAGMA_ENABLE_DEPRECATION_WARNINGS
 			return true;
+
+#else // !WITH_EDITORONLY_DATA
+			return false; // Auto-Update is now never permitted in builds without non-editoronly data
+#endif // !WITH_EDITORONLY_DATA
+
 		}
 #endif // WITH_EDITORONLY_DATA
 
