Fix MetaSound connections being invalid when changing output format

#rnx

diff --git a/Engine/Plugins/Runtime/Metasound/Source/MetasoundEditor/Private/MetasoundEditorGraphBuilder.cpp b/Engine/Plugins/Runtime/Metasound/Source/MetasoundEditor/Private/MetasoundEditorGraphBuilder.cpp
--- a/Engine/Plugins/Runtime/Metasound/Source/MetasoundEditor/Private/MetasoundEditorGraphBuilder.cpp
+++ b/Engine/Plugins/Runtime/Metasound/Source/MetasoundEditor/Private/MetasoundEditorGraphBuilder.cpp
@@ -1587,7 +1587,7 @@ PRAGMA_ENABLE_DEPRECATION_WARNINGS
 			if (InEditorPin.Direction == EGPD_Input)
 			{
 				const FMetasoundFrontendVertex* PinVertex = GetPinVertex(InBuilder, &InEditorPin);
-				return PinVertex && PinVertex->VertexID == InVertex.VertexID;
+				return PinVertex && PinVertex->VertexID == InVertex.VertexID && InVertex.Name == InEditorPin.PinName;
 			}
 
 			return false;
@@ -1598,7 +1598,7 @@ PRAGMA_ENABLE_DEPRECATION_WARNINGS
 			if (InEditorPin.Direction == EGPD_Output)
 			{
 				const FMetasoundFrontendVertex* PinVertex = GetPinVertex(InBuilder, &InEditorPin);
-				return PinVertex && PinVertex->VertexID == InVertex.VertexID;
+				return PinVertex && PinVertex->VertexID == InVertex.VertexID && InVertex.Name == InEditorPin.PinName;
 			}
 
 			return false;
