Rework page perefence data plubming. Optimize perf for page resolution
#rnx

diff --git a/Engine/Plugins/Experimental/TechAudioTools/Source/TechAudioToolsMetaSound/Private/ViewModels/MetaSoundViewModel.cpp b/Engine/Plugins/Experimental/TechAudioTools/Source/TechAudioToolsMetaSound/Private/ViewModels/MetaSoundViewModel.cpp
--- a/Engine/Plugins/Experimental/TechAudioTools/Source/TechAudioToolsMetaSound/Private/ViewModels/MetaSoundViewModel.cpp
+++ b/Engine/Plugins/Experimental/TechAudioTools/Source/TechAudioToolsMetaSound/Private/ViewModels/MetaSoundViewModel.cpp
@@ -4,7 +4,9 @@
 
 #include "Logging/StructuredLog.h"
 #include "MetasoundDocumentBuilderRegistry.h"
+#include "MetasoundFrontendPages.h"
 #include "MetasoundOutputSubsystem.h"
+#include "MetasoundSettings.h"
 
 #include UE_INLINE_GENERATED_CPP_BY_NAME(MetaSoundViewModel)
 
@@ -170,12 +172,14 @@ void UMetaSoundInputViewModel::InitializeInput(UMetaSoundBuilderBase* InBuilder,
 	UE_MVVM_SET_PROPERTY_VALUE(InputName, InInput.Name);
 	UE_MVVM_SET_PROPERTY_VALUE(DataType, InInput.TypeName);
 
-	const FGuid PageID = Metasound::Engine::FDocumentBuilderRegistry::GetChecked().ResolveTargetPageID(InInput);
-	const FMetasoundFrontendLiteral& DefaultLiteral = InInput.FindConstDefaultChecked(PageID);
+	if (const FMetasoundFrontendClassInputDefault* Default = Metasound::Frontend::FindPreferredPage(InInput.GetDefaults(), UMetaSoundSettings::GetPageOrder()); ensure(Default))
+	{
+		const FMetasoundFrontendLiteral& DefaultLiteral = Default->Literal;
 
-	UE_MVVM_SET_PROPERTY_VALUE(LiteralType, DefaultLiteral.GetType());
-	UE_MVVM_SET_PROPERTY_VALUE(Literal, DefaultLiteral);
-	UE_MVVM_SET_PROPERTY_VALUE(bIsArray, DefaultLiteral.IsArray());
+		UE_MVVM_SET_PROPERTY_VALUE(LiteralType, DefaultLiteral.GetType());
+		UE_MVVM_SET_PROPERTY_VALUE(Literal, DefaultLiteral);
+		UE_MVVM_SET_PROPERTY_VALUE(bIsArray, DefaultLiteral.IsArray());
+	}
 
 	SetIsInitialized(true);
 }
diff --git a/Engine/Plugins/Runtime/Metasound/Source/MetasoundEditor/Private/MetasoundEditor.cpp b/Engine/Plugins/Runtime/Metasound/Source/MetasoundEditor/Private/MetasoundEditor.cpp
--- a/Engine/Plugins/Runtime/Metasound/Source/MetasoundEditor/Private/MetasoundEditor.cpp
+++ b/Engine/Plugins/Runtime/Metasound/Source/MetasoundEditor/Private/MetasoundEditor.cpp
@@ -57,6 +57,7 @@
 #include "MetasoundFrontendDocument.h"
 #include "MetasoundFrontendDocumentBuilder.h"
 #include "MetasoundFrontendNodeTemplateRegistry.h"
+#include "MetasoundFrontendPages.h"
 #include "MetasoundFrontendRegistries.h"
 #include "MetasoundFrontendSearchEngine.h"
 #include "MetasoundFrontendTransform.h"
@@ -132,12 +133,15 @@ namespace Metasound
 		bool IsPreviewingPageInputDefault(const FMetaSoundFrontendDocumentBuilder& Builder, const FMetasoundFrontendClassInput& InClassInput, const FGuid& InPageID)
 		{
 			using namespace Engine;
+			using namespace Frontend;
 
 			const UObject& MetaSound = Builder.CastDocumentObjectChecked<const UObject>();
 			if (IsPreviewingMetaSound(MetaSound))
 			{
-				const FGuid TargetPageID = FDocumentBuilderRegistry::GetChecked().ResolveTargetPageID(InClassInput);
-				return TargetPageID == InPageID;
+				if (const FMetasoundFrontendClassInputDefault* InputDefault = FindPreferredPage(InClassInput.GetDefaults(), UMetaSoundSettings::GetPageOrder()))
+				{
+					return InputDefault->PageID == InPageID;
+				}
 			}
 
 			return false;
@@ -146,13 +150,16 @@ namespace Metasound
 		bool IsPreviewingPageGraph(const FMetaSoundFrontendDocumentBuilder& Builder, const FGuid& InPageID)
 		{
 			using namespace Engine;
+			using namespace Frontend;
 
 			const UObject& MetaSound = Builder.CastDocumentObjectChecked<const UObject>();
 			if (IsPreviewingMetaSound(MetaSound))
 			{
 				const FMetasoundFrontendGraphClass& GraphClass = Builder.GetConstDocumentChecked().RootGraph;
-				const FGuid TargetPageID = FDocumentBuilderRegistry::GetChecked().ResolveTargetPageID(GraphClass);
-				return TargetPageID == InPageID;
+				if (const FMetasoundFrontendGraph* GraphPage = FindPreferredPage(GraphClass.GetConstGraphPages(), UMetaSoundSettings::GetPageOrder()))
+				{
+					return GraphPage->PageID == InPageID;
+				}
 			}
 
 			return false;
@@ -1923,7 +1930,7 @@ namespace Metasound
 					{
 						if (UMetasoundEditorSettings* EditorSettings = GetMutableDefault<UMetasoundEditorSettings>())
 						{
-							EditorSettings->AuditionPlatform = PlatformName;
+							EditorSettings->SetAuditionPlatform(PlatformName);
 							Stop();
 							SyncAuditionState();
 						}
@@ -1932,7 +1939,7 @@ namespace Metasound
 					{
 						if (const UMetasoundEditorSettings* EditorSettings = GetDefault<UMetasoundEditorSettings>())
 						{
-							if (EditorSettings->AuditionPlatform == PlatformName)
+							if (EditorSettings->GetAuditionPlatform() == PlatformName)
 							{
 								return ECheckBoxState::Checked;
 							}
@@ -1949,7 +1956,7 @@ namespace Metasound
 				// Protects against stale setting not showing after platform values are manipulated just for visibility
 				if (const UMetasoundEditorSettings* EditorSettings = GetDefault<UMetasoundEditorSettings>())
 				{
-					AuditionPlatforms.AddUnique(EditorSettings->AuditionPlatform);
+					AuditionPlatforms.AddUnique(EditorSettings->GetAuditionPlatform());
 				}
 
 				for (const FName& PlatformName : AuditionPlatforms)
@@ -1997,7 +2004,7 @@ namespace Metasound
 							{
 								case ECheckBoxState::Checked:
 								{
-									EdSettings->AuditionPageMode = EAuditionPageMode::Focused;
+									EdSettings->SetAuditionPageMode(EAuditionPageMode::Focused);
 									ThisEditor->Stop();
 									ThisEditor->SyncAuditionState();
 									break;
@@ -2007,7 +2014,7 @@ namespace Metasound
 								case ECheckBoxState::Undetermined:
 								default:
 								{
-									EdSettings->AuditionPageMode = EAuditionPageMode::User;
+									EdSettings->SetAuditionPageMode(EAuditionPageMode::User);
 								}
 								break;
 							}
@@ -2017,7 +2024,7 @@ namespace Metasound
 					{
 						if (const UMetasoundEditorSettings* EdSettings = GetDefault<UMetasoundEditorSettings>())
 						{
-							if (EdSettings->AuditionPageMode == EAuditionPageMode::Focused)
+							if (EdSettings->GetAuditionPageMode() == EAuditionPageMode::Focused)
 							{
 								return ECheckBoxState::Checked;
 							}
@@ -2042,9 +2049,9 @@ namespace Metasound
 					{
 						if (UMetasoundEditorSettings* EditorSettings = GetMutableDefault<UMetasoundEditorSettings>())
 						{
-							if (EditorSettings->AuditionPage != AuditionPage)
+							if (EditorSettings->GetAuditionPage() != AuditionPage)
 							{
-								EditorSettings->AuditionPage = AuditionPage;
+								EditorSettings->SetAuditionPage(AuditionPage);
 								if (TSharedPtr<FEditor> ThisEditor = EditorPtr.Pin())
 								{
 									ThisEditor->Stop();
@@ -2058,7 +2065,7 @@ namespace Metasound
 					{
 						if (const UMetasoundEditorSettings* EditorSettings = GetDefault<UMetasoundEditorSettings>())
 						{
-							if (EditorSettings->AuditionPage == AuditionPage)
+							if (EditorSettings->GetAuditionPage() == AuditionPage)
 							{
 								return ECheckBoxState::Checked;
 							}
@@ -2070,7 +2077,7 @@ namespace Metasound
 					{
 						if (const UMetasoundEditorSettings* EdSettings = GetDefault<UMetasoundEditorSettings>())
 						{
-							return EdSettings->AuditionPageMode == EAuditionPageMode::User;
+							return EdSettings->GetAuditionPageMode() == EAuditionPageMode::User;
 						}
 
 						if (const UMetaSoundSettings* Settings = GetDefault<UMetaSoundSettings>())
@@ -2084,24 +2091,24 @@ namespace Metasound
 					{
 						if (const UMetasoundEditorSettings* EditorSettings = GetDefault<UMetasoundEditorSettings>())
 						{
-							if (EditorSettings->AuditionPlatform != UMetasoundEditorSettings::EditorAuditionPlatform)
+							if (EditorSettings->GetAuditionPlatform() != UMetasoundEditorSettings::EditorAuditionPlatform)
 							{
 								if (const UMetaSoundSettings* Settings = GetDefault<UMetaSoundSettings>())
 								{
 									if (const FMetaSoundPageSettings* PageSetting = Settings->FindPageSettings(AuditionPage))
 									{
-										if (!PageSetting->PlatformCanTargetPage(EditorSettings->AuditionPlatform))
+										if (!PageSetting->PlatformCanTargetPage(EditorSettings->GetAuditionPlatform()))
 										{
-											return FText::Format(LOCTEXT("AuditionPageInvalidForPlatformToolTip", "Platform '{0}' does not target page '{1}'. See 'MetaSound' Project Settings"), FText::FromName(EditorSettings->AuditionPlatform), PageText);
+											return FText::Format(LOCTEXT("AuditionPageInvalidForPlatformToolTip", "Platform '{0}' does not target page '{1}'. See 'MetaSound' Project Settings"), FText::FromName(EditorSettings->GetAuditionPlatform()), PageText);
 										}
 									}
 								}
 							}
 
-							const bool bUserAuditionMode = EditorSettings->AuditionPageMode == EAuditionPageMode::User;
+							const bool bUserAuditionMode = EditorSettings->GetAuditionPageMode() == EAuditionPageMode::User;
 							if (bUserAuditionMode)
 							{
-								if (EditorSettings->AuditionPage != AuditionPage)
+								if (EditorSettings->GetAuditionPage() != AuditionPage)
 								{
 									return FText::Format(LOCTEXT("SetAuditionPageToolTip", "Sets the user's editor AuditionPage setting to '{0}'."), PageText);
 								}
@@ -2655,9 +2662,9 @@ namespace Metasound
 
 				if (const UMetasoundEditorSettings* EdSettings = GetDefault<UMetasoundEditorSettings>())
 				{
-					if (EdSettings->AuditionPlatform != UMetasoundEditorSettings::EditorAuditionPlatform)
+					if (EdSettings->GetAuditionPlatform() != UMetasoundEditorSettings::EditorAuditionPlatform)
 					{
-						if (!UMetaSoundEditorSubsystem::GetChecked().IsPageAuditionPlatformCookTarget(EdSettings->AuditionPage))
+						if (!UMetaSoundEditorSubsystem::GetChecked().IsPageAuditionPlatformCookTarget(EdSettings->GetAuditionPage()))
 						{
 							GraphStatusDescriptionOverride = LOCTEXT("InvalidAuditionPageWarning",
 								"Selected Audition Page in MetaSound Editor Settings is not a target page for the selected 'Audition Platform'. "
@@ -5011,6 +5018,7 @@ PRAGMA_ENABLE_DEPRECATION_WARNINGS
 		TUniquePtr<FGraphConnectionManager> FEditor::RebuildConnectionManager(UAudioComponent* PreviewComp) const
 		{
 			using namespace Engine;
+			using namespace Frontend;
 
 			if (!PreviewComp || !Builder.IsValid())
 			{
@@ -5023,8 +5031,14 @@ PRAGMA_ENABLE_DEPRECATION_WARNINGS
 				return MakeUnique<FGraphConnectionManager>();
 			}
 
-			const FGuid ResolvedGraphPageID = FDocumentBuilderRegistry::GetChecked().ResolveTargetPageID(Source->GetConstDocument().RootGraph);
-			if (ResolvedGraphPageID != Builder->GetConstBuilder().GetBuildPageID())
+			if (const FMetasoundFrontendGraph* GraphPage = FindPreferredPage(Source->GetConstDocument().RootGraph.GetConstGraphPages(), UMetaSoundSettings::GetPageOrder()))
+			{
+				if (GraphPage->PageID !=  Builder->GetConstBuilder().GetBuildPageID())
+				{
+					return MakeUnique<FGraphConnectionManager>();
+				}
+			}
+			else
 			{
 				return MakeUnique<FGraphConnectionManager>();
 			}
@@ -5041,6 +5055,7 @@ PRAGMA_ENABLE_DEPRECATION_WARNINGS
 		void FEditor::UpdatePageInfo(bool bIsPlaying)
 		{
 			using namespace Engine;
+			using namespace Frontend;
 
 			const UMetaSoundSettings* Settings = GetDefault<UMetaSoundSettings>();
 			check(Settings)
@@ -5055,18 +5070,19 @@ PRAGMA_ENABLE_DEPRECATION_WARNINGS
 				{
 					if (const UMetasoundEditorSettings* EditorSettings = GetDefault<UMetasoundEditorSettings>())
 					{
-						AuditionPageSettings = Settings->FindPageSettings(EditorSettings->AuditionPage);
+						AuditionPageSettings = Settings->FindPageSettings(EditorSettings->GetAuditionPage());
 					}
 
 					const FMetaSoundFrontendDocumentBuilder& DocBuilder = Builder->GetConstBuilder();
 					const FGuid& PageID = DocBuilder.GetBuildPageID();
 					if (bIsPlaying)
 					{
-						const FMetasoundFrontendGraphClass& GraphClass = DocBuilder.GetConstDocumentChecked().RootGraph;
-						const FGuid ResolvePageID = FDocumentBuilderRegistry::GetChecked().ResolveTargetPageID(GraphClass);
-						if (ResolvePageID == PageID)
+						if (const FMetasoundFrontendGraph* Graph= FindPreferredPage(DocBuilder.GetConstDocumentChecked().RootGraph.GetConstGraphPages(), UMetaSoundSettings::GetPageOrder()))
 						{
-							Color = &Style::GetPageExecutingColor();
+							if (Graph->PageID == PageID)
+							{
+								Color = &Style::GetPageExecutingColor();
+							}
 						}
 					}
 					GraphPageSettings = Settings->FindPageSettings(PageID);
diff --git a/Engine/Plugins/Runtime/Metasound/Source/MetasoundEditor/Private/MetasoundEditorGraph.cpp b/Engine/Plugins/Runtime/Metasound/Source/MetasoundEditor/Private/MetasoundEditorGraph.cpp
--- a/Engine/Plugins/Runtime/Metasound/Source/MetasoundEditor/Private/MetasoundEditorGraph.cpp
+++ b/Engine/Plugins/Runtime/Metasound/Source/MetasoundEditor/Private/MetasoundEditorGraph.cpp
@@ -20,6 +20,7 @@
 #include "MetasoundFrontendDataTypeRegistry.h"
 #include "MetasoundFrontendDocumentBuilder.h"
 #include "MetasoundFrontendNodeTemplateRegistry.h"
+#include "MetasoundFrontendPages.h"
 #include "MetasoundFrontendSearchEngine.h"
 #include "MetasoundLog.h"
 #include "MetasoundSettings.h"
@@ -82,6 +83,8 @@ namespace Metasound::Editor
 		// avoids 'Reset To Default' action from being enabled when the default is equal)
 		void SetOrClearIfLiteralMatchesNodeVertexDefault(FMetaSoundFrontendDocumentBuilder& InBuilder, const FMetasoundFrontendVertexHandle& VertexHandle, const FMetasoundFrontendLiteral& InDefaultLiteral)
 		{
+			using namespace Frontend; 
+
 			const FMetasoundFrontendVertex* Vertex = InBuilder.FindNodeInput(VertexHandle.NodeID, VertexHandle.VertexID);
 			check(Vertex);
 
@@ -89,9 +92,7 @@ namespace Metasound::Editor
 			const TArray<FMetasoundFrontendClassInputDefault>* ClassDefaults = InBuilder.FindNodeClassInputDefaults(VertexHandle.NodeID, Vertex->Name);
 			if (ClassDefaults)
 			{
-				const FGuid PageID = Engine::FDocumentBuilderRegistry::GetChecked().ResolveTargetPageID(*ClassDefaults);
-				auto MatchesPageID = [&PageID](const FMetasoundFrontendClassInputDefault& InputDefault) { return InputDefault.PageID == PageID; };
-				if (const FMetasoundFrontendClassInputDefault* ClassDefault = ClassDefaults->FindByPredicate(MatchesPageID))
+				if (const FMetasoundFrontendClassInputDefault* ClassDefault = FindPreferredPage(*ClassDefaults, UMetaSoundSettings::GetPageOrder()))
 				{
 					bClearLiteral = ClassDefault->Literal.IsEqual(InDefaultLiteral);
 				}
@@ -709,6 +710,7 @@ void UMetasoundEditorGraphMemberDefaultLiteral::PostEditUndo()
 bool UMetasoundEditorGraphMemberDefaultLiteral::TryGetPreviewPageID(FGuid& OutPreviewPageID) const
 {
 	using namespace Metasound::Engine;
+	using namespace Metasound::Frontend;
 
 	const UMetasoundEditorGraphMember* Member = FindMember();
 	if (!ensure(Member))
@@ -720,17 +722,18 @@ bool UMetasoundEditorGraphMemberDefaultLiteral::TryGetPreviewPageID(FGuid& OutPr
 	const FMetaSoundFrontendDocumentBuilder& Builder = Member->GetFrontendBuilderChecked();
 	if (const FMetasoundFrontendClassInput* Input = Builder.FindGraphInput(Member->GetMemberName()))
 	{
-		const FGuid PageID = FDocumentBuilderRegistry::GetChecked().ResolveTargetPageID(*Input);
-
 		const UMetasoundEditorSettings* EdSettings = ::GetDefault<UMetasoundEditorSettings>();
 		const UMetaSoundSettings* Settings = ::GetDefault<UMetaSoundSettings>();
 		if (Settings && EdSettings)
 		{
-			if (const FMetaSoundPageSettings* PageSettings = Settings->FindPageSettings(EdSettings->AuditionPage))
+			const FMetasoundFrontendClassInputDefault* DefaultLiteral = FindPreferredPage(Input->GetDefaults(), Settings->GetPageOrder());
+			const FMetaSoundPageSettings* PageSettings = Settings->FindPageSettings(EdSettings->GetAuditionPage());
+
+			if (DefaultLiteral && PageSettings)
 			{
-				if (PageID == PageSettings->UniqueId)
+				if (DefaultLiteral->PageID == PageSettings->UniqueId)
 				{
-					OutPreviewPageID = PageID;
+					OutPreviewPageID = PageSettings->UniqueId;
 					return true;
 				}
 			}
diff --git a/Engine/Plugins/Runtime/Metasound/Source/MetasoundEditor/Private/MetasoundEditorGraphBuilder.cpp b/Engine/Plugins/Runtime/Metasound/Source/MetasoundEditor/Private/MetasoundEditorGraphBuilder.cpp
--- a/Engine/Plugins/Runtime/Metasound/Source/MetasoundEditor/Private/MetasoundEditorGraphBuilder.cpp
+++ b/Engine/Plugins/Runtime/Metasound/Source/MetasoundEditor/Private/MetasoundEditorGraphBuilder.cpp
@@ -32,6 +32,7 @@
 #include "MetasoundFrontendDataTypeRegistry.h"
 #include "MetasoundFrontendDocumentAccessPtr.h"
 #include "MetasoundFrontendDocumentBuilder.h"
+#include "MetasoundFrontendPages.h"
 #include "MetasoundFrontendQuery.h"
 #include "MetasoundFrontendQuerySteps.h"
 #include "MetasoundFrontendRegistries.h"
@@ -1242,9 +1243,7 @@ PRAGMA_ENABLE_DEPRECATION_WARNINGS
 						const TArray<FMetasoundFrontendClassInputDefault>* ClassInputDefaults = Builder.FindNodeClassInputDefaults(InputHandle.NodeID, Vertex->Name);
 						if (ClassInputDefaults)
 						{
-							const FGuid PageID = Engine::FDocumentBuilderRegistry::GetChecked().ResolveTargetPageID(*ClassInputDefaults);
-							auto MatchesPageID = [&PageID](const FMetasoundFrontendClassInputDefault& InputDefault) { return PageID == InputDefault.PageID; };
-							if (const FMetasoundFrontendClassInputDefault* ClassDefault = ClassInputDefaults->FindByPredicate(MatchesPageID))
+							if (const FMetasoundFrontendClassInputDefault* ClassDefault = FindPreferredPage(*ClassInputDefaults, UMetaSoundSettings::GetPageOrder()))
 							{
 								if (ClassDefault->Literal.GetType() == EMetasoundFrontendLiteralType::None)
 								{
@@ -1564,6 +1563,7 @@ PRAGMA_ENABLE_DEPRECATION_WARNINGS
 			FMetaSoundAssetRegistrationOptions RegOptions;
 			RegOptions.bForceReregister = true;
 			RegOptions.bForceViewSynchronization = bInForceViewSynchronization;
+			RegOptions.PageOrder = UMetaSoundSettings::GetPageOrder();
 			
 			// if EditedReferencingMetaSounds is empty, then no MetaSounds are open
 			// that reference this MetaSound, so just register this asset. Otherwise,
diff --git a/Engine/Plugins/Runtime/Metasound/Source/MetasoundEditor/Private/MetasoundEditorGraphInputNode.cpp b/Engine/Plugins/Runtime/Metasound/Source/MetasoundEditor/Private/MetasoundEditorGraphInputNode.cpp
--- a/Engine/Plugins/Runtime/Metasound/Source/MetasoundEditor/Private/MetasoundEditorGraphInputNode.cpp
+++ b/Engine/Plugins/Runtime/Metasound/Source/MetasoundEditor/Private/MetasoundEditorGraphInputNode.cpp
@@ -257,7 +257,7 @@ FText UMetasoundEditorGraphInputNode::GetTooltipText() const
 									"Showing platform/platform group '{1}' fallback '{2}'.\r\n"
 									"(See 'Audition' menu or 'MetaSound' Editor Preferences to change 'Audition Platform')."),
 								FText::FromName(PageSettings->Name),
-								FText::FromName(EditorSettings->AuditionPlatform),
+								FText::FromName(EditorSettings->GetAuditionPlatform()),
 								FText::FromName(FallbackSettings->Name));
 						}
 					}
diff --git a/Engine/Plugins/Runtime/Metasound/Source/MetasoundEditor/Private/MetasoundEditorModule.cpp b/Engine/Plugins/Runtime/Metasound/Source/MetasoundEditor/Private/MetasoundEditorModule.cpp
--- a/Engine/Plugins/Runtime/Metasound/Source/MetasoundEditor/Private/MetasoundEditorModule.cpp
+++ b/Engine/Plugins/Runtime/Metasound/Source/MetasoundEditor/Private/MetasoundEditorModule.cpp
@@ -712,6 +712,8 @@ namespace Metasound::Editor
 				return;
 			}
 
+			const UMetasoundEditorSettings* EditorSettings = GetDefault<UMetasoundEditorSettings>();
+			
 			if (UMetaSoundSettings* Settings = GetMutableDefault<UMetaSoundSettings>())
 			{
 				Settings->GetOnDefaultRenamedDelegate().AddLambda([]()
@@ -724,21 +726,18 @@ namespace Metasound::Editor
 				});
 			}
 
-			const UMetasoundEditorSettings* EditorSettings = GetDefault<UMetasoundEditorSettings>();
-			FDocumentBuilderRegistry::GetChecked().GetOnResolveAuditionPageDelegate().BindUObject(EditorSettings, &UMetasoundEditorSettings::ResolveAuditionPage);
-
 			FEditorDelegates::PreBeginPIE.AddWeakLambda(EditorSettings, [this](const bool /* bSimulating */)
 			{
 				using namespace Metasound::Frontend;
 
 				if (const UMetasoundEditorSettings* EdSettings = GetDefault<UMetasoundEditorSettings>())
 				{
-					FOnResolveEditorPage& OnResolveAuditionPage = FDocumentBuilderRegistry::GetChecked().GetOnResolveAuditionPageDelegate();
-					if (OnResolveAuditionPage.IsBoundToObject(EdSettings))
+					if (UMetaSoundSettings* Settings = GetMutableDefault<UMetaSoundSettings>())
 					{
 						if (!EdSettings->bApplyAuditionSettingsInPIE)
 						{
-							OnResolveAuditionPage.Unbind();
+
+							Settings->OverrideTargetPageSettings({}); // Clear out page override if we want it disabled in PIE
 						}
 					}
 
@@ -747,12 +746,12 @@ namespace Metasound::Editor
 			});
 			FEditorDelegates::EndPIE.AddWeakLambda(EditorSettings, [](const bool /* bSimulating */)
 			{
+
 				if (const UMetasoundEditorSettings* EdSettings = GetDefault<UMetasoundEditorSettings>())
 				{
-					FOnResolveEditorPage& OnResolveAuditionPage = FDocumentBuilderRegistry::GetChecked().GetOnResolveAuditionPageDelegate();
-					if (!OnResolveAuditionPage.IsBoundToObject(EdSettings))
+					if (UMetaSoundSettings* Settings = GetMutableDefault<UMetaSoundSettings>())
 					{
-						OnResolveAuditionPage.BindUObject(EdSettings, &UMetasoundEditorSettings::ResolveAuditionPage);
+						Settings->OverrideTargetPageSettings(EdSettings->GetPageForAudition());
 					}
 				}
 			});
@@ -1027,6 +1026,14 @@ namespace Metasound::Editor
 
 			RegisterSettingsDelegates();
 			RegisterFolderBulkOperations();
+
+			// Set page auditioning overrides
+			const UMetasoundEditorSettings* EditorSettings = GetDefault<UMetasoundEditorSettings>();
+			UMetaSoundSettings* Settings = GetMutableDefault<UMetaSoundSettings>();
+			if (Settings && EditorSettings)
+			{
+				Settings->OverrideTargetPageSettings(EditorSettings->GetPageForAudition());
+			}
 		}
 
 		virtual void ShutdownModule() override
diff --git a/Engine/Plugins/Runtime/Metasound/Source/MetasoundEditor/Private/MetasoundEditorSettings.cpp b/Engine/Plugins/Runtime/Metasound/Source/MetasoundEditor/Private/MetasoundEditorSettings.cpp
--- a/Engine/Plugins/Runtime/Metasound/Source/MetasoundEditor/Private/MetasoundEditorSettings.cpp
+++ b/Engine/Plugins/Runtime/Metasound/Source/MetasoundEditor/Private/MetasoundEditorSettings.cpp
@@ -58,15 +58,85 @@ void UMetasoundEditorSettings::PostEditChangeProperty(FPropertyChangedEvent& Pro
 		Info.bUseThrobber = true;
 		FSlateNotificationManager::Get().AddNotification(Info);
 	}
+
+	// Update the page override if relevant settings have changed. 
+	if (PropertyName == GET_MEMBER_NAME_CHECKED(UMetasoundEditorSettings, AuditionPlatform)
+		|| PropertyName == GET_MEMBER_NAME_CHECKED(UMetasoundEditorSettings, AuditionPage))
+	{
+		if (UMetaSoundSettings* Settings = GetMutableDefault<UMetaSoundSettings>())
+		{
+			Settings->OverrideTargetPageSettings(GetPageForAudition());
+		}
+	}
+}
+void UMetasoundEditorSettings::SetAuditionPageMode(EAuditionPageMode InMode)
+{
+	AuditionPageMode = InMode;
+	if (UMetaSoundSettings* Settings = GetMutableDefault<UMetaSoundSettings>())
+	{
+		Settings->OverrideTargetPageSettings(GetPageForAudition());
+	}
+}
+
+EAuditionPageMode UMetasoundEditorSettings::GetAuditionPageMode() const
+{
+	return AuditionPageMode;
+}
+
+void UMetasoundEditorSettings::SetAuditionPlatform(FName InPlatform)
+{
+	AuditionPlatform = InPlatform;
+	if (UMetaSoundSettings* Settings = GetMutableDefault<UMetaSoundSettings>())
+	{
+		Settings->OverrideTargetPageSettings(GetPageForAudition());
+	}
+}
+
+FName UMetasoundEditorSettings::GetAuditionPlatform() const
+{
+	return AuditionPlatform;
+}
+
+void UMetasoundEditorSettings::SetAuditionPage(FName InPage)
+{
+	AuditionPage = InPage;
+	if (UMetaSoundSettings* Settings = GetMutableDefault<UMetaSoundSettings>())
+	{
+		Settings->OverrideTargetPageSettings(GetPageForAudition());
+	}
 }
+
+FName UMetasoundEditorSettings::GetAuditionPage() const
+{
+	return AuditionPage;
+}
+
 #endif // WITH_EDITOR
 
-Metasound::Engine::FPageResolutionEditorResults UMetasoundEditorSettings::ResolveAuditionPage(const TArray<FGuid>& InPageIDs) const
+Metasound::Engine::FTargetPageOverride UMetasoundEditorSettings::GetPageForAudition() const
+{
+	using namespace Metasound;
+	using namespace Metasound::Engine;
+
+	FTargetPageOverride AuditionOverride{ .PlatformName = AuditionPlatform, .PageID = Frontend::DefaultPageID };
+
+	if(const UMetaSoundSettings* Settings = GetDefault<UMetaSoundSettings>())
+	{
+		if (const FMetaSoundPageSettings* AuditionPageSettings = Settings->FindPageSettings(AuditionPage))
+		{
+			AuditionOverride.PageID = AuditionPageSettings->UniqueId;
+		}
+	}
+
+	return AuditionOverride;
+}
+
+Metasound::Engine::FTargetPageOverride UMetasoundEditorSettings::ResolveAuditionPage(const TArray<FGuid>& InPageIDs) const
 {
 	using namespace Metasound;
 	using namespace Metasound::Engine;
 
-	FPageResolutionEditorResults PreviewInfo { .PlatformName = AuditionPlatform };
+	FTargetPageOverride PreviewInfo { .PlatformName = AuditionPlatform };
 
 	if (AuditionPage == Frontend::DefaultPageName)
 	{
diff --git a/Engine/Plugins/Runtime/Metasound/Source/MetasoundEditor/Private/MetasoundEditorSettings.h b/Engine/Plugins/Runtime/Metasound/Source/MetasoundEditor/Private/MetasoundEditorSettings.h
--- a/Engine/Plugins/Runtime/Metasound/Source/MetasoundEditor/Private/MetasoundEditorSettings.h
+++ b/Engine/Plugins/Runtime/Metasound/Source/MetasoundEditor/Private/MetasoundEditorSettings.h
@@ -22,7 +22,7 @@ struct FMetasoundFrontendDocument;
 
 namespace Metasound::Engine
 {
-	struct FPageResolutionEditorResults;
+	struct FTargetPageOverride;
 } // namespace Metasound::Engine
 
 
@@ -135,6 +135,7 @@ public:
 	UPROPERTY(EditAnywhere, config, Category = General)
 	FString DefaultAuthor;
 
+private:
 	/* Currently set page audition mode. Set by the MetaSound Asset Editor. */
 	UPROPERTY(EditAnywhere, config, Category = "Audition (Experimental)", meta = (DisplayName = "Page Audition Mode"))
 	EAuditionPageMode AuditionPageMode = EAuditionPageMode::Focused;
@@ -158,6 +159,39 @@ public:
 		GetOptions = "MetasoundEditor.MetasoundEditorSettings.GetAuditionPageNames")
 	)
 	FName AuditionPage = Metasound::Frontend::DefaultPageName;
+public:
+#if WITH_EDITOR
+	/* Currently set page audition mode. */
+	void SetAuditionPageMode(EAuditionPageMode InMode);
+	EAuditionPageMode GetAuditionPageMode() const;
+
+	/** Name of platform to mock when previewing playback. This will limit playback
+	  * to fallback only to paged data that are cooked for the given platform.
+	  * (see project 'MetaSound' Settings --> 'Page Settings' array for order)
+	  * If set to 'Editor', ignores cook settings and allows fallback to all page.
+	  */
+	void SetAuditionPlatform(FName InPlatform);
+	
+	/** Name of platform to mock when previewing playback. This will limit playback
+	  * to fallback only to paged data that are cooked for the given platform.
+	  * (see project 'MetaSound' Settings --> 'Page Settings' array for order)
+	  * If set to 'Editor', ignores cook settings and allows fallback to all page.
+	  */
+	FName GetAuditionPlatform() const;
+	
+	/** Name of the page to audition in editor. If unimplemented on the auditioned MetaSound, uses order of cooked pages
+	  * (see project 'MetaSound' Settings --> 'Page Settings' array for order) falling back to lower index-ordered page implemented
+	  * in MetaSound asset.
+	  */
+	void SetAuditionPage(FName InPage);
+	
+	/** Name of the page to audition in editor. If unimplemented on the auditioned MetaSound, uses order of cooked pages
+	  * (see project 'MetaSound' Settings --> 'Page Settings' array for order) falling back to lower index-ordered page implemented
+	  * in MetaSound asset.
+	  */
+	FName GetAuditionPage() const;
+#endif
+
 
 	/** Maps Pin Category To Pin Color */
 	TMap<FName, FLinearColor> CustomPinTypeColors;
@@ -268,7 +302,10 @@ public:
 #endif // WITH_EDITOR
 	//~UObject
 	
-	Metasound::Engine::FPageResolutionEditorResults ResolveAuditionPage(const TArray<FGuid>& InPageIDs) const;
+	/** Returns the page override data derived from the AuditionPlatform and AuditionPage */
+	Metasound::Engine::FTargetPageOverride GetPageForAudition() const;
+
+	Metasound::Engine::FTargetPageOverride ResolveAuditionPage(const TArray<FGuid>& InPageIDs) const;
 	
 	/** Given the provided AuditionPageID, returns the resolved PageID from the provided array of values based on fallback logic. */
 	FGuid ResolveAuditionPage(const TArray<FGuid>& InPageIDs, const FGuid& InAuditionPageID) const;
diff --git a/Engine/Plugins/Runtime/Metasound/Source/MetasoundEditor/Private/MetasoundEditorSubsystem.cpp b/Engine/Plugins/Runtime/Metasound/Source/MetasoundEditor/Private/MetasoundEditorSubsystem.cpp
--- a/Engine/Plugins/Runtime/Metasound/Source/MetasoundEditor/Private/MetasoundEditorSubsystem.cpp
+++ b/Engine/Plugins/Runtime/Metasound/Source/MetasoundEditor/Private/MetasoundEditorSubsystem.cpp
@@ -398,7 +398,7 @@ bool UMetaSoundEditorSubsystem::IsPageAuditionPlatformCookTarget(const FGuid& In
 			{
 				bIsAuditionable |= PlatformTargetPageID == InPageID;
 			};
-			Settings->IterateCookedTargetPageIDs(EditorSettings->AuditionPlatform, PageIsTargetable);
+			Settings->IterateCookedTargetPageIDs(EditorSettings->GetAuditionPlatform(), PageIsTargetable);
 			return bIsAuditionable;
 		}
 	}
@@ -484,12 +484,12 @@ bool UMetaSoundEditorSubsystem::SetFocusedPageInternal(FName PageName, const FGu
 	// audition target page.
 	UMetasoundEditorSettings* EditorSettings = GetMutableDefault<UMetasoundEditorSettings>();
 	check(EditorSettings);
-	if (EditorSettings->AuditionPageMode == EAuditionPageMode::Focused)
+	if (EditorSettings->GetAuditionPageMode() == EAuditionPageMode::Focused)
 	{
-		if (EditorSettings->AuditionPage != PageName)
+		if (EditorSettings->GetAuditionPage() != PageName)
 		{
 			EditorSettings->Modify();
-			EditorSettings->AuditionPage = PageName;
+			EditorSettings->SetAuditionPage(PageName);
 			bAuditionPageSet = true;
 		}
 	}
diff --git a/Engine/Plugins/Runtime/Metasound/Source/MetasoundEditor/Private/SMetasoundStats.cpp b/Engine/Plugins/Runtime/Metasound/Source/MetasoundEditor/Private/SMetasoundStats.cpp
--- a/Engine/Plugins/Runtime/Metasound/Source/MetasoundEditor/Private/SMetasoundStats.cpp
+++ b/Engine/Plugins/Runtime/Metasound/Source/MetasoundEditor/Private/SMetasoundStats.cpp
@@ -221,12 +221,12 @@ namespace Metasound::Editor
 						if (!Settings->GetProjectPageSettings().IsEmpty())
 						{
 							const FText AuditionPageHeader = LOCTEXT("AuditionPageActive_DebugFormat", "Auditioning Page: {0}");
-							AuditionPage = FText::Format(AuditionPageHeader, FText::FromName(EditorSettings->AuditionPage));
+							AuditionPage = FText::Format(AuditionPageHeader, FText::FromName(EditorSettings->GetAuditionPage()));
 
-							if (EditorSettings->AuditionPlatform != "Editor")
+							if (EditorSettings->GetAuditionPlatform() != "Editor")
 							{
 								const FText AuditionPlatformHeader = LOCTEXT("AuditionPlatformActive_DebugFormat", "Auditioning Platform: {0}");
-								AuditionPlatform = FText::Format(AuditionPlatformHeader, FText::FromName(EditorSettings->AuditionPlatform));
+								AuditionPlatform = FText::Format(AuditionPlatformHeader, FText::FromName(EditorSettings->GetAuditionPlatform()));
 							}
 						}
 					}
diff --git a/Engine/Plugins/Runtime/Metasound/Source/MetasoundEngine/Private/Metasound.cpp b/Engine/Plugins/Runtime/Metasound/Source/MetasoundEngine/Private/Metasound.cpp
--- a/Engine/Plugins/Runtime/Metasound/Source/MetasoundEngine/Private/Metasound.cpp
+++ b/Engine/Plugins/Runtime/Metasound/Source/MetasoundEngine/Private/Metasound.cpp
@@ -26,6 +26,7 @@
 #include "MetasoundParameterTransmitter.h"
 #include "MetasoundPrimitives.h"
 #include "MetasoundReceiveNode.h"
+#include "MetasoundSettings.h"
 #include "MetasoundTrigger.h"
 #include "MetasoundUObjectRegistry.h"
 #include "UObject/AssetRegistryTagsContext.h"
@@ -45,6 +46,26 @@
 
 #define LOCTEXT_NAMESPACE "MetaSound"
 
+namespace Metasound::MetasoundPrivate
+{
+	Frontend::FMetaSoundAssetRegistrationOptions GetInitRegistrationOptions()
+	{
+		Frontend::FMetaSoundAssetRegistrationOptions RegOptions;
+		RegOptions.bForceReregister = false;
+#if !WITH_EDITOR 
+		// When without editor, don't AutoUpdate or remove template nodes at runtime. This only happens at cook or save.
+		// When with editor, those are needed because sounds are not necessarily saved before previewing.
+		RegOptions.bAutoUpdate = false;
+#endif // !WITH_EDITOR
+		if (const UMetaSoundSettings* Settings = GetDefault<UMetaSoundSettings>())
+		{
+			RegOptions.bAutoUpdateLogWarningOnDroppedConnection = Settings->bAutoUpdateLogWarningOnDroppedConnection;
+			RegOptions.PageOrder = Settings->GetPageOrder();
+		}
+
+		return RegOptions;
+	}
+}
 
 int32 UMetasoundEditorGraphBase::GetHighestMessageSeverity() const
 {
@@ -211,6 +232,17 @@ void UMetaSoundPatch::SetReferencedAssets(TSet<Metasound::Frontend::IMetaSoundAs
 }
 #endif
 
+TSharedPtr<Audio::IProxyData> UMetaSoundPatch::CreateProxyData(const Audio::FProxyDataInitParams& InitParams)
+{
+	if (!IsRegistered())
+	{
+		// Make sure the metasound is registered and ready to go before making a proxy.
+		UpdateAndRegisterForExecution(Metasound::MetasoundPrivate::GetInitRegistrationOptions());
+	}
+
+	return FMetasoundAssetBase::CreateProxyData(InitParams);
+}
+
 TArray<FMetasoundAssetBase*> UMetaSoundPatch::GetReferencedAssets()
 {
 	return Metasound::Engine::FAssetHelper::GetReferencedAssets(*this);
diff --git a/Engine/Plugins/Runtime/Metasound/Source/MetasoundEngine/Private/MetasoundAssetSubsystem.cpp b/Engine/Plugins/Runtime/Metasound/Source/MetasoundEngine/Private/MetasoundAssetSubsystem.cpp
--- a/Engine/Plugins/Runtime/Metasound/Source/MetasoundEngine/Private/MetasoundAssetSubsystem.cpp
+++ b/Engine/Plugins/Runtime/Metasound/Source/MetasoundEngine/Private/MetasoundAssetSubsystem.cpp
@@ -230,6 +230,20 @@ namespace Metasound::Engine
 		return AddOrUpdateFromObjectInternal(InObject);
 	}
 
+	Frontend::FMetaSoundAssetRegistrationOptions FMetaSoundAssetManager::GetRegistrationOptions() const
+	{
+		Frontend::FMetaSoundAssetRegistrationOptions RegOptions;
+
+		if (const UMetaSoundSettings* Settings = GetDefault<UMetaSoundSettings>())
+		{
+			RegOptions.bAutoUpdateLogWarningOnDroppedConnection = Settings->bAutoUpdateLogWarningOnDroppedConnection;
+			RegOptions.PageOrder = Settings->GetPageOrder();
+		}
+
+		// nrvo
+		return RegOptions;
+	}
+
 	FMetaSoundAssetKey FMetaSoundAssetManager::AddOrUpdateFromObjectInternal(const UObject& InObject)
 	{
 		using namespace AssetSubsystemPrivate;
@@ -778,7 +792,7 @@ namespace Metasound::Engine
 			AssetSubsystemPrivate::AddClassInfo(TagDataMapCriticalSection, ClassInfoMap, FMetaSoundAssetClassInfo(*DocInterface));
 		}
 
-		AssetBase->UpdateAndRegisterForExecution();
+		AssetBase->UpdateAndRegisterForExecution(GetRegistrationOptions());
 		return true;
 
 #else // !WITH_EDITORONLY_DATA
@@ -933,18 +947,14 @@ namespace Metasound::Engine
 		TArray<FDirectoryPath> Directories;
 		Algo::Transform(InDirectories, Directories, [](const FMetaSoundAssetDirectory& AssetDir) { return AssetDir.Directory; });
 
-		SearchAndIterateDirectoryAssets(Directories, [this](const FAssetData& AssetData)
+		Frontend::FMetaSoundAssetRegistrationOptions RegOptions = GetRegistrationOptions();
+
+		SearchAndIterateDirectoryAssets(Directories, [this, RegOptions](const FAssetData& AssetData)
 		{
-			AddOrLoadAndUpdateFromObjectAsync(AssetData, [](FMetaSoundAssetKey, UObject& AssetObject)
+			AddOrLoadAndUpdateFromObjectAsync(AssetData, [RegOptions](FMetaSoundAssetKey, UObject& AssetObject)
 			{
 				FMetasoundAssetBase* MetaSoundAsset = Metasound::IMetasoundUObjectRegistry::Get().GetObjectAsAssetBase(&AssetObject);
 				check(MetaSoundAsset);
-
-				Frontend::FMetaSoundAssetRegistrationOptions RegOptions;
-				if (const UMetaSoundSettings* Settings = GetDefault<UMetaSoundSettings>())
-				{
-					RegOptions.bAutoUpdateLogWarningOnDroppedConnection = Settings->bAutoUpdateLogWarningOnDroppedConnection;
-				}
 				MetaSoundAsset->UpdateAndRegisterForExecution(RegOptions);
 			});
 		});
@@ -1130,11 +1140,13 @@ namespace Metasound::Engine
 			}
 		}
 
+		Frontend::FMetaSoundAssetRegistrationOptions RegOptions = GetRegistrationOptions();
+
 		// Handled in second loop to avoid re-registering referenced graphs more than once
 		for (FMetasoundAssetBase* AssetToReregister : ToReregister)
 		{
 			check(AssetToReregister);
-			AssetToReregister->UpdateAndRegisterForExecution();
+			AssetToReregister->UpdateAndRegisterForExecution(RegOptions);
 		}
 	}
 
@@ -1446,20 +1458,15 @@ void UMetaSoundAssetSubsystem::PostInitAssetScanInternal()
 	if (ensureAlways(Settings))
 	{
 		FMetaSoundAssetManager& Manager = FMetaSoundAssetManager::GetChecked();
-		Manager.SearchAndIterateDirectoryAssets(Settings->DirectoriesToRegister, [&Manager](const FAssetData& AssetData)
+		FMetaSoundAssetRegistrationOptions RegOptions = Manager.GetRegistrationOptions();
+		Manager.SearchAndIterateDirectoryAssets(Settings->DirectoriesToRegister, [&Manager, &RegOptions](const FAssetData& AssetData)
 		{
 			using namespace Metasound;
 
-			Manager.AddOrLoadAndUpdateFromObjectAsync(AssetData, [](FMetaSoundAssetKey, UObject& AssetObject)
+			Manager.AddOrLoadAndUpdateFromObjectAsync(AssetData, [RegOptions](FMetaSoundAssetKey, UObject& AssetObject)
 			{
 				FMetasoundAssetBase* MetaSoundAsset = IMetasoundUObjectRegistry::Get().GetObjectAsAssetBase(&AssetObject);
 				check(MetaSoundAsset);
-
-				Frontend::FMetaSoundAssetRegistrationOptions RegOptions;
-				if (const UMetaSoundSettings* Settings = GetDefault<UMetaSoundSettings>())
-				{
-					RegOptions.bAutoUpdateLogWarningOnDroppedConnection = Settings->bAutoUpdateLogWarningOnDroppedConnection;
-				}
 				MetaSoundAsset->UpdateAndRegisterForExecution(RegOptions);
 			});
 		});
diff --git a/Engine/Plugins/Runtime/Metasound/Source/MetasoundEngine/Private/MetasoundBuilderSubsystem.cpp b/Engine/Plugins/Runtime/Metasound/Source/MetasoundEngine/Private/MetasoundBuilderSubsystem.cpp
--- a/Engine/Plugins/Runtime/Metasound/Source/MetasoundEngine/Private/MetasoundBuilderSubsystem.cpp
+++ b/Engine/Plugins/Runtime/Metasound/Source/MetasoundEngine/Private/MetasoundBuilderSubsystem.cpp
@@ -20,6 +20,7 @@
 #include "MetasoundFrontendDocument.h"
 #include "MetasoundFrontendDocumentIdGenerator.h"
 #include "MetasoundFrontendGraphBuilder.h"
+#include "MetasoundFrontendPages.h"
 #include "MetasoundFrontendRegistries.h"
 #include "MetasoundFrontendSearchEngine.h"
 #include "MetasoundFrontendTransform.h"
@@ -75,9 +76,7 @@ namespace Metasound::Engine
 			}
 			else if (const TArray<FMetasoundFrontendClassInputDefault>* ClassDefaults = Builder.FindNodeClassInputDefaults(InNodeID, VertexName))
 			{
-				const FGuid PageID = FDocumentBuilderRegistry::GetChecked().ResolveTargetPageID(*ClassDefaults);
-				auto MatchesPageID = [&PageID](const FMetasoundFrontendClassInputDefault& InputDefault) { return InputDefault.PageID == PageID; };
-				if (const FMetasoundFrontendClassInputDefault* ClassDefault = ClassDefaults->FindByPredicate(MatchesPageID))
+				if (const FMetasoundFrontendClassInputDefault* ClassDefault = Frontend::FindPreferredPage(*ClassDefaults, UMetaSoundSettings::GetPageOrder()))
 				{
 					DefaultLiteral = ClassDefault->Literal;
 				}
@@ -378,7 +377,13 @@ void UMetaSoundSourceBuilder::InitTargetPageDelegates(Metasound::Frontend::FDocu
 	});
 
 	const IMetaSoundDocumentInterface& DocInterface = GetConstBuilder().GetConstDocumentInterfaceChecked();
-	TargetPageID = FDocumentBuilderRegistry::GetChecked().ResolveTargetPageID(DocInterface.GetConstDocument().RootGraph);
+
+	// Determine which graph page is being worked on. 
+	TargetPageID = Metasound::Frontend::DefaultPageID;
+	if (const FMetasoundFrontendGraph* GraphPage = Metasound::Frontend::FindPreferredPage(DocInterface.GetConstDocument().RootGraph.GetConstGraphPages(), UMetaSoundSettings::GetPageOrder()))
+	{
+		TargetPageID = GraphPage->PageID;
+	}
 
 	FEdgeModifyDelegates& EdgeDelegates = OutDocumentDelegates.FindEdgeDelegatesChecked(TargetPageID);
 	EdgeDelegates.OnEdgeAdded.AddUObject(this, &UMetaSoundSourceBuilder::OnEdgeAdded);
@@ -450,7 +455,7 @@ void UMetaSoundSourceBuilder::OnInputAdded(int32 InputIndex)
 
 		constexpr bool bCreateUObjectProxies = true;
 		UMetaSoundSource& Source = GetMetaSoundSource();
-		Source.RuntimeInputData.InputMap.Add(NewInput.Name, UMetaSoundSource::CreateRuntimeInput(IDataTypeRegistry::Get(), NewInput, bCreateUObjectProxies));
+		Source.RuntimeInputData.InputMap.Add(NewInput.Name, UMetaSoundSource::CreateRuntimeInput(IDataTypeRegistry::Get(), NewInput, bCreateUObjectProxies, UMetaSoundSettings::GetPageOrder()));
 
 		for (uint64 AudioComponentID : LiveComponentIDs)
 		{
diff --git a/Engine/Plugins/Runtime/Metasound/Source/MetasoundEngine/Private/MetasoundDocumentBuilderRegistry.cpp b/Engine/Plugins/Runtime/Metasound/Source/MetasoundEngine/Private/MetasoundDocumentBuilderRegistry.cpp
--- a/Engine/Plugins/Runtime/Metasound/Source/MetasoundEngine/Private/MetasoundDocumentBuilderRegistry.cpp
+++ b/Engine/Plugins/Runtime/Metasound/Source/MetasoundEngine/Private/MetasoundDocumentBuilderRegistry.cpp
@@ -415,16 +415,20 @@ namespace Metasound::Engine
 	}
 
 #if WITH_EDITOR
+PRAGMA_DISABLE_DEPRECATION_WARNINGS
 	FOnResolveEditorPage& FDocumentBuilderRegistry::GetOnResolveAuditionPageDelegate()
 	{
 		return OnResolveAuditionPage;
 	}
+PRAGMA_ENABLE_DEPRECATION_WARNINGS
 #endif // WITH_EDITOR
 
+PRAGMA_DISABLE_DEPRECATION_WARNINGS
 	FOnResolvePage& FDocumentBuilderRegistry::GetOnResolveProjectPageOverrideDelegate()
 	{
 		return OnResolveProjectPage;
 	}
+PRAGMA_ENABLE_DEPRECATION_WARNINGS
 
 	bool FDocumentBuilderRegistry::ReloadBuilder(const FMetasoundFrontendClassName& InClassName) const
 	{
@@ -463,7 +467,9 @@ namespace Metasound::Engine
 			TargetPageResolveScratch.Add(PageGraph.PageID);
 		});
 
-		return ResolveTargetPageIDInternal(TargetPageResolveScratch);
+PRAGMA_DISABLE_DEPRECATION_WARNINGS
+		return ResolveTargetPageID(TargetPageResolveScratch);
+PRAGMA_ENABLE_DEPRECATION_WARNINGS
 	}
 
 	FGuid FDocumentBuilderRegistry::ResolveTargetPageID(const FMetasoundFrontendClassInput& InClassInput) const
@@ -490,7 +496,9 @@ namespace Metasound::Engine
 			TargetPageResolveScratch.Add(PageID);
 		});
 
-		return ResolveTargetPageIDInternal(TargetPageResolveScratch);
+PRAGMA_DISABLE_DEPRECATION_WARNINGS
+		return ResolveTargetPageID(TargetPageResolveScratch);
+PRAGMA_ENABLE_DEPRECATION_WARNINGS
 	}
 
 	FGuid FDocumentBuilderRegistry::ResolveTargetPageID(const TArray<FMetasoundFrontendClassInputDefault>& InClassDefaults) const
@@ -512,17 +520,21 @@ namespace Metasound::Engine
 		FScopeLock Lock(&TargetPageResolveScratchCritSec);
 		TargetPageResolveScratch.Reset();
 		Algo::Transform(InClassDefaults, TargetPageResolveScratch, [](const FMetasoundFrontendClassInputDefault& ClassDefault) { return ClassDefault.PageID; });
-		return ResolveTargetPageIDInternal(TargetPageResolveScratch);
+PRAGMA_DISABLE_DEPRECATION_WARNINGS
+		return ResolveTargetPageID(TargetPageResolveScratch);
+PRAGMA_ENABLE_DEPRECATION_WARNINGS
 	}
 
-	FGuid FDocumentBuilderRegistry::ResolveTargetPageIDInternal(const TArray<FGuid>& InPageIDsToResolve) const
+	FGuid FDocumentBuilderRegistry::ResolveTargetPageID(const TArray<FGuid>& InPageIDsToResolve) const
 	{
 		FName PlatformName = FPlatformProperties::IniPlatformName();
 
 #if WITH_EDITOR
+PRAGMA_DISABLE_DEPRECATION_WARNINGS
 		if (OnResolveAuditionPage.IsBound())
 		{
 			FPageResolutionEditorResults PreviewInfo = OnResolveAuditionPage.Execute(InPageIDsToResolve);
+PRAGMA_ENABLE_DEPRECATION_WARNINGS
 			if (PreviewInfo.PageID.IsSet())
 			{
 				return PreviewInfo.PageID.GetValue();
@@ -532,9 +544,11 @@ namespace Metasound::Engine
 		}
 #endif // WITH_EDITOR
 
+PRAGMA_DISABLE_DEPRECATION_WARNINGS
 		if (OnResolveProjectPage.IsBound())
 		{
 			const FGuid ResolvedPageID = OnResolveProjectPage.Execute(InPageIDsToResolve);
+PRAGMA_ENABLE_DEPRECATION_WARNINGS
 			check (InPageIDsToResolve.Contains(ResolvedPageID));
 			return ResolvedPageID;
 		}
@@ -548,6 +562,34 @@ namespace Metasound::Engine
 		return Frontend::DefaultPageID;
 	}
 
+	/** These methods exist to support deprecated funtions. It allows MetaSound Settings data to be
+	 * accessed in the MetasoundFrontend module via a singleton IDocumentBuilderRegistry::Get(). Once
+	 * the deprecated callsites are removed, these methods can be removed as well. */
+	TArrayView<const FGuid> FDocumentBuilderRegistry::GetPageOrder() const
+	{
+		return UMetaSoundSettings::GetPageOrder();
+	}
+
+#if WITH_EDITORONLY_DATA
+	TArray<FGuid> FDocumentBuilderRegistry::GetCookedTargetPages(FName InPlatformName) const
+	{
+		if (const UMetaSoundSettings* Settings = GetDefault<UMetaSoundSettings>(); ensure(Settings))
+		{
+			return Settings->GetCookedTargetPageIDs(InPlatformName);
+		}
+		return TArray<FGuid>{Frontend::DefaultPageID};
+	}
+
+	TArray<FGuid> FDocumentBuilderRegistry::GetCookedPageOrder(FName InPlatformName) const
+	{
+		if (const UMetaSoundSettings* Settings = GetDefault<UMetaSoundSettings>(); ensure(Settings))
+		{
+			return Settings->GetCookedPageOrder(InPlatformName);
+		}
+		return TArray<FGuid>{Frontend::DefaultPageID};
+	}
+#endif // WITH_EDITORONLY_DATA
+
 	FGuid FDocumentBuilderRegistry::ResolveTargetPageIDInternal(const UMetaSoundSettings& Settings, const TArray<FGuid>& InPageIDsToResolve, const FGuid& TargetPageID, FName PlatformName) const
 	{
 		bool bResolved = false;
diff --git a/Engine/Plugins/Runtime/Metasound/Source/MetasoundEngine/Private/MetasoundEngineAsset.h b/Engine/Plugins/Runtime/Metasound/Source/MetasoundEngine/Private/MetasoundEngineAsset.h
--- a/Engine/Plugins/Runtime/Metasound/Source/MetasoundEngine/Private/MetasoundEngineAsset.h
+++ b/Engine/Plugins/Runtime/Metasound/Source/MetasoundEngine/Private/MetasoundEngineAsset.h
@@ -266,7 +266,17 @@ namespace Metasound::Engine
 				}
 				const bool bIsDeterministic = SerializationRequiresDeterminism(bIsCooking);
 				FDocumentIDGenerator::FScopeDeterminism DeterminismScope(bIsDeterministic);
-				InMetaSound.UpdateAndRegisterForSerialization(PlatformName);
+
+
+				FMetaSoundAssetCookOptions CookOptions;
+				if (const UMetaSoundSettings* Settings = GetDefault<UMetaSoundSettings>(); ensure(Settings))
+				{
+					CookOptions.bStripUnusedPages = true;
+					CookOptions.PagesToTarget = Settings->GetCookedTargetPageIDs(PlatformName);
+					CookOptions.PageOrder = Settings->GetCookedPageOrder(PlatformName);
+				}
+				
+				InMetaSound.UpdateAndRegisterForSerialization(CookOptions);
 			}
  			else if (FApp::CanEverRenderAudio())
 			{
diff --git a/Engine/Plugins/Runtime/Metasound/Source/MetasoundEngine/Private/MetasoundSettings.cpp b/Engine/Plugins/Runtime/Metasound/Source/MetasoundEngine/Private/MetasoundSettings.cpp
--- a/Engine/Plugins/Runtime/Metasound/Source/MetasoundEngine/Private/MetasoundSettings.cpp
+++ b/Engine/Plugins/Runtime/Metasound/Source/MetasoundEngine/Private/MetasoundSettings.cpp
@@ -2,6 +2,7 @@
 
 #include "MetasoundSettings.h"
 
+#include "Algo/AnyOf.h"
 #include "Algo/Count.h"
 #include "HAL/IConsoleManager.h"
 #include "MetasoundFrontendDocument.h"
@@ -234,6 +235,7 @@ void UMetaSoundSettings::ConformPageSettings(bool bNotifyDefaultRenamed)
 }
 #endif // WITH_EDITOR
 
+
 const FMetaSoundPageSettings* UMetaSoundSettings::FindPageSettings(FName Name) const
 {
 	if (Name == Metasound::Frontend::DefaultPageName)
@@ -284,15 +286,30 @@ TArray<FName> UMetaSoundSettings::GetAllPlatformNamesImplementingTargets() const
 TArray<FGuid> UMetaSoundSettings::GetCookedTargetPageIDs(FName PlatformName) const
 {
 	FScopeLock Lock(&CookPlatformTargetCritSec);
-	return GetCookedTargetPageIDsInternal(PlatformName);
+	CacheCookedPageIDs(PlatformName);
+	return CookPlatformTargetPageIDs;
+}
+
+TArray<FGuid> UMetaSoundSettings::GetCookedPageOrder(FName PlatformName) const
+{
+	FScopeLock Lock(&CookPlatformTargetCritSec);
+	CacheCookedPageIDs(PlatformName);
+	return CookPlatformPageOrder;
 }
 
-const TArray<FGuid>& UMetaSoundSettings::GetCookedTargetPageIDsInternal(FName PlatformName) const
+void UMetaSoundSettings::CacheCookedPageIDs(FName PlatformName) const
 {
-	if ((PlatformName != CookPlatformTargetPage) || (PlatformName.IsNone() && CookPlatformTargetPageIDs.IsEmpty()))
+	/* Cache the CookPlatformTargetPageIDs and CookPlatformPageOrder to 
+	 * avoid rebuilding these for every asset being cooked. */
+
+	if ((PlatformName != CookPlatformTargetPage) || (PlatformName.IsNone() && CookPlatformTargetPageIDs.IsEmpty() && CookPlatformPageOrder.IsEmpty()))
 	{
 		CookPlatformTargetPage = PlatformName;
 		CookPlatformTargetPageIDs.Reset();
+		CookPlatformPageOrder.Reset();
+
+		/* Cache CookPlatformTargetPageIDs by finding all target pages for the 
+		 * given platform. */
 		auto CanTargetPage = [&PlatformName](const FMetaSoundPageSettings& PageSetting)
 		{
 #if WITH_EDITOR
@@ -333,42 +350,120 @@ const TArray<FGuid>& UMetaSoundSettings::GetCookedTargetPageIDsInternal(FName Pl
 
 			CookPlatformTargetPageIDs.Add(Metasound::Frontend::DefaultPageID);
 		}
-	}
 
-	return CookPlatformTargetPageIDs;
+
+		/* Cache cook target page order by filtering out all pages that are
+		 * excluded from the platform and filtering out all pages that come before
+		 * any of the cook target page IDs. */
+		constexpr bool bReverse = true;
+		bool bFoundAnyTargetPage = false;
+		IteratePageSettings([&](const FMetaSoundPageSettings& PageSettings) 
+			{
+#if WITH_EDITOR
+				// Skip this page if it explicitly excluded from cooking.
+				if (PageSettings.GetExcludeFromCook(PlatformName))
+				{
+					return;
+				}
+#endif // !WITH_EDITOR
+				
+				// Skip any pages that occur before the any of the target pages.
+				if (!bFoundAnyTargetPage)
+				{
+					// Sets flag to true if any of the CookPlatformTargetPageIDs equal PageSettings.PageID
+					bFoundAnyTargetPage = Algo::AnyOf(CookPlatformTargetPageIDs, [&PageSettings](const FGuid& InTargetPageID) { return InTargetPageID == PageSettings.UniqueId; });
+				}
+
+				if (bFoundAnyTargetPage)
+				{
+					CookPlatformPageOrder.Add(PageSettings.UniqueId);
+				}
+
+			}, bReverse);
+
+		if (CookPlatformPageOrder.IsEmpty())
+		{
+#if WITH_EDITOR
+			// Allow default page to be in page order if no platform name is provided.
+			// This can occur when generating reflection code for a MetaSound
+			// in a cloud cooker. 
+			const bool bCanCookDefault = PlatformName.IsNone() || !DefaultPageSettings.ExcludeFromCook.GetValueForPlatform(PlatformName);
+#else // !WITH_EDITOR
+			const bool bCanCookDefault = !DefaultPageSettings.ExcludeFromCook.GetValue();
+#endif // !WITH_EDITOR
+
+			if (!PageSettings.IsEmpty() && !bCanCookDefault)
+			{
+				UE_LOG(LogMetaSound, Warning, TEXT("No pages set to be cook for platform '%s', forcing 'Default' page as only cooked page"), *PlatformName.ToString());
+			}
+
+			CookPlatformPageOrder.Add(Metasound::Frontend::DefaultPageID);
+		}
+	}
 }
 
 void UMetaSoundSettings::IterateCookedTargetPageIDs(FName PlatformName, TFunctionRef<void(const FGuid&)> Iter) const
 {
 	FScopeLock Lock(&CookPlatformTargetCritSec);
-	const TArray<FGuid>& CookPlatforms = GetCookedTargetPageIDsInternal(PlatformName);
-	for (const FGuid& CookPlatform : CookPlatforms)
+	CacheCookedPageIDs(PlatformName);
+
+	for (const FGuid& CookPlatform : CookPlatformTargetPageIDs)
 	{
 		Iter(CookPlatform);
 	}
 }
 #endif // WITH_EDITORONLY_DATA
 
-const FMetaSoundPageSettings& UMetaSoundSettings::GetTargetPageSettings() const
+TArray<FGuid> UMetaSoundSettings::GeneratePageOrder() const
 {
-	const FName TargetPage = TargetPageNameOverride.IsSet() ? *TargetPageNameOverride : TargetPageName;
+#if WITH_EDITOR
+	if (PageOverrideData.IsSet())
+	{
+		return GeneratePageOrderFromOverride();
+	}
+#endif // WITH_EDITOR
 
-#if !NO_LOGGING
+	FName PlatformName = FPlatformProperties::IniPlatformName();
+
+	return GeneratePageOrderInternal(PlatformName, GetTargetPageSettings().UniqueId);
+};
+
+TArrayView<const FGuid> UMetaSoundSettings::GetPageOrder()
+{
+	check(IsInGameThread());
+	if (const UMetaSoundSettings* Settings = GetDefault<UMetaSoundSettings>())
+	{
+		return Settings->CachedPageOrder;
+	}
+	return {};
+}
+
+const FMetaSoundPageSettings& UMetaSoundSettings::GetTargetPageSettings() const
+{
 	auto WarnIfUninitialized = [this](const FMetaSoundPageSettings& SettingsSet)
 	{
+#if !NO_LOGGING
 		if (bWarnAccessBeforeInit)
 		{
 			UE_LOG(LogMetaSound, Display, TEXT("Target Page Settings accessed prior to 'PostInitProperties' being called.  Uninitialized PageSettings '%s' being returned."), *SettingsSet.Name.ToString());
 			bWarnAccessBeforeInit = false;
 		}
-	};
 #endif // !NO_LOGGING
+	};
+
+	const FName TargetPage = TargetPageNameOverride.IsSet() ? *TargetPageNameOverride : TargetPageName;
+
+#if WITH_EDITOR
+	if (const FMetaSoundPageSettings* TargetSettings = GetTargetPageSettingsFromOverride(TargetPage))
+	{
+		WarnIfUninitialized(*TargetSettings);
+		return *TargetSettings;
+	}
+#endif // WITH_EDITOR
 
 	if (const FMetaSoundPageSettings* TargetSettings = FindPageSettings(TargetPage))
 	{
-#if !NO_LOGGING
 		WarnIfUninitialized(*TargetSettings);
-#endif // !NO_LOGGING
 		return *TargetSettings;
 	}
 
@@ -376,15 +471,11 @@ const FMetaSoundPageSettings& UMetaSoundSettings::GetTargetPageSettings() const
 	// try and return any page setting set as a valid target.
 	if (PageSettings.Num() > 0)
 	{
-#if !NO_LOGGING
 		WarnIfUninitialized(PageSettings[0]);
-#endif // !NO_LOGGING
 		return PageSettings[0];
 	}
 
-#if !NO_LOGGING
 	WarnIfUninitialized(DefaultPageSettings);
-#endif // !NO_LOGGING
 	return DefaultPageSettings;
 }
 
@@ -413,9 +504,68 @@ void UMetaSoundSettings::PostEditChangeProperty(FPropertyChangedEvent& PropertyC
 	{
 		OnPageSettingsUpdated.Broadcast();
 	}
+	UpdateCachedPageOrder();
 
 	DenyListCacheChangeID++;
 }
+
+const FMetaSoundPageSettings* UMetaSoundSettings::GetTargetPageSettingsFromOverride(const FName InTargetPage) const
+{
+	const FMetaSoundPageSettings* TargetSettings = nullptr;
+
+	if (PageOverrideData.IsSet())
+	{
+		// The override can optionally reference a specific PageID. In that case,
+		// the specific page settings are returned. 
+		if (PageOverrideData->PageID.IsSet())
+		{
+			TargetSettings = FindPageSettings(*PageOverrideData->PageID);
+			if (TargetSettings)
+			{
+				return TargetSettings;
+			}
+		}
+
+		// If the PageID is not explicitly overridden, then search the page settings
+		// to find a page that is supported by the platform. 
+		constexpr bool bReverse = true;
+		bool bFoundTargetPage = false;
+		bool bFoundPageSupportedByPlatform = false;
+
+		IteratePageSettings([&](const FMetaSoundPageSettings& PageSettings)
+		{
+			bFoundTargetPage |= PageSettings.Name == InTargetPage;
+			if (bFoundTargetPage && !bFoundPageSupportedByPlatform)
+			{
+				const bool bIsCooked = !PageSettings.GetExcludeFromCook(PageOverrideData->PlatformName);
+				if (bIsCooked)
+				{
+					bFoundPageSupportedByPlatform = true;
+					TargetSettings = &PageSettings;
+				}
+			}
+		}, bReverse);
+	}
+
+	return TargetSettings;
+}
+
+TArray<FGuid> UMetaSoundSettings::GeneratePageOrderFromOverride() const
+{
+	check(PageOverrideData.IsSet());
+	FGuid TargetPageID;
+
+	if (PageOverrideData->PageID.IsSet())
+	{
+		TargetPageID = *PageOverrideData->PageID;
+	}
+	else
+	{
+		TargetPageID = GetTargetPageSettings().UniqueId;
+	}
+
+	return GeneratePageOrderInternal(PageOverrideData->PlatformName, TargetPageID);
+}
 #endif // WITH_EDITOR
 
 void UMetaSoundSettings::PostInitProperties()
@@ -449,6 +599,8 @@ void UMetaSoundSettings::PostInitProperties()
 			TargetPageName = PageSettings.Last().Name;
 		}
 	}
+
+	UpdateCachedPageOrder();
 }
 
 bool UMetaSoundSettings::SetTargetPage(FName PageName)
@@ -490,6 +642,13 @@ FName UMetaSoundSettings::GetQualitySettingPropertyName()
 #endif // WITH_EDITORONLY_DATA
 
 #if WITH_EDITOR
+
+void UMetaSoundSettings::OverrideTargetPageSettings(TOptional<Metasound::Engine::FTargetPageOverride> InOverride)
+{
+	PageOverrideData = MoveTemp(InOverride);
+	UpdateCachedPageOrder();
+}
+
 TArray<FName> UMetaSoundSettings::GetPageNames()
 {
 	if (const UMetaSoundSettings* Settings = GetDefault<UMetaSoundSettings>())
@@ -519,6 +678,47 @@ TArray<FName> UMetaSoundSettings::GetQualityNames()
 }
 #endif // WITH_EDITOR
 
+void UMetaSoundSettings::UpdateCachedPageOrder()
+{
+	CachedPageOrder = GeneratePageOrder();
+}
+
+TArray<FGuid> UMetaSoundSettings::GeneratePageOrderInternal(const FName& InPlatformName, const FGuid& InTargetPageID) const
+{
+	// Reverse iterate over page settings to build the priority stack. Every page
+	// that exists before the TargetPageID is discarded. Ever page after the TargetPageID
+	// is a fallback in case the TargetPageID does not exist for a given element. 
+	TArray<FGuid> PageOrder;
+
+	constexpr bool bReverse = true;
+	bool bFoundTargetPage = false;
+
+	IteratePageSettings([&](const FMetaSoundPageSettings& PageSettings)
+	{
+		bFoundTargetPage |= PageSettings.UniqueId == InTargetPageID;
+		if (bFoundTargetPage)
+		{
+#if WITH_EDITOR
+			const bool bIsCooked = !PageSettings.GetExcludeFromCook(InPlatformName);
+			if (bIsCooked)
+			{
+				PageOrder.Add(PageSettings.UniqueId);
+			}
+#else
+			PageOrder.Add(PageSettings.UniqueId);
+#endif // WITH_EDITOR
+		}
+	}, bReverse);
+
+	if (PageOrder.Num() < 1)
+	{
+		PageOrder.Add(Metasound::Frontend::DefaultPageID);
+	}
+
+	// rvo
+	return PageOrder;
+}
+
 void UMetaSoundSettings::IteratePageSettings(TFunctionRef<void(const FMetaSoundPageSettings&)> Iter, bool bReverse) const
 {
 	if (bReverse)
diff --git a/Engine/Plugins/Runtime/Metasound/Source/MetasoundEngine/Private/MetasoundSource.cpp b/Engine/Plugins/Runtime/Metasound/Source/MetasoundEngine/Private/MetasoundSource.cpp
--- a/Engine/Plugins/Runtime/Metasound/Source/MetasoundEngine/Private/MetasoundSource.cpp
+++ b/Engine/Plugins/Runtime/Metasound/Source/MetasoundEngine/Private/MetasoundSource.cpp
@@ -27,6 +27,7 @@
 #include "MetasoundFrontendDataTypeRegistry.h"
 #include "MetasoundFrontendDocumentBuilder.h"
 #include "MetasoundFrontendDocumentIdGenerator.h"
+#include "MetasoundFrontendPages.h"
 #include "MetasoundFrontendQuery.h"
 #include "MetasoundFrontendQuerySteps.h"
 #include "MetasoundFrontendTransform.h"
@@ -109,6 +110,7 @@ namespace Metasound
 			if (const UMetaSoundSettings* Settings = GetDefault<UMetaSoundSettings>())
 			{
 				RegOptions.bAutoUpdateLogWarningOnDroppedConnection = Settings->bAutoUpdateLogWarningOnDroppedConnection;
+				RegOptions.PageOrder = Settings->GetPageOrder();
 			}
 
 			return RegOptions;
@@ -862,7 +864,7 @@ void UMetaSoundSource::InitParameters(TArray<FAudioParameter>& ParametersToInit,
 		// stored there will not be used. The necessary proxies in the ParametersToInit 
 		// will be created and used instead. 
 		constexpr bool bCreateUObjectProxiesInRuntimeInputMap = false; 
-		InitParametersInternal(CreateRuntimeInputMap(bCreateUObjectProxiesInRuntimeInputMap), ParametersToInit, InFeatureName);
+		InitParametersInternal(CreateRuntimeInputMap(bCreateUObjectProxiesInRuntimeInputMap, UMetaSoundSettings::GetPageOrder()), ParametersToInit, InFeatureName);
 	}
 	else
 	{
@@ -884,7 +886,7 @@ void UMetaSoundSource::InitParameters(TArray<FAudioParameter>& ParametersToInit,
 			// stored there will not be used. The necessary proxies in the ParametersToInit 
 			// will be created and used instead. 
 			constexpr bool bCreateUObjectProxiesInRuntimeInputMap = false; 
-			InitParametersInternal(CreateRuntimeInputMap(bCreateUObjectProxiesInRuntimeInputMap), ParametersToInit, InFeatureName);
+			InitParametersInternal(CreateRuntimeInputMap(bCreateUObjectProxiesInRuntimeInputMap, UMetaSoundSettings::GetPageOrder()), ParametersToInit, InFeatureName);
 		}
 	}
 }
@@ -925,7 +927,25 @@ void UMetaSoundSource::InitResources()
 	}
 }
 
-void UMetaSoundSource::UpdateAndRegisterForExecution(Metasound::Frontend::FMetaSoundAssetRegistrationOptions InRegistrationOptions)
+TSharedPtr<Audio::IProxyData> UMetaSoundSource::CreateProxyData(const Audio::FProxyDataInitParams& InitParams)
+{
+	if (!IsRegistered())
+	{
+		// Make sure the metasound is registered and ready to go before making a proxy.
+		InitResources();
+	}
+
+	if (IsRegistered())
+	{
+		return FMetasoundAssetBase::CreateProxyData(InitParams);
+	}
+	else
+	{
+		return nullptr;
+	}
+}
+
+void UMetaSoundSource::UpdateAndRegisterForExecution(const Metasound::Frontend::FMetaSoundAssetRegistrationOptions& InRegistrationOptions)
 {
 	using namespace Metasound;
 	using namespace Metasound::Frontend;
@@ -937,7 +957,7 @@ void UMetaSoundSource::UpdateAndRegisterForExecution(Metasound::Frontend::FMetaS
 	// Runtime data does not need to and should not be created in builds that don't execute graphs
 	if (!bIsRuntimeInputDataValid && Metasound::CanEverExecuteGraph())
 	{
-		CacheRuntimeInputData();
+		CacheRuntimeInputData(InRegistrationOptions.PageOrder);
 	}
 }
 
@@ -1156,7 +1176,7 @@ bool UMetaSoundSource::GetAllDefaultParameters(TArray<FAudioParameter>& OutParam
 		// Non executable contexts do not have cached runtime data and need their 
 		// parameters created each time they are requested. 
 		constexpr bool bCreateUObjectProxies = false;
-		TSortedVertexNameMap<FRuntimeInput> InputMap = CreateRuntimeInputMap(bCreateUObjectProxies);
+		TSortedVertexNameMap<FRuntimeInput> InputMap = CreateRuntimeInputMap(bCreateUObjectProxies, UMetaSoundSettings::GetPageOrder());
 
 		for(const TPair<FVertexName, FRuntimeInput>& Pair : InputMap)
 		{
@@ -1494,7 +1514,7 @@ TSharedPtr<Audio::IParameterTransmitter> UMetaSoundSource::CreateParameterTransm
 		// stored there will not be used. The necessary proxies in the ParametersToInit 
 		// will be created and used instead. 
 		constexpr bool bCreateUObjectProxiesInRuntimeInputMap = false; 
-		return CreateParameterTransmitterInternal(CreateRuntimeInputMap(bCreateUObjectProxiesInRuntimeInputMap), InParams);
+		return CreateParameterTransmitterInternal(CreateRuntimeInputMap(bCreateUObjectProxiesInRuntimeInputMap, UMetaSoundSettings::GetPageOrder()), InParams);
 	}
 	else
 	{
@@ -1768,7 +1788,7 @@ TSharedPtr<Metasound::DynamicGraph::FDynamicOperatorTransactor> UMetaSoundSource
 	return DynamicTransactor;
 }
 
-UMetaSoundSource::FRuntimeInput UMetaSoundSource::CreateRuntimeInput(const Metasound::Frontend::IDataTypeRegistry& Registry, const FMetasoundFrontendClassInput& Input, bool bCreateUObjectProxies)
+UMetaSoundSource::FRuntimeInput UMetaSoundSource::CreateRuntimeInput(const Metasound::Frontend::IDataTypeRegistry& Registry, const FMetasoundFrontendClassInput& Input, bool bCreateUObjectProxies, TArrayView<const FGuid> InPageOrder)
 {
 	using namespace Metasound;
 	using namespace Metasound::Frontend;
@@ -1783,14 +1803,23 @@ UMetaSoundSource::FRuntimeInput UMetaSoundSource::CreateRuntimeInput(const Metas
 		UE_LOG(LogMetaSound, Warning, TEXT("Failed to find data type '%s' in registry. Assuming data type is not transmittable"), *Input.TypeName.ToString());
 	}
 
-	const FGuid PageID = Engine::FDocumentBuilderRegistry::GetChecked().ResolveTargetPageID(Input);
-	const FMetasoundFrontendLiteral& DefaultLiteral = Input.FindConstDefaultChecked(PageID);
-	FAudioParameter DefaultParameter = SourcePrivate::MakeAudioParameter(Registry, Input.Name, Input.TypeName, DefaultLiteral, bCreateUObjectProxies);
+	if (const FMetasoundFrontendClassInputDefault* Default = FindPreferredPage(Input.GetDefaults(), InPageOrder))
+	{
+		FAudioParameter DefaultParameter = SourcePrivate::MakeAudioParameter(Registry, Input.Name, Input.TypeName, Default->Literal, bCreateUObjectProxies);
 
-	return FRuntimeInput { Input.Name, Input.TypeName, Input.AccessType, DefaultParameter, bIsTransmittable };
+		return FRuntimeInput { Input.Name, Input.TypeName, Input.AccessType, DefaultParameter, bIsTransmittable };
+	}
+	else
+	{
+		UE_LOG(LogMetaSound, Error, TEXT("Failed to find default literal for parameter %s"), *Input.Name.ToString());
+		FMetasoundFrontendLiteral FrontendLiteral;
+		FrontendLiteral.SetFromLiteral(Registry.CreateDefaultLiteral(Input.TypeName));
+		FAudioParameter DefaultParameter = SourcePrivate::MakeAudioParameter(Registry, Input.Name, Input.TypeName, FrontendLiteral, bCreateUObjectProxies);
+		return FRuntimeInput { Input.Name, Input.TypeName, Input.AccessType, DefaultParameter, bIsTransmittable };
+	}
 }
 
-Metasound::TSortedVertexNameMap<UMetaSoundSource::FRuntimeInput> UMetaSoundSource::CreateRuntimeInputMap(bool bCreateUObjectProxies) const
+Metasound::TSortedVertexNameMap<UMetaSoundSource::FRuntimeInput> UMetaSoundSource::CreateRuntimeInputMap(bool bCreateUObjectProxies, TArrayView<const FGuid> InPageOrder) const
 {
 	using namespace Metasound;
 	using namespace Metasound::Frontend;
@@ -1826,7 +1855,7 @@ Metasound::TSortedVertexNameMap<UMetaSoundSource::FRuntimeInput> UMetaSoundSourc
 	{
 		if (!PrivateInputs.Contains(Input.Name))
 		{
-			PublicInputs.Add(Input.Name, CreateRuntimeInput(Registry, Input, bCreateUObjectProxies));
+			PublicInputs.Add(Input.Name, CreateRuntimeInput(Registry, Input, bCreateUObjectProxies, InPageOrder));
 		}
 	}
 
@@ -1839,12 +1868,12 @@ Metasound::TSortedVertexNameMap<UMetaSoundSource::FRuntimeInput> UMetaSoundSourc
 	return PublicInputs;
 }
 
-void UMetaSoundSource::CacheRuntimeInputData()
+void UMetaSoundSource::CacheRuntimeInputData(TArrayView<const FGuid> InPageOrder)
 {
 	using namespace Metasound;
 
 	constexpr bool bCreateUObjectProxies = true; 
-	RuntimeInputData.InputMap = CreateRuntimeInputMap(bCreateUObjectProxies);
+	RuntimeInputData.InputMap = CreateRuntimeInputMap(bCreateUObjectProxies, InPageOrder);
 
 	// Determine if preset graph inflation is possible 
 	//
diff --git a/Engine/Plugins/Runtime/Metasound/Source/MetasoundEngine/Public/Metasound.h b/Engine/Plugins/Runtime/Metasound/Source/MetasoundEngine/Public/Metasound.h
--- a/Engine/Plugins/Runtime/Metasound/Source/MetasoundEngine/Public/Metasound.h
+++ b/Engine/Plugins/Runtime/Metasound/Source/MetasoundEngine/Public/Metasound.h
@@ -169,6 +169,8 @@ public:
 
 	virtual bool ConformObjectToDocument() override { return false; }
 
+	UE_API virtual TSharedPtr<Audio::IProxyData> CreateProxyData(const Audio::FProxyDataInitParams& InitParams) override;
+
 	virtual const TSet<FString>& GetReferencedAssetClassKeys() const override
 	{
 		return ReferencedAssetClassKeys;
diff --git a/Engine/Plugins/Runtime/Metasound/Source/MetasoundEngine/Public/MetasoundAssetSubsystem.h b/Engine/Plugins/Runtime/Metasound/Source/MetasoundEngine/Public/MetasoundAssetSubsystem.h
--- a/Engine/Plugins/Runtime/Metasound/Source/MetasoundEngine/Public/MetasoundAssetSubsystem.h
+++ b/Engine/Plugins/Runtime/Metasound/Source/MetasoundEngine/Public/MetasoundAssetSubsystem.h
@@ -92,7 +92,7 @@ PRAGMA_ENABLE_DEPRECATION_WARNINGS
 #endif // WITH_EDITOR
 		UE_API void RequestAsyncLoadReferencedAssets(FMetasoundAssetBase& InAssetBase);
 
-		UE_DEPRECATED(5.6, "Implementation has been privitized")
+		UE_DEPRECATED(5.6, "Implementation has been privatized")
 		void OnAssetScanComplete() { }
 
 		UE_API void SearchAndIterateDirectoryAssets(const TArray<FDirectoryPath>& InDirectories, TFunctionRef<void(const FAssetData&)> InFunction) const;
@@ -155,6 +155,7 @@ PRAGMA_ENABLE_DEPRECATION_WARNINGS
 		UE_API virtual void AddReferencedObjects(FReferenceCollector& Collector) override;
 		virtual FString GetReferencerName() const override { return TEXT("FMetaSoundAssetManager"); }
 
+		Frontend::FMetaSoundAssetRegistrationOptions GetRegistrationOptions() const;
 	private:
 		FMetaSoundAssetKey AddOrUpdateFromObjectInternal(const UObject& InObject);
 
diff --git a/Engine/Plugins/Runtime/Metasound/Source/MetasoundEngine/Public/MetasoundBuilderBase.cpp b/Engine/Plugins/Runtime/Metasound/Source/MetasoundEngine/Public/MetasoundBuilderBase.cpp
--- a/Engine/Plugins/Runtime/Metasound/Source/MetasoundEngine/Public/MetasoundBuilderBase.cpp
+++ b/Engine/Plugins/Runtime/Metasound/Source/MetasoundEngine/Public/MetasoundBuilderBase.cpp
@@ -19,6 +19,7 @@
 #include "MetasoundFrontendDataTypeRegistry.h"
 #include "MetasoundFrontendDocument.h"
 #include "MetasoundFrontendDocumentIdGenerator.h"
+#include "MetasoundFrontendPages.h"
 #include "MetasoundFrontendRegistries.h"
 #include "MetasoundFrontendSearchEngine.h"
 #include "MetasoundFrontendTransform.h"
@@ -600,7 +601,9 @@ void UMetaSoundBuilderBase::ConvertToPreset(const TScriptInterface<IMetaSoundDoc
 	// Ensure the referenced node is registered
 	if (FMetasoundAssetBase* ReferencedMetaSoundAsset = Metasound::IMetasoundUObjectRegistry::Get().GetObjectAsAssetBase(ReferencedObject))
 	{
-		ReferencedMetaSoundAsset->UpdateAndRegisterForExecution();
+		Metasound::Frontend::FMetaSoundAssetRegistrationOptions RegOptions;
+		RegOptions.PageOrder = UMetaSoundSettings::GetPageOrder();
+		ReferencedMetaSoundAsset->UpdateAndRegisterForExecution(RegOptions);
 	}
 
 	const FMetasoundFrontendDocument& ReferencedDocument = ReferencedInterface->GetConstDocument();
@@ -1090,6 +1093,7 @@ void UMetaSoundBuilderBase::RegisterGraphIfOutstandingTransactions(UObject& InMe
 	FMetaSoundAssetRegistrationOptions Options;
 	Options.bForceReregister = false;
 	Options.bRegisterDependencies = false; // Function handles registration via own recursive functionality below
+	Options.PageOrder = UMetaSoundSettings::GetPageOrder();
 
 	TArray<FMetasoundAssetBase*> References = MetaSoundAsset->GetReferencedAssets();
 	for (FMetasoundAssetBase* Reference : References)
@@ -1392,13 +1396,7 @@ FMetasoundFrontendLiteral UMetaSoundBuilderBase::GetNodeInputClassDefault(const
 	{
 		if (const TArray<FMetasoundFrontendClassInputDefault>* ClassDefaults = Builder.FindNodeClassInputDefaults(InputHandle.NodeID, Vertex->Name))
 		{
-			const FGuid ResolvedPageID = FDocumentBuilderRegistry::GetChecked().ResolveTargetPageID(*ClassDefaults);
-			OutResult = EMetaSoundBuilderResult::Succeeded;
-			auto MatchesPageID = [&ResolvedPageID](const FMetasoundFrontendClassInputDefault& Default)
-			{
-				return Default.PageID == ResolvedPageID;
-			};
-			if (const FMetasoundFrontendClassInputDefault* Default = ClassDefaults->FindByPredicate(MatchesPageID))
+			if (const FMetasoundFrontendClassInputDefault* Default = Metasound::Frontend::FindPreferredPage(*ClassDefaults, UMetaSoundSettings::GetPageOrder()))
 			{
 				OutResult = EMetaSoundBuilderResult::Succeeded;
 				return Default->Literal;
diff --git a/Engine/Plugins/Runtime/Metasound/Source/MetasoundEngine/Public/MetasoundBuilderBase.h b/Engine/Plugins/Runtime/Metasound/Source/MetasoundEngine/Public/MetasoundBuilderBase.h
--- a/Engine/Plugins/Runtime/Metasound/Source/MetasoundEngine/Public/MetasoundBuilderBase.h
+++ b/Engine/Plugins/Runtime/Metasound/Source/MetasoundEngine/Public/MetasoundBuilderBase.h
@@ -9,6 +9,7 @@
 #include "MetasoundFrontendDocument.h"
 #include "MetasoundFrontendDocumentBuilder.h"
 #include "MetasoundFrontendDocumentModifyDelegates.h"
+#include "MetasoundSettings.h"
 #include "Subsystems/EngineSubsystem.h"
 #include "Templates/Function.h"
 #include "UObject/NoExportTypes.h"
@@ -642,7 +643,9 @@ protected:
 
 		if (BuilderOptions.bAddToRegistry)
 		{
-			MetaSound->UpdateAndRegisterForExecution();
+			Metasound::Frontend::FMetaSoundAssetRegistrationOptions RegOptions;
+			RegOptions.PageOrder = UMetaSoundSettings::GetPageOrder();
+			MetaSound->UpdateAndRegisterForExecution(RegOptions);
 		}
 
 		UE_LOG(LogMetaSound, VeryVerbose, TEXT("MetaSound '%s' built from '%s'"), *BuilderOptions.Name.ToString(), *GetFullName());
diff --git a/Engine/Plugins/Runtime/Metasound/Source/MetasoundEngine/Public/MetasoundDocumentBuilderRegistry.h b/Engine/Plugins/Runtime/Metasound/Source/MetasoundEngine/Public/MetasoundDocumentBuilderRegistry.h
--- a/Engine/Plugins/Runtime/Metasound/Source/MetasoundEngine/Public/MetasoundDocumentBuilderRegistry.h
+++ b/Engine/Plugins/Runtime/Metasound/Source/MetasoundEngine/Public/MetasoundDocumentBuilderRegistry.h
@@ -8,29 +8,14 @@
 #include "MetasoundFrontendDocument.h"
 #include "MetasoundFrontendRegistries.h"
 #include "MetasoundLog.h"
+#include "MetasoundSettings.h"
 #include "Misc/Optional.h"
 #include "Misc/ScopeLock.h"
 
 #define UE_API METASOUNDENGINE_API
 
-
-// Forward Declarations
-class UMetaSoundSettings;
-
 namespace Metasound::Engine
 {
-#if WITH_EDITOR
-	struct FPageResolutionEditorResults
-	{
-		FName PlatformName;
-		TOptional<FGuid> PageID;
-	};
-
-	DECLARE_DELEGATE_RetVal_OneParam(FPageResolutionEditorResults, FOnResolveEditorPage, const TArray<FGuid>& /* InPageIDs */);
-#endif // WITH_EDITOR
-
-	DECLARE_DELEGATE_RetVal_OneParam(FGuid, FOnResolvePage, const TArray<FGuid>& /* InPageIDs */);
-
 	class FDocumentBuilderRegistry : public Frontend::IDocumentBuilderRegistry
 	{
 		mutable TMultiMap<FMetasoundFrontendClassName, TWeakObjectPtr<UMetaSoundBuilderBase>> Builders;
@@ -128,7 +113,7 @@ namespace Metasound::Engine
 		// platform, allowing users to optimize away data and scale the amount of memory required for
 		// initial load of input UObjects and graph topology, which can also positively effect runtime
 		// performance as well, etc. Returns true if builder modified the document, false if not.
-		UE_API virtual bool CookPages(FName PlatformName, FMetaSoundFrontendDocumentBuilder& Builder) const override;
+		UE_DEPRECATED(5.7, "Use Metasound::Frontend::StripUnusedPages(...) instead") UE_API virtual bool CookPages(FName PlatformName, FMetaSoundFrontendDocumentBuilder& Builder) const override;
 		UE_API virtual FMetaSoundFrontendDocumentBuilder& FindOrBeginBuilding(TScriptInterface<IMetaSoundDocumentInterface> MetaSound) override;
 #endif // WITH_EDITORONLY_DATA
 
@@ -153,17 +138,34 @@ namespace Metasound::Engine
 		UE_API TArray<UMetaSoundBuilderBase*> FindBuilderObjects(const FMetasoundFrontendClassName& InClassName) const;
 
 #if WITH_EDITOR
-		UE_API FOnResolveEditorPage& GetOnResolveAuditionPageDelegate();
+PRAGMA_DISABLE_DEPRECATION_WARNINGS
+		UE_DEPRECATED(5.7, "Use UMetaSoundSettings::GetOnResolveAuditionPageDelegate() instead.") UE_API FOnResolveEditorPage& GetOnResolveAuditionPageDelegate();
+PRAGMA_ENABLE_DEPRECATION_WARNINGS
 #endif // WITH_EDITOR
 
-		UE_API FOnResolvePage& GetOnResolveProjectPageOverrideDelegate();
+PRAGMA_DISABLE_DEPRECATION_WARNINGS
+		UE_DEPRECATED(5.7, "This is no longer supported") UE_API FOnResolvePage& GetOnResolveProjectPageOverrideDelegate();
+PRAGMA_ENABLE_DEPRECATION_WARNINGS
 
 		UE_API bool ReloadBuilder(const FMetasoundFrontendClassName& InClassName) const override;
 
 		// Given the provided document and its respective pages, returns the PageID to be used for runtime IGraph and proxy generation.
-		UE_API virtual FGuid ResolveTargetPageID(const FMetasoundFrontendGraphClass& InGraphClass) const override;
-		UE_API virtual FGuid ResolveTargetPageID(const FMetasoundFrontendClassInput& InClassInput) const override;
-		UE_API virtual FGuid ResolveTargetPageID(const TArray<FMetasoundFrontendClassInputDefault>& Defaults) const override;
+		UE_DEPRECATED(5.7, "This is no longer supported.") UE_API virtual FGuid ResolveTargetPageID(const FMetasoundFrontendGraphClass& InGraphClass) const override;
+		UE_DEPRECATED(5.7, "This is no longer supported.") UE_API virtual FGuid ResolveTargetPageID(const FMetasoundFrontendClassInput& InClassInput) const override;
+		UE_DEPRECATED(5.7, "This is no longer supported.") UE_API virtual FGuid ResolveTargetPageID(const TArray<FMetasoundFrontendClassInputDefault>& Defaults) const override;
+
+		UE_DEPRECATED(5.7, "Retrieve page order from UMetaSoundSettings.") UE_API FGuid ResolveTargetPageID(const TArray<FGuid>& PageIdsToResolve) const;
+
+		UE_DEPRECATED(5.7, "This method exist only to supported deprecated functionality in FrontendNodeControllers and should not be in any new code.")
+		virtual TArrayView<const FGuid> GetPageOrder() const override;
+
+#if WITH_EDITORONLY_DATA
+		UE_DEPRECATED(5.7, "This is only used to support existing deprecated functionality in FMetasoundAssetBase cooking and should not be used in any new code.")
+		virtual TArray<FGuid> GetCookedTargetPages(FName InPlatformName) const override;
+
+		UE_DEPRECATED(5.7, "This is only used to support existing deprecated functionality in FMetasoundAssetBase cooking and should not be used in any new code.")
+		virtual TArray<FGuid> GetCookedPageOrder(FName InPlatformName) const override;
+#endif // WITH_EDITORONLY_DATA
 
 		UE_API void SetEventLogVerbosity(ELogEvent Event, ELogVerbosity::Type Verbosity);
 
@@ -171,14 +173,17 @@ namespace Metasound::Engine
 		UE_API void AddBuilderInternal(const FMetasoundFrontendClassName& InClassName, UMetaSoundBuilderBase* NewBuilder) const;
 		UE_API bool CanPostEventLog(ELogEvent Event, ELogVerbosity::Type Verbosity) const;
 		UE_API void FinishBuildingInternal(UMetaSoundBuilderBase& Builder, bool bForceUnregisterNodeClass) const;
-		UE_API FGuid ResolveTargetPageIDInternal(const TArray<FGuid>& PageIdsToResolve) const;
 		UE_API FGuid ResolveTargetPageIDInternal(const UMetaSoundSettings& Settings, const TArray<FGuid>& PageIdsToResolve, const FGuid& TargetPageID, FName PlatformName) const;
 
 #if WITH_EDITOR
+PRAGMA_DISABLE_DEPRECATION_WARNINGS
 		FOnResolveEditorPage OnResolveAuditionPage;
+PRAGMA_ENABLE_DEPRECATION_WARNINGS
 #endif // WITH_EDITOR
 
+PRAGMA_DISABLE_DEPRECATION_WARNINGS
 		FOnResolvePage OnResolveProjectPage;
+PRAGMA_ENABLE_DEPRECATION_WARNINGS
 
 		// Reuseable scratch array of pages to resolve, which is used to
 		// optimize/reduce number of allocations required when resolving document.
diff --git a/Engine/Plugins/Runtime/Metasound/Source/MetasoundEngine/Public/MetasoundSettings.h b/Engine/Plugins/Runtime/Metasound/Source/MetasoundEngine/Public/MetasoundSettings.h
--- a/Engine/Plugins/Runtime/Metasound/Source/MetasoundEngine/Public/MetasoundSettings.h
+++ b/Engine/Plugins/Runtime/Metasound/Source/MetasoundEngine/Public/MetasoundSettings.h
@@ -20,14 +20,32 @@
 struct FMetasoundFrontendClassName;
 struct FPropertyChangedChainEvent;
 
-#if WITH_EDITORONLY_DATA
 namespace Metasound::Engine
 {
+#if WITH_EDITORONLY_DATA
 	DECLARE_MULTICAST_DELEGATE(FOnSettingsDefaultConformed);
 	DECLARE_MULTICAST_DELEGATE(FOnPageSettingsUpdated);
-} // namespace Metasound::Engine
+
+#if WITH_EDITOR
+	struct FTargetPageOverride
+	{
+		FName PlatformName;
+		TOptional<FGuid> PageID;
+	};
+
+
+	using FPageResolutionEditorResults UE_DEPRECATED(5.7, "Use FTargetPageOverride instead.") = FTargetPageOverride;
+
+	UE_DEPRECATED(5.7, "Use FOverrideTargetPageSettings instead.")
+	DECLARE_DELEGATE_RetVal_OneParam(FTargetPageOverride, FOnResolveEditorPage, const TArray<FGuid>& /* InPageIDs */);
+
+#endif // WITH_EDITOR
 #endif // WITH_EDITORONLY_DATA
 
+	UE_DEPRECATED(5.7, "This is no longer supported.")
+	DECLARE_DELEGATE_RetVal_OneParam(FGuid, FOnResolvePage, const TArray<FGuid>& /* InPageIDs */);
+} // namespace Metasound::Engine
+
 UENUM()
 enum class EMetaSoundMessageLevel : uint8
 {
@@ -203,10 +221,24 @@ private:
 #if WITH_EDITORONLY_DATA
 	mutable FCriticalSection CookPlatformTargetCritSec;
 	mutable TArray<FGuid> CookPlatformTargetPageIDs;
+	mutable TArray<FGuid> CookPlatformPageOrder;
 	mutable FName CookPlatformTargetPage;
 #endif // WITH_EDITORONLY_DATA
+	   
+	// Generates the page order for a current configured platform and target ID
+	TArray<FGuid> GeneratePageOrder() const;
+
+	// Generates the page order for a given platform and desired target page ID
+	TArray<FGuid> GeneratePageOrderInternal(const FName& InPlatformName, const FGuid& InTargetPageID) const;
+
+	// Updates the cached page order with the latest settings.
+	TArray<FGuid> CachedPageOrder;
+	void UpdateCachedPageOrder();
 
 public:
+	// Returns a cached global version of the PageOrder from the UMetaSoundSettings Default Object.
+	static UE_API TArrayView<const FGuid> GetPageOrder();
+
 	// Returns the page settings with the provided name. If there are multiple settings
 	// with the same name, selection within the duplicates is undefined.
 	UE_API const FMetaSoundPageSettings* FindPageSettings(FName Name) const;
@@ -226,6 +258,7 @@ public:
 #if WITH_EDITORONLY_DATA
 	// Returns PageIDs to be cooked to the given platform/platform group.
 	UE_API TArray<FGuid> GetCookedTargetPageIDs(FName PlatformName) const;
+	UE_API TArray<FGuid> GetCookedPageOrder(FName PlatformName) const;
 
 	// Iterates PageIDs cooked to the given platform/platform group.
 	UE_API void IterateCookedTargetPageIDs(FName PlatformName, TFunctionRef<void(const FGuid&)> Iter) const;
@@ -244,6 +277,10 @@ public:
 
 	const TArray<FMetaSoundQualitySettings>& GetQualitySettings() const { return QualitySettings; }
 
+#if WITH_EDITOR
+	UE_API void OverrideTargetPageSettings(TOptional<Metasound::Engine::FTargetPageOverride> InOverride);
+#endif // WITH_EDITOR
+
 #if WITH_EDITORONLY_DATA
 	UE_API Metasound::Engine::FOnSettingsDefaultConformed& GetOnDefaultRenamedDelegate();
 	UE_API Metasound::Engine::FOnPageSettingsUpdated& GetOnPageSettingsUpdatedDelegate();
@@ -279,7 +316,8 @@ public:
 
 private:
 #if WITH_EDITORONLY_DATA
-	UE_API const TArray<FGuid>& GetCookedTargetPageIDsInternal(FName PlatformName) const;
+	void CacheCookedPageIDs(const FName InPlatformName) const;
+	
 #endif // WITH_EDITORONLY_DATA
 
 #if WITH_EDITOR
@@ -287,6 +325,11 @@ private:
 
 	UE_API virtual void PostEditChangeChainProperty(FPropertyChangedChainEvent& PropertyChangedEvent) override;
 	UE_API virtual void PostEditChangeProperty(FPropertyChangedEvent& PropertyChangedEvent) override;
+
+	const FMetaSoundPageSettings* GetTargetPageSettingsFromOverride(const FName InTargetPageName) const;
+	TArray<FGuid> GeneratePageOrderFromOverride() const;
+	
+	TOptional<Metasound::Engine::FTargetPageOverride> PageOverrideData;
 #endif // WITH_EDITOR
 
 	UE_API virtual void PostInitProperties() override;
diff --git a/Engine/Plugins/Runtime/Metasound/Source/MetasoundEngine/Public/MetasoundSource.h b/Engine/Plugins/Runtime/Metasound/Source/MetasoundEngine/Public/MetasoundSource.h
--- a/Engine/Plugins/Runtime/Metasound/Source/MetasoundEngine/Public/MetasoundSource.h
+++ b/Engine/Plugins/Runtime/Metasound/Source/MetasoundEngine/Public/MetasoundSource.h
@@ -270,7 +270,11 @@ public:
 	UE_API virtual void InitParameters(TArray<FAudioParameter>& ParametersToInit, FName InFeatureName) override;
 
 	UE_API virtual void InitResources() override;
-	UE_API virtual void UpdateAndRegisterForExecution(Metasound::Frontend::FMetaSoundAssetRegistrationOptions InRegistrationOptions = Metasound::Frontend::FMetaSoundAssetRegistrationOptions()) override;
+
+	UE_API virtual TSharedPtr<Audio::IProxyData> CreateProxyData(const Audio::FProxyDataInitParams& InitParams) override;
+
+	using FMetasoundAssetBase::UpdateAndRegisterForExecution;
+	UE_API virtual void UpdateAndRegisterForExecution(const Metasound::Frontend::FMetaSoundAssetRegistrationOptions& InRegistrationOptions) override;
 
 	UE_API virtual bool IsPlayable() const override;
 	UE_API virtual float GetDuration() const override;
@@ -372,9 +376,9 @@ private:
 	UE_API void TrackGenerator(Metasound::FGeneratorInstanceInfo&& GeneratorInfo);
 	UE_API void ForgetGenerator(ISoundGeneratorPtr Generator);
 
-	static UE_API FRuntimeInput CreateRuntimeInput(const Metasound::Frontend::IDataTypeRegistry& Registry, const FMetasoundFrontendClassInput& Input, bool bCreateUObjectProxies);
-	UE_API Metasound::TSortedVertexNameMap<FRuntimeInput> CreateRuntimeInputMap(bool bCreateUObjectProxies) const;
-	UE_API void CacheRuntimeInputData();
+	static UE_API FRuntimeInput CreateRuntimeInput(const Metasound::Frontend::IDataTypeRegistry& Registry, const FMetasoundFrontendClassInput& Input, bool bCreateUObjectProxies, TArrayView<const FGuid> InPageOrder);
+	UE_API Metasound::TSortedVertexNameMap<FRuntimeInput> CreateRuntimeInputMap(bool bCreateUObjectProxies, TArrayView<const FGuid> InPageOrder) const;
+	UE_API void CacheRuntimeInputData(TArrayView<const FGuid> InPageOrder);
 	UE_API void InvalidateCachedRuntimeInputData();
 
 	FRuntimeInputData RuntimeInputData;
diff --git a/Engine/Plugins/Runtime/Metasound/Source/MetasoundEngineTest/Private/EngineTestMetasoundSource.cpp b/Engine/Plugins/Runtime/Metasound/Source/MetasoundEngineTest/Private/EngineTestMetasoundSource.cpp
--- a/Engine/Plugins/Runtime/Metasound/Source/MetasoundEngineTest/Private/EngineTestMetasoundSource.cpp
+++ b/Engine/Plugins/Runtime/Metasound/Source/MetasoundEngineTest/Private/EngineTestMetasoundSource.cpp
@@ -396,7 +396,7 @@ bool FAudioMetasoundSourceTest::RunTest(const FString& Parameters)
 	{
 		// TODO: Move to using builder to swap in doc
 		PRAGMA_DISABLE_DEPRECATION_WARNINGS
-		MetaSoundSource->SetDocument(CreateMonoSourceDocument());
+		MetaSoundSource->GetDocumentChecked() = CreateMonoSourceDocument();
 		PRAGMA_ENABLE_DEPRECATION_WARNINGS
 
 		if (UAudioComponent* AudioComponent = CreateTestComponent(*this, MetaSoundSource))
diff --git a/Engine/Plugins/Runtime/Metasound/Source/MetasoundEngineTest/Private/NodeTestGraphBuilder.cpp b/Engine/Plugins/Runtime/Metasound/Source/MetasoundEngineTest/Private/NodeTestGraphBuilder.cpp
--- a/Engine/Plugins/Runtime/Metasound/Source/MetasoundEngineTest/Private/NodeTestGraphBuilder.cpp
+++ b/Engine/Plugins/Runtime/Metasound/Source/MetasoundEngineTest/Private/NodeTestGraphBuilder.cpp
@@ -155,7 +155,7 @@ namespace Metasound::Test
 			if (const FMetasoundFrontendLiteral* Default = InputToConnect->GetClassDefaultLiteral(); nullptr != Default)
 			{
 				const FGuid InputId = Builder.RootGraph->GetVertexIDForInputVertex(Input->GetName());
-				Builder.RootGraph->SetDefaultInput(InputId, *Default);
+				Builder.RootGraph->SetDefaultInput(InputId, *Default, {});
 			}
 
 			if (!InputToConnect->Connect(*OutputToConnect))
diff --git a/Engine/Plugins/Runtime/Metasound/Source/MetasoundFrontend/Private/MetasoundAssetBase.cpp b/Engine/Plugins/Runtime/Metasound/Source/MetasoundFrontend/Private/MetasoundAssetBase.cpp
--- a/Engine/Plugins/Runtime/Metasound/Source/MetasoundFrontend/Private/MetasoundAssetBase.cpp
+++ b/Engine/Plugins/Runtime/Metasound/Source/MetasoundFrontend/Private/MetasoundAssetBase.cpp
@@ -25,6 +25,7 @@
 #include "MetasoundFrontendNodeClassRegistry.h"
 #include "MetasoundFrontendNodeClassRegistryPrivate.h"
 #include "MetasoundFrontendNodeTemplateRegistry.h"
+#include "MetasoundFrontendPages.h"
 #include "MetasoundFrontendProxyDataCache.h"
 #include "MetasoundFrontendRegistries.h"
 #include "MetasoundFrontendSearchEngine.h"
@@ -261,19 +262,16 @@ namespace Metasound
 
 const FString FMetasoundAssetBase::FileExtension(TEXT(".metasound"));
 
-bool FMetasoundAssetBase::ConformObjectDataToInterfaces()
-{
-	return false;
-}
-
 TSharedPtr<Audio::IProxyData> FMetasoundAssetBase::CreateProxyData(const Audio::FProxyDataInitParams& InitParams)
 {
 	using namespace Metasound::Frontend;
 	
 	if (!IsRegistered())
 	{
-		UpdateAndRegisterForExecution();
+		UE_LOG(LogMetaSound, Error, TEXT("Cannot make proxy for MetaSound %s. Derived classes should ensure that they are registered before calling into FMetasoundAssetBase::CreateProxyData"), *GetOwningAssetName());
+		return nullptr;
 	}
+
 	const FGraphRegistryKey& Key = GetGraphRegistryKey();
 	TScriptInterface<const IMetaSoundDocumentInterface> DocInterface = GetOwningAsset();
 
@@ -320,12 +318,18 @@ TSharedPtr<Audio::IProxyData> FMetasoundAssetBase::CreateProxyData(const Audio::
 	return Proxy;
 }
 
-void FMetasoundAssetBase::RegisterGraphWithFrontend(Metasound::Frontend::FMetaSoundAssetRegistrationOptions InRegistrationOptions)
+void FMetasoundAssetBase::UpdateAndRegisterForExecution()
 {
-	UpdateAndRegisterForExecution(MoveTemp(InRegistrationOptions));
+	// This is handling the case of a call where no registration options have been
+	// explicitly set. As a fallback, use the default page ID. 
+	Metasound::Frontend::FMetaSoundAssetRegistrationOptions RegOptions;
+	TArray<FGuid> PageOrder({Metasound::Frontend::DefaultPageID});
+	RegOptions.PageOrder = PageOrder;
+	
+	UpdateAndRegisterForExecution(RegOptions);
 }
 
-void FMetasoundAssetBase::UpdateAndRegisterForExecution(Metasound::Frontend::FMetaSoundAssetRegistrationOptions InRegistrationOptions)
+void FMetasoundAssetBase::UpdateAndRegisterForExecution(const Metasound::Frontend::FMetaSoundAssetRegistrationOptions& InRegistrationOptions)
 {
 	using namespace Metasound;
 	using namespace Metasound::Frontend;
@@ -407,7 +411,7 @@ void FMetasoundAssetBase::UpdateAndRegisterForExecution(Metasound::Frontend::FMe
 	}
 #endif // WITH_EDITOR
 
-	GraphRegistryKey = FNodeClassRegistry::Get().RegisterGraph(Owner);
+	GraphRegistryKey = FNodeClassRegistry::Get().RegisterGraph(Owner, InRegistrationOptions.PageOrder);
 	if (!GraphRegistryKey.IsValid())
 	{
 		UClass* Class = Owner->GetClass();
@@ -418,15 +422,26 @@ void FMetasoundAssetBase::UpdateAndRegisterForExecution(Metasound::Frontend::FMe
 	}
 }
 
-void FMetasoundAssetBase::CookMetaSound()
-{
-#if WITH_EDITORONLY_DATA
-	UpdateAndRegisterForSerialization();
-#endif // WITH_EDITORONLY_DATA
-}
 
 #if WITH_EDITORONLY_DATA
-void FMetasoundAssetBase::UpdateAndRegisterForSerialization(FName CookPlatformName)
+void FMetasoundAssetBase::UpdateAndRegisterForSerialization(FName InCookPlatformName)
+{
+	using namespace Metasound::Frontend;
+	FMetaSoundAssetCookOptions CookOptions{ .bStripUnusedPages = !InCookPlatformName.IsNone() };
+
+	if (CookOptions.bStripUnusedPages)
+	{
+		const IDocumentBuilderRegistry& DocBuilderRegistry = IDocumentBuilderRegistry::GetChecked();
+PRAGMA_DISABLE_DEPRECATION_WARNINGS
+		CookOptions.PagesToTarget = DocBuilderRegistry.GetCookedTargetPages(InCookPlatformName);
+		CookOptions.PageOrder = DocBuilderRegistry.GetCookedPageOrder(InCookPlatformName);
+PRAGMA_ENABLE_DEPRECATION_WARNINGS
+	}
+
+	UpdateAndRegisterForSerialization(CookOptions);
+			
+}
+void FMetasoundAssetBase::UpdateAndRegisterForSerialization(const Metasound::Frontend::FMetaSoundAssetCookOptions& InCookOptions)
 {
 	using namespace Metasound;
 	using namespace Metasound::Frontend;
@@ -439,7 +454,7 @@ void FMetasoundAssetBase::UpdateAndRegisterForSerialization(FName CookPlatformNa
 		return;
 	}
 
-	UpdateAndRegisterReferencesForSerialization(CookPlatformName);
+	UpdateAndRegisterReferencesForSerialization(InCookOptions);
 
 	UObject* Owner = GetOwningAsset();
 	check(Owner);
@@ -448,9 +463,10 @@ void FMetasoundAssetBase::UpdateAndRegisterForSerialization(FName CookPlatformNa
 	bool bDidUpdate = false;
 
 	FMetaSoundFrontendDocumentBuilder& DocBuilder = IDocumentBuilderRegistry::GetChecked().FindOrBeginBuilding(Owner);
-	if (CookPlatformName.IsValid())
+
+	if (InCookOptions.bStripUnusedPages)
 	{
-		bDidUpdate |= IDocumentBuilderRegistry::GetChecked().CookPages(CookPlatformName, DocBuilder);
+		bDidUpdate |= StripUnusedPages(DocBuilder, InCookOptions.PagesToTarget, InCookOptions.PageOrder);
 	}
 
 	// Auto update must be done after all referenced asset classes are registered
@@ -565,21 +581,6 @@ bool FMetasoundAssetBase::IsInterfaceDeclared(const FMetasoundFrontendVersion& I
 	return DocInterface->GetConstDocument().Interfaces.Contains(InVersion);
 }
 
-void FMetasoundAssetBase::SetDocument(FMetasoundFrontendDocument InDocument, bool bMarkDirty)
-{
-PRAGMA_DISABLE_DEPRECATION_WARNINGS
-	FMetasoundFrontendDocument* Document = GetDocumentAccessPtr().Get();
-PRAGMA_ENABLE_DEPRECATION_WARNINGS
-
-	*Document = MoveTemp(InDocument);
-	if (bMarkDirty)
-	{
-		UObject* OwningAsset = GetOwningAsset();
-		check(OwningAsset);
-		OwningAsset->MarkPackageDirty();
-	}
-}
-
 #if WITH_EDITORONLY_DATA
 bool FMetasoundAssetBase::VersionAsset(FMetaSoundFrontendDocumentBuilder& Builder)
 {
@@ -653,12 +654,6 @@ const FMetasoundFrontendDocumentModifyContext& FMetasoundAssetBase::GetConstModi
 	return DocInterface->GetConstDocument().Metadata.ModifyContext;
 }
 
-const FMetasoundFrontendDocumentModifyContext& FMetasoundAssetBase::GetModifyContext() const
-{
-	TScriptInterface<const IMetaSoundDocumentInterface> DocInterface = GetOwningAsset();
-	check(DocInterface.GetObject());
-	return DocInterface->GetConstDocument().Metadata.ModifyContext;
-}
 #endif // WITH_EDITOR
 
 bool FMetasoundAssetBase::IsRegistered() const
@@ -715,24 +710,6 @@ bool FMetasoundAssetBase::AddingReferenceCausesLoop(const FMetasoundAssetBase& I
 	return bCausesLoop;
 }
 
-bool FMetasoundAssetBase::AddingReferenceCausesLoop(const FSoftObjectPath& InReferencePath) const
-{
-	using namespace Metasound::Frontend;
-
-	const FMetasoundAssetBase* ReferenceAsset = IMetaSoundAssetManager::GetChecked().TryLoadAsset(InReferencePath);
-	if (!ensureAlways(ReferenceAsset))
-	{
-		return false;
-	}
-
-	return AddingReferenceCausesLoop(*ReferenceAsset);
-}
-
-TArray<FMetasoundAssetBase::FSendInfoAndVertexName> FMetasoundAssetBase::GetSendInfos(uint64 InInstanceID) const
-{
-	return TArray<FSendInfoAndVertexName>();
-}
-
 #if WITH_EDITOR
 FText FMetasoundAssetBase::GetDisplayName(FString&& InTypeName) const
 {
@@ -847,13 +824,6 @@ PRAGMA_ENABLE_DEPRECATION_WARNINGS
 	return *Document;
 }
 
-const FMetasoundFrontendDocument& FMetasoundAssetBase::GetDocumentChecked() const
-{
-	TScriptInterface<const IMetaSoundDocumentInterface> DocInterface = GetOwningAsset();
-	check(DocInterface.GetObject());
-	return DocInterface->GetConstDocument();
-}
-
 const Metasound::Frontend::FGraphRegistryKey& FMetasoundAssetBase::GetGraphRegistryKey() const
 {
 	return GraphRegistryKey;
@@ -894,15 +864,25 @@ void FMetasoundAssetBase::RegisterAssetDependencies(const Metasound::Frontend::F
 	}
 }
 
-void FMetasoundAssetBase::CookReferencedMetaSounds()
-{
 #if WITH_EDITORONLY_DATA
-	UpdateAndRegisterReferencesForSerialization({ });
-#endif // WITH_EDITORONLY_DATA
+void FMetasoundAssetBase::UpdateAndRegisterReferencesForSerialization(FName InCookPlatformName)
+{
+	using namespace Metasound::Frontend;
+	Metasound::Frontend::FMetaSoundAssetCookOptions CookOptions{ .bStripUnusedPages = !InCookPlatformName.IsNone() };
+
+	if (CookOptions.bStripUnusedPages)
+	{
+		const IDocumentBuilderRegistry& DocBuilderRegistry = IDocumentBuilderRegistry::GetChecked();
+PRAGMA_DISABLE_DEPRECATION_WARNINGS
+		CookOptions.PagesToTarget = DocBuilderRegistry.GetCookedTargetPages(InCookPlatformName);
+		CookOptions.PageOrder = DocBuilderRegistry.GetCookedPageOrder(InCookPlatformName);
+PRAGMA_ENABLE_DEPRECATION_WARNINGS
+	}
+
+	UpdateAndRegisterReferencesForSerialization(CookOptions);
 }
 
-#if WITH_EDITORONLY_DATA
-void FMetasoundAssetBase::UpdateAndRegisterReferencesForSerialization(FName CookPlatformName)
+void FMetasoundAssetBase::UpdateAndRegisterReferencesForSerialization(const Metasound::Frontend::FMetaSoundAssetCookOptions& InCookOptions)
 {
 	using namespace Metasound::Frontend;
 
@@ -911,27 +891,12 @@ void FMetasoundAssetBase::UpdateAndRegisterReferencesForSerialization(FName Cook
 	{
 		if (!Reference->IsRegistered())
 		{
-			Reference->UpdateAndRegisterForSerialization(CookPlatformName);
+			Reference->UpdateAndRegisterForSerialization(InCookOptions);
 		}
 	}
 }
 #endif // WITH_EDITORONLY_DATA
 
-bool FMetasoundAssetBase::AutoUpdate(bool bInLogWarningsOnDroppedConnection)
-{
-	using namespace Metasound::Frontend;
-
-	UObject* Owner = GetOwningAsset();
-	check(Owner);
-
-#if WITH_EDITORONLY_DATA
-	FMetaSoundFrontendDocumentBuilder& DocBuilder = IDocumentBuilderRegistry::GetChecked().FindOrBeginBuilding(Owner);
-	return VersionDependencies(DocBuilder, bInLogWarningsOnDroppedConnection);
-#else // !WITH_EDITORONLY_DATA
-	return false;
-#endif // !WITH_EDITORONLY_DATA
-}
-
 bool FMetasoundAssetBase::TryUpdateInterfaceFromVersion(const FMetasoundFrontendVersion& Version, FMetaSoundFrontendDocumentBuilder& Builder)
 {
 	using namespace Metasound::Frontend;
diff --git a/Engine/Plugins/Runtime/Metasound/Source/MetasoundFrontend/Private/MetasoundFrontendDocument.cpp b/Engine/Plugins/Runtime/Metasound/Source/MetasoundFrontend/Private/MetasoundFrontendDocument.cpp
--- a/Engine/Plugins/Runtime/Metasound/Source/MetasoundFrontend/Private/MetasoundFrontendDocument.cpp
+++ b/Engine/Plugins/Runtime/Metasound/Source/MetasoundFrontend/Private/MetasoundFrontendDocument.cpp
@@ -13,6 +13,7 @@
 #include "MetasoundFrontend.h"
 #include "MetasoundFrontendDocumentIdGenerator.h"
 #include "MetasoundFrontendDocumentVersioning.h"
+#include "MetasoundFrontendPages.h"
 #include "MetasoundFrontendRegistries.h"
 #include "MetasoundLog.h"
 #include "MetasoundParameterTransmitter.h"
@@ -41,7 +42,10 @@ namespace Metasound
 				// that returns the default page (or whatever page behavior is desired for testing).
 				if (IDocumentBuilderRegistry* BuilderRegistry = IDocumentBuilderRegistry::Get())
 				{
+PRAGMA_DISABLE_DEPRECATION_WARNINGS
+					// This deprecated function can be removed when the deprecated versions of ForEachLiteral are removed. 
 					return IDocumentBuilderRegistry::GetChecked().ResolveTargetPageID(InToResolve);
+PRAGMA_ENABLE_DEPRECATION_WARNINGS
 				}
 
 				return Frontend::DefaultPageID;
@@ -1712,28 +1716,38 @@ namespace Metasound::Frontend
 
 	void ForEachLiteral(const FMetasoundFrontendDocument& InDoc, FForEachLiteralFunctionRef OnLiteral)
 	{
+PRAGMA_DISABLE_DEPRECATION_WARNINGS
 		ForEachLiteral(InDoc.RootGraph, OnLiteral);
+PRAGMA_ENABLE_DEPRECATION_WARNINGS
 
 		for (const FMetasoundFrontendGraphClass& GraphClass : InDoc.Subgraphs)
 		{
+PRAGMA_DISABLE_DEPRECATION_WARNINGS
 			ForEachLiteral(GraphClass, OnLiteral);
+PRAGMA_ENABLE_DEPRECATION_WARNINGS
 		}
 
 		for (const FMetasoundFrontendClass& Dependency : InDoc.Dependencies)
 		{
+PRAGMA_DISABLE_DEPRECATION_WARNINGS
 			ForEachLiteral(Dependency, OnLiteral);
+PRAGMA_ENABLE_DEPRECATION_WARNINGS
 		}
 	}
 
 	void ForEachLiteral(const FMetasoundFrontendGraphClass& InGraphClass, FForEachLiteralFunctionRef OnLiteral)
 	{
+PRAGMA_DISABLE_DEPRECATION_WARNINGS
 		ForEachLiteral(static_cast<const FMetasoundFrontendClass&>(InGraphClass), OnLiteral);
+PRAGMA_ENABLE_DEPRECATION_WARNINGS
 
 		const FGuid PageID = DocumentPrivate::ResolveTargetPageID(InGraphClass);
 		const FMetasoundFrontendGraph& Graph = InGraphClass.FindConstGraphChecked(PageID);
 		for (const FMetasoundFrontendNode& Node : Graph.Nodes)
 		{
+PRAGMA_DISABLE_DEPRECATION_WARNINGS
 			ForEachLiteral(Node, OnLiteral);
+PRAGMA_ENABLE_DEPRECATION_WARNINGS
 		}
 
 		for (const FMetasoundFrontendVariable& Variable : Graph.Variables)
@@ -1744,7 +1758,9 @@ namespace Metasound::Frontend
 
 	void ForEachLiteral(const FMetasoundFrontendClass& InClass, FForEachLiteralFunctionRef OnLiteral)
 	{
+PRAGMA_DISABLE_DEPRECATION_WARNINGS
 		ForEachLiteral(InClass.GetDefaultInterface(), OnLiteral);
+PRAGMA_ENABLE_DEPRECATION_WARNINGS
 	}
 
 	void ForEachLiteral(const FMetasoundFrontendClassInterface& InClassInterface, FForEachLiteralFunctionRef OnLiteral)
@@ -1774,7 +1790,88 @@ namespace Metasound::Frontend
 
 		if (InNode.ClassInterfaceOverride.IsValid())
 		{
+PRAGMA_DISABLE_DEPRECATION_WARNINGS
 			ForEachLiteral(InNode.ClassInterfaceOverride.Get(), OnLiteral);
+PRAGMA_ENABLE_DEPRECATION_WARNINGS
+		}
+	}
+
+
+
+
+	namespace DocumentPrivate
+	{
+	}
+
+	void ForEachLiteral(const FMetasoundFrontendDocument& InDoc, FForEachLiteralFunctionRef OnLiteral, TArrayView<const FGuid> InPageOrder)
+	{
+		ForEachLiteral(InDoc.RootGraph, OnLiteral, InPageOrder);
+
+		for (const FMetasoundFrontendGraphClass& GraphClass : InDoc.Subgraphs)
+		{
+			ForEachLiteral(GraphClass, OnLiteral, InPageOrder);
+		}
+
+		for (const FMetasoundFrontendClass& Dependency : InDoc.Dependencies)
+		{
+			ForEachLiteral(Dependency, OnLiteral, InPageOrder);
+		}
+	}
+
+	void ForEachLiteral(const FMetasoundFrontendGraphClass& InGraphClass, FForEachLiteralFunctionRef OnLiteral, TArrayView<const FGuid> InPageOrder)
+	{
+		ForEachLiteral(static_cast<const FMetasoundFrontendClass&>(InGraphClass), OnLiteral, InPageOrder);
+
+		const FMetasoundFrontendGraph* Graph = FindPreferredPage<FMetasoundFrontendGraph>(InGraphClass.GetConstGraphPages(), InPageOrder);
+		if (Graph)
+		{
+			for (const FMetasoundFrontendNode& Node : Graph->Nodes)
+			{
+				ForEachLiteral(Node, OnLiteral, InPageOrder);
+			}
+
+			for (const FMetasoundFrontendVariable& Variable : Graph->Variables)
+			{
+				OnLiteral(Variable.TypeName, Variable.Literal);
+			}
+		}
+	}
+
+	void ForEachLiteral(const FMetasoundFrontendClass& InClass, FForEachLiteralFunctionRef OnLiteral, TArrayView<const FGuid> InPageOrder)
+	{
+		ForEachLiteral(InClass.GetDefaultInterface(), OnLiteral, InPageOrder);
+	}
+
+	void ForEachLiteral(const FMetasoundFrontendClassInterface& InClassInterface, FForEachLiteralFunctionRef OnLiteral, TArrayView<const FGuid> InPageOrder)
+	{
+		for (const FMetasoundFrontendClassInput& ClassInput : InClassInterface.Inputs)
+		{
+			const FMetasoundFrontendClassInputDefault* DefaultInput = FindPreferredPage<FMetasoundFrontendClassInputDefault>(ClassInput.GetDefaults(), InPageOrder);
+			if (DefaultInput)
+			{
+				OnLiteral(ClassInput.TypeName, DefaultInput->Literal);
+			}
+		}
+	}
+
+	void ForEachLiteral(const FMetasoundFrontendNode& InNode, FForEachLiteralFunctionRef OnLiteral, TArrayView<const FGuid> InPageOrder)
+	{
+		for (const FMetasoundFrontendVertexLiteral& VertexLiteral : InNode.InputLiterals)
+		{
+			auto HasEqualVertexID = [&VertexLiteral](const FMetasoundFrontendVertex& InVertex) -> bool
+			{ 
+				return InVertex.VertexID == VertexLiteral.VertexID; 
+			};
+
+			if (const FMetasoundFrontendVertex* InputVertex = InNode.Interface.Inputs.FindByPredicate(HasEqualVertexID))
+			{
+				OnLiteral(InputVertex->TypeName, VertexLiteral.Value);
+			}
+		}
+
+		if (InNode.ClassInterfaceOverride.IsValid())
+		{
+			ForEachLiteral(InNode.ClassInterfaceOverride.Get(), OnLiteral, InPageOrder);
 		}
 	}
 } // namespace Metasound::Frontend
diff --git a/Engine/Plugins/Runtime/Metasound/Source/MetasoundFrontend/Private/MetasoundFrontendDocumentVersioning.cpp b/Engine/Plugins/Runtime/Metasound/Source/MetasoundFrontend/Private/MetasoundFrontendDocumentVersioning.cpp
--- a/Engine/Plugins/Runtime/Metasound/Source/MetasoundFrontend/Private/MetasoundFrontendDocumentVersioning.cpp
+++ b/Engine/Plugins/Runtime/Metasound/Source/MetasoundFrontend/Private/MetasoundFrontendDocumentVersioning.cpp
@@ -196,6 +196,7 @@ PRAGMA_ENABLE_DEPRECATION_WARNINGS
 	#if WITH_EDITOR
 				FGraphHandle GraphHandle = InDocument->GetRootGraph();
 				TArray<FNodeHandle> FrontendNodes = GraphHandle->GetNodes();
+				TArray<FGuid> PageOrder({DefaultPageID}); // pages did not exist at this point in the history of metasound development.
 
 				// Before literals could be stored on node inputs directly, they were stored
 				// by creating hidden input nodes. Update the doc by finding all hidden input
@@ -211,7 +212,7 @@ PRAGMA_ENABLE_DEPRECATION_WARNINGS
 					{
 						// Get literal value from input node.
 						const FGuid VertexID = GraphHandle->GetVertexIDForInputVertex(NodeHandle->GetNodeName());
-						const FMetasoundFrontendLiteral DefaultLiteral = GraphHandle->GetDefaultInput(VertexID);
+						const FMetasoundFrontendLiteral DefaultLiteral = GraphHandle->GetDefaultInput(VertexID, PageOrder);
 
 						// Apply literal value to downstream node's inputs.
 						TArray<FOutputHandle> OutputHandles = NodeHandle->GetOutputs();
diff --git a/Engine/Plugins/Runtime/Metasound/Source/MetasoundFrontend/Private/MetasoundFrontendGraphBuilder.cpp b/Engine/Plugins/Runtime/Metasound/Source/MetasoundFrontend/Private/MetasoundFrontendGraphBuilder.cpp
--- a/Engine/Plugins/Runtime/Metasound/Source/MetasoundFrontend/Private/MetasoundFrontendGraphBuilder.cpp
+++ b/Engine/Plugins/Runtime/Metasound/Source/MetasoundFrontend/Private/MetasoundFrontendGraphBuilder.cpp
@@ -13,6 +13,7 @@
 #include "MetasoundFrontend.h"
 #include "MetasoundFrontendDataTypeRegistry.h"
 #include "MetasoundFrontendNodeTemplateRegistry.h"
+#include "MetasoundFrontendPages.h"
 #include "MetasoundFrontendProxyDataCache.h"
 #include "MetasoundFrontendRegistries.h"
 #include "MetasoundGraph.h"
@@ -37,6 +38,7 @@ namespace Metasound::Frontend
 			const FString& DebugAssetName;
 			const Frontend::IDataTypeRegistry& DataTypeRegistry;
 			const Frontend::FProxyDataCache* ProxyDataCache;
+			TArrayView<const FGuid> PageOrder;
 		};
 
 		// Context related to the document being built. 
@@ -56,26 +58,12 @@ namespace Metasound::Frontend
 			const FBuildDocumentContext& BuildDocumentContext;
 		};
 
-		template <typename ResolveType>
-		FGuid ResolveTargetPageID(const ResolveType& InToResolve)
-		{
-			// Registry is not available in tests, so for now resolution is considered successful at this level
-			// if registry is not initialized and providing a resolved page ID. TODO: Add a test implementation
-			// that returns the default page (or whatever page behavior is desired for testing).
-			if (IDocumentBuilderRegistry* BuilderRegistry = IDocumentBuilderRegistry::Get())
-			{
-				return IDocumentBuilderRegistry::GetChecked().ResolveTargetPageID(InToResolve);
-			}
-
-			return Frontend::DefaultPageID;
-		}
-
 		bool InterfacesHaveEqualSize(const FMetasoundFrontendClassInterface& InClassInterface, const FMetasoundFrontendNodeInterface& InNodeInterface)
 		{
 			return (InClassInterface.Inputs.Num() == InNodeInterface.Inputs.Num()) && (InClassInterface.Outputs.Num() == InNodeInterface.Outputs.Num());
 		}
 
-		const FMetasoundFrontendLiteral* FindLiteralForInputVertex(const FVertexName& InVertexName, const FMetasoundFrontendNode& InNode, const FMetasoundFrontendClassInput& InNodeClassInput)
+		const FMetasoundFrontendLiteral* FindLiteralForInputVertex(const FBuildContext& InContext, const FVertexName& InVertexName, const FMetasoundFrontendNode& InNode, const FMetasoundFrontendClassInput& InNodeClassInput)
 		{
 			using namespace Frontend;
 
@@ -103,13 +91,12 @@ namespace Metasound::Frontend
 			// Check for default value on node class
 			if (nullptr == Literal)
 			{
-				const FGuid PageID = GraphBuilderPrivate::ResolveTargetPageID(InNodeClassInput);
-				const FMetasoundFrontendLiteral* DefaultLiteral = InNodeClassInput.FindConstDefault(PageID);
-				if (ensure(DefaultLiteral))
+				const FMetasoundFrontendClassInputDefault* InputDefault= FindPreferredPage(InNodeClassInput.GetDefaults(), InContext.PageOrder);
+				if (ensure(InputDefault))
 				{
-					if (DefaultLiteral->IsValid())
+					if (InputDefault->Literal.IsValid())
 					{
-						Literal = DefaultLiteral;
+						Literal = &InputDefault->Literal;
 					}
 				}
 			}
@@ -117,7 +104,7 @@ namespace Metasound::Frontend
 			return Literal;
 		}
 
-		const FMetasoundFrontendLiteral* FindInputLiteralForInputNode(const FMetasoundFrontendNode& InInputNode, const FMetasoundFrontendClassInterface& InInputNodeClassInterface, const FMetasoundFrontendClassInput& InOwningGraphClassInput)
+		const FMetasoundFrontendLiteral* FindInputLiteralForInputNode(const FBuildContext& InContext, const FMetasoundFrontendNode& InInputNode, const FMetasoundFrontendClassInterface& InInputNodeClassInterface, const FMetasoundFrontendClassInput& InOwningGraphClassInput)
 		{
 			using namespace Frontend;
 
@@ -151,11 +138,13 @@ namespace Metasound::Frontend
 			if (nullptr == Literal)
 			{
 				// Find Class Default that is not invalid
-				const FGuid PageID = GraphBuilderPrivate::ResolveTargetPageID(InOwningGraphClassInput);
-				const FMetasoundFrontendLiteral& DefaultLiteral = InOwningGraphClassInput.FindConstDefaultChecked(PageID);
-				if (DefaultLiteral.IsValid())
+				const FMetasoundFrontendClassInputDefault* InputDefault = FindPreferredPage(InOwningGraphClassInput.GetDefaults(), InContext.PageOrder);
+				if (ensure(InputDefault))
 				{
-					Literal = &DefaultLiteral;
+					if (InputDefault->Literal.IsValid())
+					{
+						Literal = &InputDefault->Literal;
+					}
 				}
 			}
 
@@ -163,11 +152,13 @@ namespace Metasound::Frontend
 			if (nullptr == Literal && ensure(InInputNodeClassInterface.Inputs.Num() == 1))
 			{
 				const FMetasoundFrontendClassInput& InputNodeClassInput = InInputNodeClassInterface.Inputs.Last();
-				const FGuid PageID = GraphBuilderPrivate::ResolveTargetPageID(InputNodeClassInput);
-				const FMetasoundFrontendLiteral& DefaultLiteral = InputNodeClassInput.FindConstDefaultChecked(PageID);
-				if (DefaultLiteral.IsValid())
+				const FMetasoundFrontendClassInputDefault* InputDefault = FindPreferredPage(InputNodeClassInput.GetDefaults(), InContext.PageOrder);
+				if (ensure(InputDefault))
 				{
-					Literal = &DefaultLiteral;
+					if (InputDefault->Literal.IsValid())
+					{
+						Literal = &InputDefault->Literal;
+					}
 				}
 			}
 
@@ -328,7 +319,7 @@ namespace Metasound::Frontend
 								check(InNode.Interface.Inputs.Num() == 1);
 								const FMetasoundFrontendVertex& NodeInputVertex = InNode.Interface.Inputs[0];
 								const FMetasoundFrontendClassInput& ClassInputVertex = InClassInterface.Inputs[0];
-								const FMetasoundFrontendLiteral* FrontendLiteral = FindInputLiteralForInputNode(InNode, InClassInterface, *InOwningGraphClassInput);
+								const FMetasoundFrontendLiteral* FrontendLiteral = FindInputLiteralForInputNode(InContext, InNode, InClassInterface, *InOwningGraphClassInput);
 								
 								// NOTE: The access type reported on the ClassInputVertex may be incorrect. When the access type is changed in the editor,
 								// the access type of the input node's FMetasoundFrontendClass is not updated. 
@@ -344,7 +335,7 @@ namespace Metasound::Frontend
 								{
 									const FMetasoundFrontendVertex& NodeInputVertex = InNode.Interface.Inputs[0];
 									const FMetasoundFrontendClassInput& ClassInputVertex = InClassInterface.Inputs[0];
-									const FMetasoundFrontendLiteral* FrontendLiteral = FindLiteralForInputVertex(NodeInputVertex.Name, InNode, ClassInputVertex);
+									const FMetasoundFrontendLiteral* FrontendLiteral = FindLiteralForInputVertex(InContext, NodeInputVertex.Name, InNode, ClassInputVertex);
 									// NOTE: The access type reported on the ClassInputVertex may be incorrect. When the access type is changed in the editor,
 									// the access type of the output node's FMetasoundFrontendClass is not updated. 
 									Vertices.Add(MakeVertex(NodeInputVertex.Name, NodeInputVertex.TypeName, InOwningGraphClassOutput->AccessType, FrontendLiteral));
@@ -370,7 +361,7 @@ namespace Metasound::Frontend
 							// node config is done as a feature. 
 							for (const FMetasoundFrontendClassInput& ClassInput : InClassInterface.Inputs)
 							{
-								const FMetasoundFrontendLiteral* FrontendLiteral = FindLiteralForInputVertex(ClassInput.Name, InNode, ClassInput);
+								const FMetasoundFrontendLiteral* FrontendLiteral = FindLiteralForInputVertex(InContext, ClassInput.Name, InNode, ClassInput);
 								Vertices.Add(MakeVertex(ClassInput.Name, ClassInput.TypeName, ClassInput.AccessType, FrontendLiteral));
 							}
 						}
@@ -603,7 +594,8 @@ namespace Metasound::Frontend
 						InGraphContext.GraphClass,
 						InGraphContext.BuildContext.ProxyDataCache, 
 						&InGraphContext.BuildContext.DataTypeRegistry,
-						&InGraphContext.BuildDocumentContext.Graphs
+						&InGraphContext.BuildDocumentContext.Graphs,
+						InGraphContext.BuildContext.PageOrder
 					};
 					TUniquePtr<INode> Node = FGraphBuilder::CreateNode(CreateNodeParams);
 
@@ -784,13 +776,18 @@ namespace Metasound::Frontend
 		{
 			const FString GraphName = InContext.DebugAssetName;
 
-			const FGuid PageID = GraphBuilderPrivate::ResolveTargetPageID(InGraphClass);
-			const FMetasoundFrontendGraph& PageGraph = InGraphClass.FindConstGraphChecked(PageID);
+			const FMetasoundFrontendGraph* PageGraph = FindPreferredPage(InGraphClass.GetConstGraphPages(), InContext.PageOrder);
+			if (!PageGraph)
+			{
+				UE_LOG(LogMetaSound, Error, TEXT("Cannot create FrontendGraph because could not find preferred graph page."));
+				return nullptr;
+			}
+
 			FBuildGraphContext BuildGraphContext
 			{
 				MakeUnique<FFrontendGraph>(GraphName, InGraphId),
 				InGraphClass,
-				PageGraph,
+				*PageGraph,
 				InContext,
 				InDocumentContext
 			};
@@ -839,7 +836,8 @@ namespace Metasound::Frontend
 		{
 			InParams.DebugAssetName,
 			InParams.DataTypeRegistry ? *InParams.DataTypeRegistry : IDataTypeRegistry::Get(),
-			InParams.ProxyDataCache
+			InParams.ProxyDataCache,
+			InParams.PageOrder
 		};
 
 		switch (InParams.FrontendNodeClass.Metadata.GetType())
@@ -932,53 +930,8 @@ PRAGMA_ENABLE_DEPRECATION_WARNINGS
 		return Node;
 	}
 
-	/** Check that all dependencies are C++ class dependencies. */
-	bool FGraphBuilder::IsFlat(const FMetasoundFrontendDocument& InDocument)
-	{
-		if (InDocument.Subgraphs.Num() > 0)
-		{
-			return false;
-		}
-
-PRAGMA_DISABLE_DEPRECATION_WARNINGS
-		return IsFlat(InDocument.RootGraph, InDocument.Dependencies);
-PRAGMA_ENABLE_DEPRECATION_WARNINGS
-	}
-
-	bool FGraphBuilder::IsFlat(const FMetasoundFrontendGraphClass& InRoot, const TArray<FMetasoundFrontendClass>& InDependencies)
-	{
-		// All dependencies are external dependencies in a flat graph
-		auto IsClassExternal = [&](const FMetasoundFrontendClass& InDesc) 
-		{
-			const bool bIsExternal = (InDesc.Metadata.GetType() == EMetasoundFrontendClassType::External) ||
-				(InDesc.Metadata.GetType() == EMetasoundFrontendClassType::Template) ||
-				(InDesc.Metadata.GetType() == EMetasoundFrontendClassType::Input) ||
-				(InDesc.Metadata.GetType() == EMetasoundFrontendClassType::Output);
-			return bIsExternal;
-		};
-		const bool bIsEveryDependencyExternal = Algo::AllOf(InDependencies, IsClassExternal);
-
-		if (!bIsEveryDependencyExternal)
-		{
-			return false;
-		}
-
-		// All the dependencies are met 
-		TSet<FGuid> AvailableDependencies;
-		Algo::Transform(InDependencies, AvailableDependencies, [](const FMetasoundFrontendClass& InDesc) { return InDesc.ID; });
-
-		auto IsDependencyMet = [&](const FMetasoundFrontendNode& InNode) 
-		{ 
-			return AvailableDependencies.Contains(InNode.ClassID);
-		};
-
-		const bool bIsEveryDependencyMet = Algo::AllOf(InRoot.GetConstDefaultGraph().Nodes, IsDependencyMet);
-
-		return bIsEveryDependencyMet;
-	}
-
 
-	TUniquePtr<FFrontendGraph> FGraphBuilder::CreateGraph(const FMetasoundFrontendGraphClass& InGraph, const TArray<FMetasoundFrontendGraphClass>& InSubgraphs, const TArray<FMetasoundFrontendClass>& InDependencies, const Frontend::FProxyDataCache& InProxyDataCache, const FString& InDebugAssetName, const FGuid InGraphId)
+	TUniquePtr<FFrontendGraph> FGraphBuilder::CreateGraph(const FMetasoundFrontendGraphClass& InGraph, const TArray<FMetasoundFrontendGraphClass>& InSubgraphs, const TArray<FMetasoundFrontendClass>& InDependencies, const Frontend::FProxyDataCache& InProxyDataCache, const FString& InDebugAssetName, const FGuid InGraphId, TArrayView<const FGuid> InPageOrder)
 	{
 		using namespace GraphBuilderPrivate;
 
@@ -987,6 +940,7 @@ PRAGMA_ENABLE_DEPRECATION_WARNINGS
 			InDebugAssetName, 					// DebugAssetName
 			Frontend::IDataTypeRegistry::Get(), // DataTypeRegistry
 			&InProxyDataCache, 					// ProxyDataCache
+			InPageOrder                         // Page resolution order
 		};
 
 		FBuildDocumentContext DocumentContext;
@@ -1031,37 +985,37 @@ PRAGMA_ENABLE_DEPRECATION_WARNINGS
 		return CreateGraphInternal(Context, DocumentContext, InGraph, InGraphId);
 	}
 
-	TUniquePtr<FFrontendGraph> FGraphBuilder::CreateGraph(const FMetasoundFrontendDocument& InDocument, const FString& InDebugAssetName)
+	TUniquePtr<FFrontendGraph> FGraphBuilder::CreateGraph(const FMetasoundFrontendDocument& InDocument, const FString& InDebugAssetName, TArrayView<const FGuid> InPageOrder)
 	{
 		// Create proxies before creating graph.
 		Frontend::FProxyDataCache ProxyDataCache;
-		ProxyDataCache.CreateAndCacheProxies(InDocument);
+		ProxyDataCache.CreateAndCacheProxies(InDocument, InPageOrder);
 		
-		return CreateGraph(InDocument, ProxyDataCache, InDebugAssetName, Frontend::CreateLocallyUniqueId());
+		return CreateGraph(InDocument, ProxyDataCache, InDebugAssetName, Frontend::CreateLocallyUniqueId(), InPageOrder);
 	}
 
-	TUniquePtr<FFrontendGraph> FGraphBuilder::CreateGraph(const FMetasoundFrontendGraphClass& InGraph, const TArray<FMetasoundFrontendGraphClass>& InSubgraphs, const TArray<FMetasoundFrontendClass>& InDependencies, const FString& InDebugAssetName)
+	TUniquePtr<FFrontendGraph> FGraphBuilder::CreateGraph(const FMetasoundFrontendGraphClass& InGraph, const TArray<FMetasoundFrontendGraphClass>& InSubgraphs, const TArray<FMetasoundFrontendClass>& InDependencies, const FString& InDebugAssetName, TArrayView<const FGuid> InPageOrder)
 	{
 		// Create proxies before building graph
 		Frontend::FProxyDataCache ProxyDataCache;
-		ProxyDataCache.CreateAndCacheProxies(InGraph);
+		ProxyDataCache.CreateAndCacheProxies(InGraph, InPageOrder);
 
 		for (const FMetasoundFrontendGraphClass& SubgraphClass : InSubgraphs)
 		{
-			ProxyDataCache.CreateAndCacheProxies(SubgraphClass);
+			ProxyDataCache.CreateAndCacheProxies(SubgraphClass, InPageOrder);
 		}
 
 		for (const FMetasoundFrontendClass& DependencyClass : InDependencies)
 		{
-			ProxyDataCache.CreateAndCacheProxies(DependencyClass);
+			ProxyDataCache.CreateAndCacheProxies(DependencyClass, InPageOrder);
 		}
 
-		return CreateGraph(InGraph, InSubgraphs, InDependencies, ProxyDataCache, InDebugAssetName, Frontend::CreateLocallyUniqueId());
+		return CreateGraph(InGraph, InSubgraphs, InDependencies, ProxyDataCache, InDebugAssetName, Frontend::CreateLocallyUniqueId(), InPageOrder);
 	}
 
-	TUniquePtr<FFrontendGraph> FGraphBuilder::CreateGraph(const FMetasoundFrontendDocument& InDocument, const Frontend::FProxyDataCache& InProxyDataCache, const FString& InDebugAssetName, const FGuid InGraphId)
+	TUniquePtr<FFrontendGraph> FGraphBuilder::CreateGraph(const FMetasoundFrontendDocument& InDocument, const Frontend::FProxyDataCache& InProxyDataCache, const FString& InDebugAssetName, const FGuid InGraphId, TArrayView<const FGuid> InPageOrder)
 	{
-		return CreateGraph(InDocument.RootGraph, InDocument.Subgraphs, InDocument.Dependencies, InProxyDataCache, InDebugAssetName, InGraphId);
+		return CreateGraph(InDocument.RootGraph, InDocument.Subgraphs, InDocument.Dependencies, InProxyDataCache, InDebugAssetName, InGraphId, InPageOrder);
 	}
 }
 
diff --git a/Engine/Plugins/Runtime/Metasound/Source/MetasoundFrontend/Private/MetasoundFrontendGraphController.cpp b/Engine/Plugins/Runtime/Metasound/Source/MetasoundFrontend/Private/MetasoundFrontendGraphController.cpp
--- a/Engine/Plugins/Runtime/Metasound/Source/MetasoundFrontend/Private/MetasoundFrontendGraphController.cpp
+++ b/Engine/Plugins/Runtime/Metasound/Source/MetasoundFrontend/Private/MetasoundFrontendGraphController.cpp
@@ -13,6 +13,7 @@
 #include "MetasoundFrontendNodeClassRegistry.h"
 #include "MetasoundFrontendNodeController.h"
 #include "MetasoundFrontendNodeTemplateRegistry.h"
+#include "MetasoundFrontendPages.h"
 #include "MetasoundFrontendProxyDataCache.h"
 #include "MetasoundFrontendSubgraphNodeController.h"
 #include "MetasoundFrontendInvalidController.h"
@@ -1259,34 +1260,52 @@ namespace Metasound
 		}
 
 		FMetasoundFrontendLiteral FGraphController::GetDefaultInput(const FGuid& InVertexID) const
+		{
+			TArray<FGuid> PageOrder;
+			if (IDocumentBuilderRegistry* DocRegistry = IDocumentBuilderRegistry::Get())
+			{
+PRAGMA_DISABLE_DEPRECATION_WARNINGS
+				PageOrder = DocRegistry->GetPageOrder();
+PRAGMA_ENABLE_DEPRECATION_WARNINGS
+			}
+			return GetDefaultInput(InVertexID, PageOrder);
+		}
+
+		FMetasoundFrontendLiteral FGraphController::GetDefaultInput(const FGuid& InVertexID, TArrayView<const FGuid> InPageOrder) const
 		{
 			if (const FMetasoundFrontendClassInput* Desc = FindInputDescriptionWithVertexID(InVertexID))
 			{
-				FGuid TargetPageID = Frontend::DefaultPageID;
-				if (IDocumentBuilderRegistry* DocRegistry = IDocumentBuilderRegistry::Get())
+				if (const FMetasoundFrontendClassInputDefault* Default = FindPreferredPage(Desc->GetDefaults(), InPageOrder))
 				{
-					TargetPageID = DocRegistry->ResolveTargetPageID(*Desc);
+					return Default->Literal;
 				}
-				const FMetasoundFrontendLiteral& DefaultLiteral = Desc->FindConstDefaultChecked(TargetPageID);
-				return DefaultLiteral;
 			}
 			return FMetasoundFrontendLiteral{};
 		}
 
 		bool FGraphController::SetDefaultInput(const FGuid& InVertexID, const FMetasoundFrontendLiteral& InLiteral)
+		{
+			TArray<FGuid> PageOrder;
+			if (IDocumentBuilderRegistry* DocRegistry = IDocumentBuilderRegistry::Get())
+			{
+PRAGMA_DISABLE_DEPRECATION_WARNINGS
+				PageOrder = DocRegistry->GetPageOrder();
+PRAGMA_ENABLE_DEPRECATION_WARNINGS
+			}
+			return SetDefaultInput(InVertexID, InLiteral, PageOrder);
+		}
+
+		bool FGraphController::SetDefaultInput(const FGuid& InVertexID, const FMetasoundFrontendLiteral& InLiteral, TArrayView<const FGuid> InPageOrder)
 		{
 			if (FMetasoundFrontendClassInput* Desc = FindInputDescriptionWithVertexID(InVertexID))
 			{
 				if (ensure(IDataTypeRegistry::Get().IsLiteralTypeSupported(Desc->TypeName, InLiteral.GetType())))
 				{
-					FGuid TargetPageID = Frontend::DefaultPageID;
-					if (IDocumentBuilderRegistry* DocRegistry = IDocumentBuilderRegistry::Get())
+					if (const FMetasoundFrontendClassInputDefault* Default = FindPreferredPage(Desc->GetDefaults(), InPageOrder))
 					{
-						TargetPageID = DocRegistry->ResolveTargetPageID(*Desc);
+						Desc->FindDefaultChecked(Default->PageID) = InLiteral;
+						return true;
 					}
-					FMetasoundFrontendLiteral& DefaultLiteral = Desc->FindDefaultChecked(TargetPageID);
-					DefaultLiteral = InLiteral;
-					return true;
 				}
 				else
 				{
@@ -1640,35 +1659,6 @@ namespace Metasound
 			return INodeController::GetInvalidHandle();
 		}
 
-		TUniquePtr<IOperator> FGraphController::BuildOperator(const FOperatorSettings& InSettings, const FMetasoundEnvironment& InEnvironment, FBuildResults& OutResults) const
-		{
-			if (const FMetasoundFrontendGraphClass* GraphClass = GraphClassPtr.Get())
-			{
-				const TArray<FMetasoundFrontendGraphClass>& Subgraphs = OwningDocument->GetSubgraphs();
-				const TArray<FMetasoundFrontendClass>& Dependencies = OwningDocument->GetDependencies();
-
-				FString UnknownAsset = TEXT("UnknownAsset");
-				FProxyDataCache ProxyCache;
-				const FMetasoundFrontendDocument& Doc = *(OwningDocument->GetDocumentPtr().Get());
-				ProxyCache.CreateAndCacheProxies(Doc);
-
-				TUniquePtr<FFrontendGraph> Graph = FGraphBuilder::CreateGraph(*GraphClass, Subgraphs, Dependencies, ProxyCache, UnknownAsset, Frontend::CreateLocallyUniqueId());
-
-				if (!Graph.IsValid())
-				{
-					return TUniquePtr<IOperator>(nullptr);
-				}
-
-				FInputVertexInterfaceData InterfaceData;
-				FBuildGraphOperatorParams BuildParams { *Graph, InSettings, InterfaceData, InEnvironment };
-				return FOperatorBuilder(FOperatorBuilderSettings::GetDefaultSettings()).BuildGraphOperator(BuildParams, OutResults);
-			}
-			else
-			{
-				return TUniquePtr<IOperator>(nullptr);
-			}
-		}
-
 		FDocumentHandle FGraphController::GetOwningDocument()
 		{
 			return OwningDocument;
diff --git a/Engine/Plugins/Runtime/Metasound/Source/MetasoundFrontend/Private/MetasoundFrontendGraphController.h b/Engine/Plugins/Runtime/Metasound/Source/MetasoundFrontend/Private/MetasoundFrontendGraphController.h
--- a/Engine/Plugins/Runtime/Metasound/Source/MetasoundFrontend/Private/MetasoundFrontendGraphController.h
+++ b/Engine/Plugins/Runtime/Metasound/Source/MetasoundFrontend/Private/MetasoundFrontendGraphController.h
@@ -142,11 +142,18 @@ namespace Metasound
 			// For inputs whose preferred literal type is UObject or UObjectArray, this can be used to determine the UObject corresponding to that input's datatype.
 			UClass* GetSupportedClassForInputVertex(const FVertexName& InInputName) override;
 
-			FMetasoundFrontendLiteral GetDefaultInput(const FGuid& InVertexID) const override;
+			UE_DEPRECATED(5.7, "Use version which provides page order")
+			virtual FMetasoundFrontendLiteral GetDefaultInput(const FGuid& InVertexID) const override;
+
+			virtual FMetasoundFrontendLiteral GetDefaultInput(const FGuid& InVertexID, TArrayView<const FGuid> InPageOrder) const override;
 
 			// These can be used to set the default value for a given input on this graph.
 			// @returns false if the input name couldn't be found, or if the literal type was incompatible with the Data Type of this input.
-			bool SetDefaultInput(const FGuid& InVertexID, const FMetasoundFrontendLiteral& InLiteral) override;
+
+			UE_DEPRECATED(5.7, "Use version which provides page order")
+			virtual bool SetDefaultInput(const FGuid& InVertexID, const FMetasoundFrontendLiteral& InLiteral) override;
+
+			virtual bool SetDefaultInput(const FGuid& InVertexID, const FMetasoundFrontendLiteral& InLiteral, TArrayView<const FGuid> InPageOrder) override;
 			bool SetDefaultInputToDefaultLiteralOfType(const FGuid& InVertexID) override;
 
 #if WITH_EDITOR
@@ -200,8 +207,6 @@ namespace Metasound
 
 			FNodeHandle CreateEmptySubgraph(const FMetasoundFrontendClassMetadata& InInfo) override;
 
-			TUniquePtr<IOperator> BuildOperator(const FOperatorSettings& InSettings, const FMetasoundEnvironment& InEnvironment, FBuildResults& OutBuildErrors) const override;
-
 			FDocumentHandle GetOwningDocument() override;
 			FConstDocumentHandle GetOwningDocument() const override;
 
diff --git a/Engine/Plugins/Runtime/Metasound/Source/MetasoundFrontend/Private/MetasoundFrontendInvalidController.h b/Engine/Plugins/Runtime/Metasound/Source/MetasoundFrontend/Private/MetasoundFrontendInvalidController.h
--- a/Engine/Plugins/Runtime/Metasound/Source/MetasoundFrontend/Private/MetasoundFrontendInvalidController.h
+++ b/Engine/Plugins/Runtime/Metasound/Source/MetasoundFrontend/Private/MetasoundFrontendInvalidController.h
@@ -386,11 +386,19 @@ namespace Metasound
 
 			virtual FGuid GetVertexIDForInputVertex(const FVertexName& InInputName) const override { return Metasound::FrontendInvalidID; }
 			virtual FGuid GetVertexIDForOutputVertex(const FVertexName& InOutputName) const override { return Metasound::FrontendInvalidID; }
+			UE_DEPRECATED(5.7, "Use version which provides page order")
 			virtual FMetasoundFrontendLiteral GetDefaultInput(const FGuid& InVertexID) const override { return FMetasoundFrontendLiteral{}; }
 
+			virtual FMetasoundFrontendLiteral GetDefaultInput(const FGuid& InVertexID, TArrayView<const FGuid> InPageOrder) const override { return FMetasoundFrontendLiteral{}; }
+
 			// These can be used to set the default value for a given input on this graph.
 			// @returns false if the input name couldn't be found, or if the literal type was incompatible with the Data Type of this input.
+
+			UE_DEPRECATED(5.7, "Use version which provides page order")
 			virtual bool SetDefaultInput(const FGuid& InVertexID, const FMetasoundFrontendLiteral& InLiteral) override { return false; }
+
+			virtual bool SetDefaultInput(const FGuid& InVertexID, const FMetasoundFrontendLiteral& InLiteral, TArrayView<const FGuid> InPageOrder) override { return false; }
+
 			virtual bool SetDefaultInputToDefaultLiteralOfType(const FGuid& InVertexID) override { return false; }
 
 #if WITH_EDITOR
diff --git a/Engine/Plugins/Runtime/Metasound/Source/MetasoundFrontend/Private/MetasoundFrontendNodeClassRegistry.cpp b/Engine/Plugins/Runtime/Metasound/Source/MetasoundFrontend/Private/MetasoundFrontendNodeClassRegistry.cpp
--- a/Engine/Plugins/Runtime/Metasound/Source/MetasoundFrontend/Private/MetasoundFrontendNodeClassRegistry.cpp
+++ b/Engine/Plugins/Runtime/Metasound/Source/MetasoundFrontend/Private/MetasoundFrontendNodeClassRegistry.cpp
@@ -23,6 +23,7 @@
 #include "MetasoundFrontendDataTypeRegistry.h"
 #include "MetasoundFrontendDocumentBuilder.h"
 #include "MetasoundFrontendGraph.h"
+#include "MetasoundFrontendPages.h"
 #include "MetasoundFrontendProxyDataCache.h"
 #include "MetasoundFrontendRegistryTransaction.h"
 #include "MetasoundFrontendSearchEngine.h"
@@ -340,7 +341,7 @@ PRAGMA_ENABLE_DEPRECATION_WARNINGS
 		}
 	} // namespace NodeClassRegistryKey
 
-	void FNodeClassRegistry::BuildAndRegisterGraphFromDocument(const FMetasoundFrontendDocument& InDocument, const FProxyDataCache& InProxyDataCache, FNodeClassInfo&& InNodeClassInfo, const FTopLevelAssetPath& AssetPath)
+	void FNodeClassRegistry::BuildAndRegisterGraphFromDocument(const FMetasoundFrontendDocument& InDocument, const FProxyDataCache& InProxyDataCache, FNodeClassInfo&& InNodeClassInfo, const FTopLevelAssetPath& AssetPath, TArrayView<const FGuid> InPageOrder)
 	{
 		using namespace RegistryPrivate;
 
@@ -359,7 +360,7 @@ PRAGMA_ENABLE_DEPRECATION_WARNINGS
 		}
 
 		// Use the asset class id for the graph id because it should be locally unique per asset.
-		TUniquePtr<FFrontendGraph> FrontendGraph = Frontend::FGraphBuilder::CreateGraph(InDocument, InProxyDataCache, AssetPath.ToString(), /*GraphId=*/AssetClassID);
+		TUniquePtr<FFrontendGraph> FrontendGraph = Frontend::FGraphBuilder::CreateGraph(InDocument, InProxyDataCache, AssetPath.ToString(), /*GraphId=*/AssetClassID, InPageOrder);
 		if (!FrontendGraph.IsValid())
 		{
 			UE_LOG(LogMetaSound, Error, TEXT("Failed to build MetaSound graph in asset '%s'"), *AssetPath.ToString());
@@ -514,7 +515,7 @@ PRAGMA_ENABLE_DEPRECATION_WARNINGS
 		return MakeUnique<FNodeClassRegistryTransactionStream>(TransactionBuffer);
 	}
 
-	FGraphRegistryKey FNodeClassRegistry::RegisterGraph(const TScriptInterface<IMetaSoundDocumentInterface>& InDocumentInterface)
+	FGraphRegistryKey FNodeClassRegistry::RegisterGraph(const TScriptInterface<IMetaSoundDocumentInterface>& InDocumentInterface, TArrayView<const FGuid> InPageOrder)
 	{
 		using namespace UE;
 
@@ -551,18 +552,23 @@ PRAGMA_ENABLE_DEPRECATION_WARNINGS
 		// Proxies are created synchronously to avoid creating proxies in async tasks. Proxies
 		// are created from UObjects which need to be protected from GC and non-GT access.
 		FProxyDataCache ProxyDataCache;
-		ProxyDataCache.CreateAndCacheProxies(Document);
+		ProxyDataCache.CreateAndCacheProxies(Document, InPageOrder);
 
 #if !NO_LOGGING
 		if (UE_LOG_ACTIVE(LogMetaSound, Verbose))
 		{
-			const FGuid PageID = IDocumentBuilderRegistry::GetChecked().ResolveTargetPageID(Document.RootGraph);
+			FGuid GraphPageID = DefaultPageID;
+			if (const FMetasoundFrontendGraph* GraphPage = FindPreferredPage(Document.RootGraph.GetConstGraphPages(), InPageOrder))
+			{
+				GraphPageID = GraphPage->PageID;
+			}
+
 			const bool bContainsMultipleGraphs = Document.RootGraph.GetConstGraphPages().Num() > 1;
-			if (bContainsMultipleGraphs || PageID != Metasound::Frontend::DefaultPageID)
+			if (bContainsMultipleGraphs || GraphPageID != Metasound::Frontend::DefaultPageID)
 			{
 				UE_LOG(LogMetaSound, Verbose, TEXT("Registered MetaSound '%s' Graph Page with PageID '%s'."),
 					*AssetPath.GetAssetName().ToString(),
-					*PageID.ToString());
+					*GraphPageID.ToString());
 				if (bContainsMultipleGraphs)
 				{
 					UE_LOG(LogMetaSound, Verbose, TEXT("Graphs found with following PageIDs Implemented:"));
@@ -585,14 +591,15 @@ PRAGMA_ENABLE_DEPRECATION_WARNINGS
 
 		if (bAsync)
 		{
+			TArray<FGuid> PageOrder(InPageOrder.GetData(), InPageOrder.Num());
 			Tasks::FTask BuildAndRegisterTask = AsyncRegistrationPipe.Launch(
 				UE_SOURCE_LOCATION,
-				[RegistryKey, ClassInfo = MoveTemp(NodeClassInfo), AssetPath, RegistryDocInterface, ProxyDataCache = MoveTemp(ProxyDataCache)]() mutable
+				[RegistryKey, ClassInfo = MoveTemp(NodeClassInfo), AssetPath, RegistryDocInterface, ProxyDataCache = MoveTemp(ProxyDataCache), PageOrder = MoveTemp(PageOrder)]() mutable
 				{
 					FNodeClassRegistry& Registry = FNodeClassRegistry::Get();
 					// Unregister the graph before re-registering
 					Registry.UnregisterGraphInternal(RegistryKey);
-					Registry.BuildAndRegisterGraphFromDocument(RegistryDocInterface->GetConstDocument(), ProxyDataCache, MoveTemp(ClassInfo), AssetPath);
+					Registry.BuildAndRegisterGraphFromDocument(RegistryDocInterface->GetConstDocument(), ProxyDataCache, MoveTemp(ClassInfo), AssetPath, PageOrder);
 					Registry.RemoveRegistrationTask(RegistryKey, FNodeClassRegistryTransaction::ETransactionType::NodeRegistration);
 					Registry.RemoveDocumentReference(RegistryDocInterface);
 				}
@@ -611,7 +618,7 @@ PRAGMA_ENABLE_DEPRECATION_WARNINGS
 			UnregisterGraphInternal(RegistryKey);
 
 			// Build and register graph synchronously
-			BuildAndRegisterGraphFromDocument(RegistryDocInterface->GetConstDocument(), ProxyDataCache, MoveTemp(NodeClassInfo), AssetPath);
+			BuildAndRegisterGraphFromDocument(RegistryDocInterface->GetConstDocument(), ProxyDataCache, MoveTemp(NodeClassInfo), AssetPath, InPageOrder);
 		}
 
 		return RegistryKey;
diff --git a/Engine/Plugins/Runtime/Metasound/Source/MetasoundFrontend/Private/MetasoundFrontendNodeClassRegistryPrivate.h b/Engine/Plugins/Runtime/Metasound/Source/MetasoundFrontend/Private/MetasoundFrontendNodeClassRegistryPrivate.h
--- a/Engine/Plugins/Runtime/Metasound/Source/MetasoundFrontend/Private/MetasoundFrontendNodeClassRegistryPrivate.h
+++ b/Engine/Plugins/Runtime/Metasound/Source/MetasoundFrontend/Private/MetasoundFrontendNodeClassRegistryPrivate.h
@@ -117,7 +117,7 @@ namespace Metasound::Frontend
 		virtual TArray<FConverterNodeInfo> GetPossibleConverterNodes(const FName& FromDataType, const FName& ToDataType) override;
 
 		// Register a graph from an IMetaSoundDocumentInterface
-		FGraphRegistryKey RegisterGraph(const TScriptInterface<IMetaSoundDocumentInterface>& InDocumentInterface);
+		FGraphRegistryKey RegisterGraph(const TScriptInterface<IMetaSoundDocumentInterface>& InDocumentInterface, TArrayView<const FGuid> InPageOrder);
 
 		// Unregister a graph
 		bool UnregisterGraph(const FGraphRegistryKey& InRegistryKey, const TScriptInterface<IMetaSoundDocumentInterface>& InDocumentInterface);
@@ -137,7 +137,7 @@ namespace Metasound::Frontend
 			FTopLevelAssetPath AssetPath;
 		};
 
-		void BuildAndRegisterGraphFromDocument(const FMetasoundFrontendDocument& InDocument, const FProxyDataCache& InProxyDataCache, FNodeClassInfo&& InNodeClassInfo, const FTopLevelAssetPath& AssetPath);
+		void BuildAndRegisterGraphFromDocument(const FMetasoundFrontendDocument& InDocument, const FProxyDataCache& InProxyDataCache, FNodeClassInfo&& InNodeClassInfo, const FTopLevelAssetPath& AssetPath, TArrayView<const FGuid> InPageOrder);
 
 		void AddRegistrationTask(const FGraphRegistryKey& InKey, FActiveRegistrationTaskInfo&& TaskInfo);
 		void RemoveRegistrationTask(const FGraphRegistryKey& InKey, FNodeClassRegistryTransaction::ETransactionType TransactionType);
diff --git a/Engine/Plugins/Runtime/Metasound/Source/MetasoundFrontend/Private/MetasoundFrontendPages.cpp b/Engine/Plugins/Runtime/Metasound/Source/MetasoundFrontend/Private/MetasoundFrontendPages.cpp
new file mode 100644
--- /dev/null
+++ b/Engine/Plugins/Runtime/Metasound/Source/MetasoundFrontend/Private/MetasoundFrontendPages.cpp
@@ -0,0 +1,168 @@
+// Copyright Epic Games, Inc. All Rights Reserved.
+
+
+#include "MetasoundFrontendPages.h"
+
+#include "Algo/Transform.h"
+#include "Containers/Array.h"
+#include "Containers/ArrayView.h"
+#include "Misc/Guid.h"
+#include "MetasoundFrontendDocument.h"
+#include "MetasoundFrontendDocumentBuilder.h"
+#include "MetasoundTrace.h"
+
+namespace Metasound::Frontend
+{
+	TArrayView<const FGuid> GetPageOrderForTargetPage(const FGuid& InTargetPage, TArrayView<const FGuid> InPageOrder)
+	{
+		int32 Pos = InPageOrder.Find(InTargetPage);
+		if (Pos == INDEX_NONE)
+		{
+			return {};
+		}
+		else
+		{
+			// Discards and pages in the page order which occur before the target page. 
+			return InPageOrder.RightChop(Pos);
+		}
+	}
+
+#if WITH_EDITORONLY_DATA
+	bool StripUnusedGraphPages(FMetaSoundFrontendDocumentBuilder& InBuilder, TArrayView<const FGuid> InTargetPages, TArrayView<const FGuid> InPageOrder)
+	{
+		bool bModified = false;
+		const FMetasoundFrontendDocument& Document = InBuilder.GetConstDocumentChecked();
+
+		// Find which graphs to remove. 
+		TArray<FGuid> PageIDsToRemove;
+		Document.RootGraph.IterateGraphPages([&PageIDsToRemove](const FMetasoundFrontendGraph& InGraph) { PageIDsToRemove.Add(InGraph.PageID);});
+		
+		for (const FGuid& TargetPage : InTargetPages)
+		{
+			TArrayView<const FGuid> PageOrderForTarget = GetPageOrderForTargetPage(TargetPage, InPageOrder);
+			if (const FMetasoundFrontendGraph* ResolvedGraph = FindPreferredPage(Document.RootGraph.GetConstGraphPages(), PageOrderForTarget))
+			{
+				// Don't remove the specific page ID if it is the one which will
+				// be chosen at runtime.
+				PageIDsToRemove.Remove(ResolvedGraph->PageID);
+			}
+			else
+			{
+				UE_LOG(LogMetaSound, Error, TEXT("Failed to find graph to cook for target page %s on MetaSound %s"), *TargetPage.ToString(), *InBuilder.GetDebugName());
+			}
+		}
+
+		const int32 NumInitGraphs = Document.RootGraph.GetConstGraphPages().Num();
+
+		// Remove unused graphs
+		for (const FGuid& PageIDToRemove : PageIDsToRemove)
+		{
+			bool bDidRemove = InBuilder.RemoveGraphPage(PageIDToRemove);
+			if (bDidRemove)
+			{
+				UE_LOG(LogMetaSound, Display, TEXT("%s: Removed Graph w/PageID '%s'"),
+					*InBuilder.GetDebugName(),
+					*PageIDToRemove.ToString());
+				bModified = true;
+			}
+		}
+
+		// Log results and perform sanity checking
+		const int32 NumRemainingGraphs = Document.RootGraph.GetConstGraphPages().Num();
+
+		checkf(NumRemainingGraphs > 0,
+			TEXT("Document in MetaSound asset '%s' had all default values "
+				"cooked away leaving it in an invalid state. "
+				"Graph must always have at least one implementation."),
+			*InBuilder.GetDebugName());
+
+		if (NumInitGraphs > NumRemainingGraphs)
+		{
+			UE_LOG(LogMetaSound, Display, TEXT("Cook removed %i graph page(s) from '%s'"), NumInitGraphs - NumRemainingGraphs, *InBuilder.GetDebugName());
+		}
+
+		return bModified;
+	}
+
+	bool StripUnusedClassInputPages(FMetaSoundFrontendDocumentBuilder& InBuilder, TArrayView<const FGuid> InTargetPages, TArrayView<const FGuid> InPageOrder)
+	{
+		bool bModified = false;
+		const FMetasoundFrontendDocument& Document = InBuilder.GetConstDocumentChecked();
+		for (const FMetasoundFrontendClassInput& GraphInput : Document.RootGraph.GetDefaultInterface().Inputs)
+		{
+			// Begin with the assumption that all pages will be removed. Then only keep the ones which will actaully
+			// be utilized in a cooked build. 
+			TArray<FGuid> PageIDsToRemove;
+			Algo::Transform(GraphInput.GetDefaults(), PageIDsToRemove, [](const FMetasoundFrontendClassInputDefault& InDefault) { return InDefault.PageID; });
+
+			
+			for (const FGuid& TargetPage : InTargetPages)
+			{
+				TArrayView<const FGuid> PageOrderForTarget = GetPageOrderForTargetPage(TargetPage, InPageOrder);
+				if (const FMetasoundFrontendClassInputDefault* ResolvedInput = FindPreferredPage(GraphInput.GetDefaults(), PageOrderForTarget))
+				{
+					// Don't remove the specific page ID if it is the one which will
+					// be chosen at runtime.
+					PageIDsToRemove.Remove(ResolvedInput->PageID);
+				}
+				else
+				{
+					UE_LOG(LogMetaSound, Error, TEXT("Failed to find class input default to cook for target page %s on MetaSound %s class input %s"), *TargetPage.ToString(), *InBuilder.GetDebugName(), *GraphInput.Name.ToString());
+				}
+			}
+
+			const int32 NumInitDefaults = GraphInput.GetDefaults().Num();
+
+			// Remove unused graphs
+			for (const FGuid& PageIDToRemove : PageIDsToRemove)
+			{
+				constexpr bool bClearInheritsDefault = false;
+				const bool bDidRemove = InBuilder.RemoveGraphInputDefault(GraphInput.Name, PageIDToRemove, bClearInheritsDefault);
+				if (bDidRemove)
+				{
+					UE_LOG(LogMetaSound, Display, TEXT("%s: Removed Graph Input %s w/PageID '%s'"),
+						*InBuilder.GetDebugName(),
+						*GraphInput.Name.ToString(),
+						*PageIDToRemove.ToString());
+					bModified = true;
+				}
+			}
+
+			// Log results and perform sanity checking
+			const int32 NumRemainingDefaults = GraphInput.GetDefaults().Num();
+
+			checkf(NumRemainingDefaults > 0,
+				TEXT("Document in MetaSound asset '%s' had all default input values "
+					"cooked away for input %s leaving it in an invalid state. "
+					"Graph must always have at least one implementation."),
+				*InBuilder.GetDebugName(), 
+				*GraphInput.Name.ToString());
+
+			if (NumInitDefaults > NumRemainingDefaults)
+			{
+				UE_LOG(LogMetaSound, Display, TEXT("Cook removed %i input default page(s) from input %s on '%s'"), NumInitDefaults - NumRemainingDefaults, *GraphInput.Name.ToString(), *InBuilder.GetDebugName());
+			}
+
+		}
+
+		return bModified;
+	}
+
+	bool StripUnusedPages(FMetaSoundFrontendDocumentBuilder& InBuilder, TArrayView<const FGuid> InTargetPages, TArrayView<const FGuid> InPageOrder)
+	{
+		METASOUND_LLM_SCOPE;
+		METASOUND_TRACE_CPUPROFILER_EVENT_SCOPE(Metasound::Frontend::StripUnusedPages);
+
+		if (InTargetPages.Num() < 1)
+		{
+			UE_LOG(LogMetaSound, Error, TEXT("Must have at least one target page to strip unused pages in asset %s"), *InBuilder.GetDebugName());
+			return false; // Not modified
+		}
+
+		bool bModified = StripUnusedGraphPages(InBuilder, InTargetPages, InPageOrder);
+		bModified |= StripUnusedClassInputPages(InBuilder, InTargetPages, InPageOrder);
+
+		return bModified;
+	}
+#endif // WITH_EDITORONLY_DATA
+}
diff --git a/Engine/Plugins/Runtime/Metasound/Source/MetasoundFrontend/Private/MetasoundFrontendProxyDataCache.cpp b/Engine/Plugins/Runtime/Metasound/Source/MetasoundFrontend/Private/MetasoundFrontendProxyDataCache.cpp
--- a/Engine/Plugins/Runtime/Metasound/Source/MetasoundFrontend/Private/MetasoundFrontendProxyDataCache.cpp
+++ b/Engine/Plugins/Runtime/Metasound/Source/MetasoundFrontend/Private/MetasoundFrontendProxyDataCache.cpp
@@ -23,7 +23,7 @@ namespace Metasound
 		namespace MetasoundFrontendProxyDataCachePrivate
 		{
 			template<typename DocElementType>
-			void CreateAndCacheProxies(FProxyDataCache& InCache, const DocElementType& InDocElement)
+			void CreateAndCacheProxies(FProxyDataCache& InCache, const DocElementType& InDocElement, TArrayView<const FGuid> InPageOrder)
 			{
 				checkf(IsInGameThread() || IsInAudioThread(), TEXT("Proxies should only be created in the game thread or audio thread"));
 				IDataTypeRegistry& DataRegistry = IDataTypeRegistry::Get();
@@ -61,26 +61,26 @@ namespace Metasound
 						}
 					}
 				};
-				ForEachLiteral(InDocElement, CreateProxies);
+				ForEachLiteral(InDocElement, CreateProxies, InPageOrder);
 			}
 		}
 
-		void FProxyDataCache::CreateAndCacheProxies(const FMetasoundFrontendDocument& InDocument)
+		void FProxyDataCache::CreateAndCacheProxies(const FMetasoundFrontendDocument& InDocument, TArrayView<const FGuid> InPageOrder)
 		{
 			METASOUND_TRACE_CPUPROFILER_EVENT_SCOPE(Metasound::FProxyDataCache::CreateAndCacheProxies_Document)
-			MetasoundFrontendProxyDataCachePrivate::CreateAndCacheProxies(*this, InDocument);
+			MetasoundFrontendProxyDataCachePrivate::CreateAndCacheProxies(*this, InDocument, InPageOrder);
 		}
 
-		void FProxyDataCache::CreateAndCacheProxies(const FMetasoundFrontendGraphClass& InGraphClass)
+		void FProxyDataCache::CreateAndCacheProxies(const FMetasoundFrontendGraphClass& InGraphClass, TArrayView<const FGuid> InPageOrder)
 		{
 			METASOUND_TRACE_CPUPROFILER_EVENT_SCOPE(Metasound::FProxyDataCache::CreateAndCacheProxies_GraphClass)
-			MetasoundFrontendProxyDataCachePrivate::CreateAndCacheProxies(*this, InGraphClass);
+			MetasoundFrontendProxyDataCachePrivate::CreateAndCacheProxies(*this, InGraphClass, InPageOrder);
 		}
 
-		void FProxyDataCache::CreateAndCacheProxies(const FMetasoundFrontendClass& InClass)
+		void FProxyDataCache::CreateAndCacheProxies(const FMetasoundFrontendClass& InClass, TArrayView<const FGuid> InPageOrder)
 		{
 			METASOUND_TRACE_CPUPROFILER_EVENT_SCOPE(Metasound::FProxyDataCache::CreateAndCacheProxies_Class)
-			MetasoundFrontendProxyDataCachePrivate::CreateAndCacheProxies(*this, InClass);
+			MetasoundFrontendProxyDataCachePrivate::CreateAndCacheProxies(*this, InClass, InPageOrder);
 		}
 
 		bool FProxyDataCache::Contains(const UObject* InUObject) const
diff --git a/Engine/Plugins/Runtime/Metasound/Source/MetasoundFrontend/Private/MetasoundFrontendProxyDataCache.h b/Engine/Plugins/Runtime/Metasound/Source/MetasoundFrontend/Private/MetasoundFrontendProxyDataCache.h
--- a/Engine/Plugins/Runtime/Metasound/Source/MetasoundFrontend/Private/MetasoundFrontendProxyDataCache.h
+++ b/Engine/Plugins/Runtime/Metasound/Source/MetasoundFrontend/Private/MetasoundFrontendProxyDataCache.h
@@ -32,13 +32,13 @@ namespace Metasound
 		{
 		public:
 			/** Create and cache all proxies found in a document. */
-			void CreateAndCacheProxies(const FMetasoundFrontendDocument& InDocument);
+			void CreateAndCacheProxies(const FMetasoundFrontendDocument& InDocument, TArrayView<const FGuid> InPageOrder);
 
 			/** Create and cache all proxies found in a graph class. */
-			void CreateAndCacheProxies(const FMetasoundFrontendGraphClass& InGraphClass);
+			void CreateAndCacheProxies(const FMetasoundFrontendGraphClass& InGraphClass, TArrayView<const FGuid> InPageOrder);
 
 			/** Create and cache all proxies found in a frontend class class. */
-			void CreateAndCacheProxies(const FMetasoundFrontendClass& InClass);
+			void CreateAndCacheProxies(const FMetasoundFrontendClass& InClass, TArrayView<const FGuid> InPageOrder);
 
 			/** Cache a specific object associated with a UObject pointer.
 			 *
diff --git a/Engine/Plugins/Runtime/Metasound/Source/MetasoundFrontend/Public/MetasoundAssetBase.h b/Engine/Plugins/Runtime/Metasound/Source/MetasoundFrontend/Public/MetasoundAssetBase.h
--- a/Engine/Plugins/Runtime/Metasound/Source/MetasoundFrontend/Public/MetasoundAssetBase.h
+++ b/Engine/Plugins/Runtime/Metasound/Source/MetasoundFrontend/Public/MetasoundAssetBase.h
@@ -85,27 +85,27 @@ public:
 	virtual void SetRegistryAssetClassInfo(const Metasound::Frontend::FNodeClassInfo& InClassInfo) { }
 #endif // WITH_EDITORONLY_DATA
 
-	UE_DEPRECATED(5.5, "Moved to IMetaSoundDocumentInterface::ConformObjectToDocument")
-	UE_API virtual bool ConformObjectDataToInterfaces();
+	UE_DEPRECATED(5.7, "Use method which requires a FMetaSoundAssetRegistrationOptions")
+	UE_API void UpdateAndRegisterForExecution();
 
 	// Registers the root graph of the given asset with the MetaSound Frontend. Unlike 'UpdateAndRegisterForSerialization", this call
 	// generates all necessary runtime data to execute the given graph (i.e. INodes).
-	UE_API virtual void UpdateAndRegisterForExecution(Metasound::Frontend::FMetaSoundAssetRegistrationOptions InRegistrationOptions = Metasound::Frontend::FMetaSoundAssetRegistrationOptions());
-
-	UE_DEPRECATED(5.5, "Moved to UpdateAndRegisterForExecution.")
-	UE_API virtual void RegisterGraphWithFrontend(Metasound::Frontend::FMetaSoundAssetRegistrationOptions InRegistrationOptions = Metasound::Frontend::FMetaSoundAssetRegistrationOptions());
+	UE_API virtual void UpdateAndRegisterForExecution(const Metasound::Frontend::FMetaSoundAssetRegistrationOptions& InRegistrationOptions);
 
 	// Unregisters the root graph of the given asset with the MetaSound Frontend.
 	UE_API void UnregisterGraphWithFrontend();
 
-	UE_DEPRECATED(5.5, "Moved to UpdateAndRegisterForSerialization instead, which is only in builds set to load editor-only data.")
-	UE_API void CookMetaSound();
-
 #if WITH_EDITORONLY_DATA
 	// Updates and registers this and referenced MetaSound document objects with the NodeClass Registry. AutoUpdates and
 	// optimizes aforementioned documents for serialization. Unlike 'UpdateAndRegisterForRuntime', does not generate required
-	// runtime data for graph execution. If CookPlatformName is set, used to strip data not required for the provided platform.
+	// runtime data for graph execution. 
+	UE_DEPRECATED(5.7, "Use method which passes in FMetaSoundAssetCookOptions")
 	UE_API void UpdateAndRegisterForSerialization(FName CookPlatformName = { });
+
+	// Updates and registers this and referenced MetaSound document objects with the NodeClass Registry. AutoUpdates and
+	// optimizes aforementioned documents for serialization. Unlike 'UpdateAndRegisterForRuntime', does not generate required
+	// runtime data for graph execution. 
+	UE_API void UpdateAndRegisterForSerialization(const Metasound::Frontend::FMetaSoundAssetCookOptions& InCookOptions);
 #endif // WITH_EDITORONLY_DATA
 
 #if WITH_EDITOR
@@ -131,9 +131,6 @@ public:
 
 	UE_API bool AddingReferenceCausesLoop(const FMetasoundAssetBase& InMetaSound) const;
 
-	UE_DEPRECATED(5.5, "Use overload that is provided an AssetBase")
-	UE_API bool AddingReferenceCausesLoop(const FSoftObjectPath& InReferencePath) const;
-
 	UE_API bool IsReferencedAsset(const FMetasoundAssetBase& InAssetToCheck) const;
 
 	UE_API bool IsRegistered() const;
@@ -152,18 +149,12 @@ public:
 	UE_API Metasound::Frontend::FGraphHandle GetRootGraphHandle();
 	UE_API Metasound::Frontend::FConstGraphHandle GetRootGraphHandle() const;
 
-	UE_DEPRECATED(5.5, "Direct mutation of the document is no longer supported via AssetBase.")
-	UE_API void SetDocument(FMetasoundFrontendDocument InDocument, bool bMarkDirty = true);
-
 	UE_API virtual const FMetasoundFrontendDocument& GetConstDocumentChecked() const;
 
 	// Soft deprecated.  Document layer should not be directly mutated via asset base in anticipation
 	// of moving all mutable document calls to the Frontend/Subsystem Document Builder API.
 	UE_API FMetasoundFrontendDocument& GetDocumentChecked();
 
-	UE_DEPRECATED(5.5, "Use GetConstDocument from casting Owning Asset to IMetaSoundDocumentInterface (See 'GetOwningAsset') instead.")
-	UE_API const FMetasoundFrontendDocument& GetDocumentChecked() const;
-
 	UE_API const Metasound::Frontend::FGraphRegistryKey& GetGraphRegistryKey() const;
 
 #if WITH_EDITORONLY_DATA
@@ -176,9 +167,6 @@ public:
 
 	UE_API FMetasoundFrontendDocumentModifyContext& GetModifyContext();
 
-	UE_DEPRECATED(5.5, "Use GetConstModifyContext")
-	UE_API const FMetasoundFrontendDocumentModifyContext& GetModifyContext() const;
-
 	UE_API const FMetasoundFrontendDocumentModifyContext& GetConstModifyContext() const;
 #endif // WITH_EDITOR
 
@@ -205,9 +193,6 @@ public:
 	UE_API void SetVersionedOnLoad();
 #endif // WITH_EDITORONLY_DATA
 
-	UE_DEPRECATED(5.5, "Use IMetaSoundDocumentInterface 'IsActivelyBuilding' instead")
-	virtual bool IsBuilderActive() const { checkNoEntry(); return false; }
-
 protected:
 	UE_API void OnNotifyBeginDestroy();
 
@@ -220,10 +205,6 @@ PRAGMA_ENABLE_DEPRECATION_WARNINGS
 	virtual void SetReferencedAssets(TSet<Metasound::Frontend::IMetaSoundAssetManager::FAssetRef>&& InAssetRefs) = 0;
 #endif
 
-	// Get information for communicating asynchronously with MetaSound running instance.
-	UE_DEPRECATED(5.3, "MetaSounds no longer communicate using FSendInfo.")
-	UE_API TArray<FSendInfoAndVertexName> GetSendInfos(uint64 InInstanceID) const;
-
 #if WITH_EDITORONLY_DATA
 	UE_API FText GetDisplayName(FString&& InTypeName) const;
 #endif // WITH_EDITORONLY_DATA
@@ -236,18 +217,15 @@ PRAGMA_ENABLE_DEPRECATION_WARNINGS
 	// Returns an access pointer to the document.
 	virtual Metasound::Frontend::FConstDocumentAccessPtr GetDocumentConstAccessPtr() const = 0;
 
-	UE_DEPRECATED(5.5, "AutoUpdate implementation now private and implemented within 'Version Dependencies'")
-	UE_API bool AutoUpdate(bool bInLogWarningsOnDroppedConnection);
-
-	UE_DEPRECATED(5.5, "Moved to private, non-cook specific implementation")
-	UE_API void CookReferencedMetaSounds();
-
 	// Ensures all referenced graph classes are registered (or re-registers depending on options).
 	UE_API void RegisterAssetDependencies(const Metasound::Frontend::FMetaSoundAssetRegistrationOptions & InRegistrationOptions);
 
 private:
 #if WITH_EDITORONLY_DATA
+	UE_DEPRECATED(5.7, "Call version of function which accepts FMetaSoundAssetCookOptions.")
 	UE_API void UpdateAndRegisterReferencesForSerialization(FName CookPlatformName);
+
+	UE_API void UpdateAndRegisterReferencesForSerialization(const Metasound::Frontend::FMetaSoundAssetCookOptions& InCookOptions);
 #endif // WITH_EDITORONLY_DATA
 
 	// Checks if version is up-to-date. If so, returns true. If false, updates the interfaces within the given asset's document to the most recent version.
diff --git a/Engine/Plugins/Runtime/Metasound/Source/MetasoundFrontend/Public/MetasoundAssetManager.h b/Engine/Plugins/Runtime/Metasound/Source/MetasoundFrontend/Public/MetasoundAssetManager.h
--- a/Engine/Plugins/Runtime/Metasound/Source/MetasoundFrontend/Public/MetasoundAssetManager.h
+++ b/Engine/Plugins/Runtime/Metasound/Source/MetasoundFrontend/Public/MetasoundAssetManager.h
@@ -119,6 +119,21 @@ namespace Metasound::Frontend
 		// No longer used. Memory management of document (i.e. copying or using object's version) inferred internally
 		bool bRegisterCopyIfAsync = false;
 #endif // WITH_EDITOR
+
+		// Contains a ranking of the preferred page to use with the most preferred page ID being first. 
+		TArrayView<const FGuid> PageOrder;
+	};
+
+	struct FMetaSoundAssetCookOptions
+	{
+		// If true, any unused page will get removed during cook.
+		bool bStripUnusedPages = false; 
+
+		// Which PageIDs are going to available for targeting after cook. 
+		TArray<FGuid> PagesToTarget; 
+
+		// The PageOrder to use when cooking.
+		TArray<FGuid> PageOrder;
 	};
 
 
diff --git a/Engine/Plugins/Runtime/Metasound/Source/MetasoundFrontend/Public/MetasoundDocumentInterface.h b/Engine/Plugins/Runtime/Metasound/Source/MetasoundFrontend/Public/MetasoundDocumentInterface.h
--- a/Engine/Plugins/Runtime/Metasound/Source/MetasoundFrontend/Public/MetasoundDocumentInterface.h
+++ b/Engine/Plugins/Runtime/Metasound/Source/MetasoundFrontend/Public/MetasoundDocumentInterface.h
@@ -85,6 +85,7 @@ namespace Metasound::Frontend
 		// platform, allowing users to optimize away data and scale the amount of memory required for
 		// initial load of input UObjects and graph topology, which can also positively effect runtime
 		// performance as well, etc. Returns true if builder modified the document, false if not.
+		UE_DEPRECATED(5.7, "Use StripUnusedPages(...) instead.")
 		virtual bool CookPages(FName CookPlatformName, FMetaSoundFrontendDocumentBuilder& Builder) const = 0;
 #endif // WITH_EDITORONLY_DATA
 
@@ -115,16 +116,30 @@ namespace Metasound::Frontend
 
 		// Given the provided document and its respective pages, provides the PageID to be used for runtime IGraph and proxy generation.
 		// Returns true if valid page was found, false if not.
+		UE_DEPRECATED(5.7, "This is no longer supported. FindPreferredPage(...)")
 		virtual FGuid ResolveTargetPageID(const FMetasoundFrontendGraphClass& InGraphClass) const = 0;
 
 		// Given the provided input and its respective paged default values, provides the default PageID to be used for runtime IGraph and proxy generation.
 		// Returns true if valid page was found, false if not.
+		UE_DEPRECATED(5.7, "This is no longer supported. FindPreferredPage(...)")
 		virtual FGuid ResolveTargetPageID(const FMetasoundFrontendClassInput& InClassInput) const = 0;
 
 		// Given the provided array of default values, provides the default PageID to be used for runtime IGraph and proxy generation.
 		// Returns true if valid page was found, false if not.
+		UE_DEPRECATED(5.7, "This is no longer supported. FindPreferredPage(...)")
 		virtual FGuid ResolveTargetPageID(const TArray<FMetasoundFrontendClassInputDefault>& Defaults) const = 0;
 
+		UE_DEPRECATED(5.7, "This is only used to support existing deprecated functionality")
+		virtual TArrayView<const FGuid> GetPageOrder() const = 0;
+
+#if WITH_EDITORONLY_DATA
+		UE_DEPRECATED(5.7, "This is only used to support existing deprecated functionality")
+		virtual TArray<FGuid> GetCookedTargetPages(FName InPlatformName) const = 0;
+
+		UE_DEPRECATED(5.7, "This is only used to support existing deprecated functionality")
+		virtual TArray<FGuid> GetCookedPageOrder(FName InPlatformName) const = 0;
+#endif // WITH_EDITORONLY_DATA
+
 		static UE_API IDocumentBuilderRegistry* Get();
 		static UE_API IDocumentBuilderRegistry& GetChecked();
 		static UE_API void Deinitialize();
diff --git a/Engine/Plugins/Runtime/Metasound/Source/MetasoundFrontend/Public/MetasoundFrontendController.h b/Engine/Plugins/Runtime/Metasound/Source/MetasoundFrontend/Public/MetasoundFrontendController.h
--- a/Engine/Plugins/Runtime/Metasound/Source/MetasoundFrontend/Public/MetasoundFrontendController.h
+++ b/Engine/Plugins/Runtime/Metasound/Source/MetasoundFrontend/Public/MetasoundFrontendController.h
@@ -843,9 +843,19 @@ namespace Metasound
 			virtual FGuid GetVertexIDForInputVertex(const FVertexName& InInputName) const = 0;
 			virtual FGuid GetVertexIDForOutputVertex(const FVertexName& InOutputName) const = 0;
 
+			UE_DEPRECATED(5.7, "Use version which provides page order")
 			virtual FMetasoundFrontendLiteral GetDefaultInput(const FGuid& InVertexID) const = 0;
+
+			virtual FMetasoundFrontendLiteral GetDefaultInput(const FGuid& InVertexID, TArrayView<const FGuid> InPageOrder) const = 0;
+
+			// These can be used to set the default value for a given input on this graph.
+			// @returns false if the input name couldn't be found, or if the literal type was incompatible with the Data Type of this input.
+
+			UE_DEPRECATED(5.7, "Use version which provides page order")
 			virtual bool SetDefaultInput(const FGuid& InVertexID, const FMetasoundFrontendLiteral& InLiteral) = 0;
 
+			virtual bool SetDefaultInput(const FGuid& InVertexID, const FMetasoundFrontendLiteral& InLiteral, TArrayView<const FGuid> InPageOrder) = 0;
+
 			/** Set the default value for the graph input.
 			 *
 			 * @param InInputName - Name of the graph input.
@@ -964,9 +974,6 @@ namespace Metasound
 			 */
 			virtual FNodeHandle CreateEmptySubgraph(const FMetasoundFrontendClassMetadata& InInfo) = 0;
 
-			UE_DEPRECATED(5.5, "Building graph operator is implemented internally and now supports pages")
-			virtual TUniquePtr<IOperator> BuildOperator(const FOperatorSettings& InSettings, const FMetasoundEnvironment& InEnvironment, FBuildResults& OutResults) const { return { }; }
-
 			/** Returns a handle to the document owning this graph. */
 			virtual FDocumentHandle GetOwningDocument() = 0;
 
diff --git a/Engine/Plugins/Runtime/Metasound/Source/MetasoundFrontend/Public/MetasoundFrontendDocument.h b/Engine/Plugins/Runtime/Metasound/Source/MetasoundFrontend/Public/MetasoundFrontendDocument.h
--- a/Engine/Plugins/Runtime/Metasound/Source/MetasoundFrontend/Public/MetasoundFrontendDocument.h
+++ b/Engine/Plugins/Runtime/Metasound/Source/MetasoundFrontend/Public/MetasoundFrontendDocument.h
@@ -2080,25 +2080,38 @@ namespace Metasound::Frontend
 	using FForEachLiteralFunctionRef = TFunctionRef<void(const FName& InDataTypeName, const FMetasoundFrontendLiteral&)>;
 
 	/** Execute the provided function for each literal on a FMetasoundFrontendDocument.*/
-	METASOUNDFRONTEND_API void ForEachLiteral(const FMetasoundFrontendDocument& InDoc, FForEachLiteralFunctionRef OnLiteral);
-
-	/** Execute the provided function for each literal on a FMetasoundFrontendDocument filtered by the given PageID.*/
-	METASOUNDFRONTEND_API void ForEachLiteral(const FMetasoundFrontendDocument& InDoc, FForEachLiteralFunctionRef OnLiteral, const FGuid& InPageID);
+	UE_DEPRECATED(5.7, "Use version which requires a PageOrder")
+	UE_API void ForEachLiteral(const FMetasoundFrontendDocument& InDoc, FForEachLiteralFunctionRef OnLiteral);
 
 	/** Execute the provided function for each literal on a FMetasoundFrontendGraphClass.*/
-	METASOUNDFRONTEND_API void ForEachLiteral(const FMetasoundFrontendGraphClass& InGraphClass, FForEachLiteralFunctionRef OnLiteral);
-
-	/** Execute the provided function for each literal on a FMetasoundFrontendGraphClass filtered by the given PageID.*/
-	METASOUNDFRONTEND_API void ForEachLiteral(const FMetasoundFrontendGraphClass& InGraphClass, FForEachLiteralFunctionRef OnLiteral, const FGuid& InPageID);
+	UE_DEPRECATED(5.7, "Use version which requires a PageOrder")
+	UE_API void ForEachLiteral(const FMetasoundFrontendGraphClass& InGraphClass, FForEachLiteralFunctionRef OnLiteral);
 
 	/** Execute the provided function for each literal on a FMetasoundFrontendClass.*/
-	METASOUNDFRONTEND_API void ForEachLiteral(const FMetasoundFrontendClass& InClass, FForEachLiteralFunctionRef OnLiteral);
+	UE_DEPRECATED(5.7, "Use version which requires a PageOrder")
+	UE_API void ForEachLiteral(const FMetasoundFrontendClass& InClass, FForEachLiteralFunctionRef OnLiteral);
+
+	/** Execute the provided function for each literal on a FMetasoundFrontendNode.*/
+	UE_DEPRECATED(5.7, "Use version which requires a PageOrder")
+	UE_API void ForEachLiteral(const FMetasoundFrontendNode& InNode, FForEachLiteralFunctionRef OnLiteral);
+
+	UE_DEPRECATED(5.7, "Use version which requires a PageOrder")
+	UE_API void ForEachLiteral(const FMetasoundFrontendClassInterface& InClassInterface, FForEachLiteralFunctionRef OnLiteral);
 
-	METASOUNDFRONTEND_API void ForEachLiteral(const FMetasoundFrontendClassInterface& InClassInterface, FForEachLiteralFunctionRef OnLiteral);
+	/** Execute the provided function for each literal on a FMetasoundFrontendDocument.*/
+	UE_API void ForEachLiteral(const FMetasoundFrontendDocument& InDoc, FForEachLiteralFunctionRef OnLiteral, TArrayView<const FGuid> InPageOrder);
+	
+	/** Execute the provided function for each literal on a FMetasoundFrontendGraphClass.*/
+	UE_API void ForEachLiteral(const FMetasoundFrontendGraphClass& InGraphClass, FForEachLiteralFunctionRef OnLiteral, TArrayView<const FGuid> InPageOrder);
+	
+	/** Execute the provided function for each literal on a FMetasoundFrontendClass.*/
+	UE_API void ForEachLiteral(const FMetasoundFrontendClass& InClass, FForEachLiteralFunctionRef OnLiteral, TArrayView<const FGuid> InPageOrder);
 
+	/** Execute the provided function for each literal on a FMetasoundFrontendClassInterface.*/
+	UE_API void ForEachLiteral(const FMetasoundFrontendClassInterface& InClassInterface, FForEachLiteralFunctionRef OnLiteral, TArrayView<const FGuid> InPageOrder);
 
 	/** Execute the provided function for each literal on a FMetasoundFrontendNode.*/
-	METASOUNDFRONTEND_API void ForEachLiteral(const FMetasoundFrontendNode& InNode, FForEachLiteralFunctionRef OnLiteral);
+	UE_API void ForEachLiteral(const FMetasoundFrontendNode& InNode, FForEachLiteralFunctionRef OnLiteral, TArrayView<const FGuid> InPageOrder);
 
 #if WITH_EDITORONLY_DATA
 	// Provides list of interface members that have been added or removed
diff --git a/Engine/Plugins/Runtime/Metasound/Source/MetasoundFrontend/Public/MetasoundFrontendGraphBuilder.h b/Engine/Plugins/Runtime/Metasound/Source/MetasoundFrontend/Public/MetasoundFrontendGraphBuilder.h
--- a/Engine/Plugins/Runtime/Metasound/Source/MetasoundFrontend/Public/MetasoundFrontendGraphBuilder.h
+++ b/Engine/Plugins/Runtime/Metasound/Source/MetasoundFrontend/Public/MetasoundFrontendGraphBuilder.h
@@ -61,6 +61,10 @@ namespace Metasound::Frontend
 			/** Optional subgraph map. Subgraphs are not fully supported in MetaSound
 			 * so this is kept optional to avoid creation of unused TMap<>s. */
 			const TMap<FGuid, TSharedPtr<const IGraph>>* Subgraphs = nullptr;
+
+			/** An array of preferred pages where most preferred pages occur before
+			 * less preferred. */
+			TArrayView<const FGuid> PageOrder = {};
 		};
 
 		/** Create a INode from a document representation of the node. 
@@ -75,20 +79,15 @@ namespace Metasound::Frontend
 		 *
 		 * @return True if all dependencies are C++ classes. False otherwise.
 		 */
-		UE_DEPRECATED(5.5, "No longer supported (not useful nor maintainable with paged graph support")
-		static UE_API bool IsFlat(const FMetasoundFrontendDocument& InDocument);
-
-		UE_DEPRECATED(5.5, "No longer supported (not useful nor maintainable with paged graph support")
-		static UE_API bool IsFlat(const FMetasoundFrontendGraphClass& InRoot, const TArray<FMetasoundFrontendClass>& InDependencies);
 
 		/* Create a FFrontendGraph from a FMetasoundFrontendDocument.*/
-		static UE_API TUniquePtr<FFrontendGraph> CreateGraph(const FMetasoundFrontendDocument& InDocument, const FString& InDebugAssetName);
+		static UE_API TUniquePtr<FFrontendGraph> CreateGraph(const FMetasoundFrontendDocument& InDocument, const FString& InDebugAssetName, TArrayView<const FGuid> InPageOrder = {});
 
 		/* Create a FFrontendGraph from a FMetasoundFrontendDocument retrieving proxies from a FProxyDataCache.*/
-		static UE_API TUniquePtr<FFrontendGraph> CreateGraph(const FMetasoundFrontendDocument& InDocument, const Frontend::FProxyDataCache& InProxies, const FString& InDebugAssetName, const FGuid InGraphId = Frontend::CreateLocallyUniqueId());
+		static UE_API TUniquePtr<FFrontendGraph> CreateGraph(const FMetasoundFrontendDocument& InDocument, const Frontend::FProxyDataCache& InProxies, const FString& InDebugAssetName, const FGuid InGraphId = Frontend::CreateLocallyUniqueId(), TArrayView<const FGuid> InPageOrder = {});
 
 		/* Create a FFrontendGraph from a FMetasoundFrontendDocument and subobjects.*/
-		static UE_API TUniquePtr<FFrontendGraph> CreateGraph(const FMetasoundFrontendGraphClass& InGraph, const TArray<FMetasoundFrontendGraphClass>& InSubgraphs, const TArray<FMetasoundFrontendClass>& InDependencies, const FString& InDebugAssetName);
+		static UE_API TUniquePtr<FFrontendGraph> CreateGraph(const FMetasoundFrontendGraphClass& InGraph, const TArray<FMetasoundFrontendGraphClass>& InSubgraphs, const TArray<FMetasoundFrontendClass>& InDependencies, const FString& InDebugAssetName, TArrayView<const FGuid> InPageOrder = {});
 
 		/* Create a FFrontendGraph from a FMetasoundFrontendDocument and subobjects, retrieving proxies from a FProxyDataCache.*/
 		static UE_API TUniquePtr<FFrontendGraph> CreateGraph(
@@ -97,7 +96,8 @@ namespace Metasound::Frontend
 			const TArray<FMetasoundFrontendClass>& InDependencies,
 			const Frontend::FProxyDataCache& InProxyDataCache,
 			const FString& InDebugAssetName,
-			const FGuid InGraphId = Frontend::CreateLocallyUniqueId());
+			const FGuid InGraphId = Frontend::CreateLocallyUniqueId(), 
+			TArrayView<const FGuid> InPageOrder = {});
 
 	};
 }
diff --git a/Engine/Plugins/Runtime/Metasound/Source/MetasoundFrontend/Public/MetasoundFrontendPages.h b/Engine/Plugins/Runtime/Metasound/Source/MetasoundFrontend/Public/MetasoundFrontendPages.h
new file mode 100644
--- /dev/null
+++ b/Engine/Plugins/Runtime/Metasound/Source/MetasoundFrontend/Public/MetasoundFrontendPages.h
@@ -0,0 +1,71 @@
+// Copyright Epic Games, Inc. All Rights Reserved.
+ 
+#pragma once
+
+#include "Containers/Array.h"
+#include "Containers/ArrayView.h"
+
+#define UE_API METASOUNDFRONTEND_API
+
+// Forward declare
+struct FGuid;
+struct FMetaSoundFrontendDocumentBuilder;
+
+namespace Metasound::Frontend
+{
+#if WITH_EDITORONLY_DATA
+	/** Remove page data which will not be use given the supplied target pages and page order.
+	 *
+	 * @return true if the builder removed any pages. 
+	 */
+	UE_API bool StripUnusedPages(FMetaSoundFrontendDocumentBuilder& InBuilder, TArrayView<const FGuid> InTargetPages, TArrayView<const FGuid> InPageOrder);
+#endif // WITH_EDITORONLY_DATA
+
+	/** Return a pointer to the most preferred element for a given page order. 
+	 *
+	 * PageIDs with lower indices in InPageOrder are preferred over those with
+	 * higher indices. If none of the elements have a PageID which exists in the 
+	 * InPageOrder, then a pointer to the first element is returned. If the InElements array is
+	 * empty, then nullptr is returned. 
+	 */
+	template<typename ElementType>
+	const ElementType* FindPreferredPage(const TArray<ElementType>& InElements, TArrayView<const FGuid> InPageOrder)
+	{
+		const ElementType* Target = nullptr;
+		// Handle common case of only a single page existing
+		if (InElements.Num() == 1)
+		{
+			Target = &InElements[0];
+		}
+		else if (InElements.Num() > 0)
+		{
+			// Handle multiple pages. InPageOrder represents the sorted page resolution
+			// in order of preference. We find the most preferred page.
+			int32 BestPos = InPageOrder.Num();
+			for (const ElementType& Element : InElements)
+			{
+				int32 Pos = InPageOrder.Find(Element.PageID);
+				if ((Pos != INDEX_NONE) && (Pos < BestPos))
+				{
+					BestPos = Pos;
+					Target = &Element;
+					if (Pos == 0)
+					{
+						// The best page has been found. 
+						break;
+					}
+				}
+			}
+
+			// Fallback to any available page if none were found
+			if (nullptr == Target)
+			{
+				Target = &InElements[0];
+			}
+		}
+
+		return Target;
+	}
+}
+
+#undef UE_API 
