Fix clang compilation errors about shadowed variables

diff --git a/Engine/Plugins/Runtime/Metasound/Source/MetasoundEngine/Private/MetasoundSettings.cpp b/Engine/Plugins/Runtime/Metasound/Source/MetasoundEngine/Private/MetasoundSettings.cpp
--- a/Engine/Plugins/Runtime/Metasound/Source/MetasoundEngine/Private/MetasoundSettings.cpp
+++ b/Engine/Plugins/Runtime/Metasound/Source/MetasoundEngine/Private/MetasoundSettings.cpp
@@ -357,11 +357,11 @@ void UMetaSoundSettings::CacheCookedPageIDs(FName PlatformName) const
 		 * any of the cook target page IDs. */
 		constexpr bool bReverse = true;
 		bool bFoundAnyTargetPage = false;
-		IteratePageSettings([&](const FMetaSoundPageSettings& PageSettings) 
+		IteratePageSettings([&](const FMetaSoundPageSettings& InPageSettings) 
 			{
 #if WITH_EDITOR
 				// Skip this page if it explicitly excluded from cooking.
-				if (PageSettings.GetExcludeFromCook(PlatformName))
+				if (InPageSettings.GetExcludeFromCook(PlatformName))
 				{
 					return;
 				}
@@ -371,12 +371,12 @@ void UMetaSoundSettings::CacheCookedPageIDs(FName PlatformName) const
 				if (!bFoundAnyTargetPage)
 				{
 					// Sets flag to true if any of the CookPlatformTargetPageIDs equal PageSettings.PageID
-					bFoundAnyTargetPage = Algo::AnyOf(CookPlatformTargetPageIDs, [&PageSettings](const FGuid& InTargetPageID) { return InTargetPageID == PageSettings.UniqueId; });
+					bFoundAnyTargetPage = Algo::AnyOf(CookPlatformTargetPageIDs, [&InPageSettings](const FGuid& InTargetPageID) { return InTargetPageID == InPageSettings.UniqueId; });
 				}
 
 				if (bFoundAnyTargetPage)
 				{
-					CookPlatformPageOrder.Add(PageSettings.UniqueId);
+					CookPlatformPageOrder.Add(InPageSettings.UniqueId);
 				}
 
 			}, bReverse);
@@ -532,16 +532,16 @@ const FMetaSoundPageSettings* UMetaSoundSettings::GetTargetPageSettingsFromOverr
 		bool bFoundTargetPage = false;
 		bool bFoundPageSupportedByPlatform = false;
 
-		IteratePageSettings([&](const FMetaSoundPageSettings& PageSettings)
+		IteratePageSettings([&](const FMetaSoundPageSettings& InPageSettings)
 		{
-			bFoundTargetPage |= PageSettings.Name == InTargetPage;
+			bFoundTargetPage |= InPageSettings.Name == InTargetPage;
 			if (bFoundTargetPage && !bFoundPageSupportedByPlatform)
 			{
-				const bool bIsCooked = !PageSettings.GetExcludeFromCook(PageOverrideData->PlatformName);
+				const bool bIsCooked = !InPageSettings.GetExcludeFromCook(PageOverrideData->PlatformName);
 				if (bIsCooked)
 				{
 					bFoundPageSupportedByPlatform = true;
-					TargetSettings = &PageSettings;
+					TargetSettings = &InPageSettings;
 				}
 			}
 		}, bReverse);
