CIS Fix shadow variable in MetaSoundSettings.cpp
#rnx

diff --git a/Engine/Plugins/Runtime/Metasound/Source/MetasoundEngine/Private/MetasoundSettings.cpp b/Engine/Plugins/Runtime/Metasound/Source/MetasoundEngine/Private/MetasoundSettings.cpp
--- a/Engine/Plugins/Runtime/Metasound/Source/MetasoundEngine/Private/MetasoundSettings.cpp
+++ b/Engine/Plugins/Runtime/Metasound/Source/MetasoundEngine/Private/MetasoundSettings.cpp
@@ -273,9 +273,9 @@ const FMetaSoundPageSettings& UMetaSoundSettings::GetDefaultPageSettings() const
 TArray<FName> UMetaSoundSettings::GetAllPlatformNamesImplementingTargets() const
 {
 	TSet<FName> PlatformNames;
-	IteratePageSettings([&](const FMetaSoundPageSettings& PageSetting)
+	IteratePageSettings([&](const FMetaSoundPageSettings& InPageSetting)
 	{
-		TArray<FName> PagePlatforms = PageSetting.GetTargetPlatforms();
+		TArray<FName> PagePlatforms = InPageSetting.GetTargetPlatforms();
 		PlatformNames.Append(MoveTemp(PagePlatforms));
 	});
 	return PlatformNames.Array();
@@ -654,9 +654,9 @@ TArray<FName> UMetaSoundSettings::GetPageNames()
 	if (const UMetaSoundSettings* Settings = GetDefault<UMetaSoundSettings>())
 	{
 		TArray<FName> Names;
-		Settings->IteratePageSettings([&Names](const FMetaSoundPageSettings& PageSetting)
+		Settings->IteratePageSettings([&Names](const FMetaSoundPageSettings& InPageSetting)
 		{
-			Names.Add(PageSetting.Name);
+			Names.Add(InPageSetting.Name);
 		});
 		return Names;
 	}
