Fix editor crash with MetaSoundBuilder Undo/Redo action by not managing the versioning of a transient asset during undo transactions.

diff --git a/Engine/Plugins/Runtime/Metasound/Source/MetasoundEngine/Private/MetasoundEngineAsset.h b/Engine/Plugins/Runtime/Metasound/Source/MetasoundEngine/Private/MetasoundEngineAsset.h
--- a/Engine/Plugins/Runtime/Metasound/Source/MetasoundEngine/Private/MetasoundEngineAsset.h
+++ b/Engine/Plugins/Runtime/Metasound/Source/MetasoundEngine/Private/MetasoundEngineAsset.h
@@ -303,7 +303,7 @@ namespace Metasound::Engine
 
 			bool bVersionedAsset = false;
 
-			if (InArchive.IsLoading())
+			if (InArchive.IsLoading() && !InArchive.IsTransacting())
 			{
 				const bool bIsTransacting = InArchive.IsTransacting();
 				TStrongObjectPtr<UMetaSoundBuilderBase> Builder;
