Updated headers using UnrealCodeFixup to make sure dllstorage is on methods/staticvars instead of types. Part 4/n

diff --git a/Engine/Plugins/Runtime/Metasound/Source/MetasoundEditor/Public/MetasoundDefaultLiteralCustomization.h b/Engine/Plugins/Runtime/Metasound/Source/MetasoundEditor/Public/MetasoundDefaultLiteralCustomization.h
--- a/Engine/Plugins/Runtime/Metasound/Source/MetasoundEditor/Public/MetasoundDefaultLiteralCustomization.h
+++ b/Engine/Plugins/Runtime/Metasound/Source/MetasoundEditor/Public/MetasoundDefaultLiteralCustomization.h
@@ -10,6 +10,8 @@
 #include "PropertyHandle.h"
 #include "Templates/Function.h"
 
+#define UE_API METASOUNDEDITOR_API
+
 // Forward Declarations
 class IDetailPropertyRow;
 class IPropertyHandle;
@@ -22,7 +24,7 @@ class UMetasoundEditorGraphNode;
 
 namespace Metasound::Editor
 {
-	class METASOUNDEDITOR_API FMetasoundDefaultLiteralCustomizationBase
+	class FMetasoundDefaultLiteralCustomizationBase
 	{
 	protected:
 		IDetailCategoryBuilder* DefaultCategoryBuilder = nullptr;
@@ -30,31 +32,31 @@ namespace Metasound::Editor
 	public:
 		using FOnDefaultPageRowAdded = TFunction<void(IDetailPropertyRow& /* Value Row */, TSharedRef<IPropertyHandle> /* Page Property Handle */)>;
 
-		FMetasoundDefaultLiteralCustomizationBase(IDetailCategoryBuilder& InDefaultCategoryBuilder);
-		virtual ~FMetasoundDefaultLiteralCustomizationBase();
+		UE_API FMetasoundDefaultLiteralCustomizationBase(IDetailCategoryBuilder& InDefaultCategoryBuilder);
+		UE_API virtual ~FMetasoundDefaultLiteralCustomizationBase();
 
-		virtual void CustomizeDefaults(UMetasoundEditorGraphMemberDefaultLiteral& InLiteral, IDetailLayoutBuilder& InDetailLayout);
+		UE_API virtual void CustomizeDefaults(UMetasoundEditorGraphMemberDefaultLiteral& InLiteral, IDetailLayoutBuilder& InDetailLayout);
 
 		UE_DEPRECATED(5.5, "Use CustomizeDefaults instead and provide returned customized handles")
-		virtual TArray<IDetailPropertyRow*> CustomizeLiteral(UMetasoundEditorGraphMemberDefaultLiteral& InLiteral, IDetailLayoutBuilder& InDetailLayout);
+		UE_API virtual TArray<IDetailPropertyRow*> CustomizeLiteral(UMetasoundEditorGraphMemberDefaultLiteral& InLiteral, IDetailLayoutBuilder& InDetailLayout);
 
-		virtual TAttribute<EVisibility> GetDefaultVisibility() const;
-		virtual TAttribute<bool> GetEnabled() const;
+		UE_API virtual TAttribute<EVisibility> GetDefaultVisibility() const;
+		UE_API virtual TAttribute<bool> GetEnabled() const;
 
-		virtual void SetDefaultVisibility(TAttribute<EVisibility> VisibilityAttribute);
-		virtual void SetEnabled(TAttribute<bool> EnableAttribute);
-		virtual void SetResetOverride(const TOptional<FResetToDefaultOverride>& InResetOverride);
+		UE_API virtual void SetDefaultVisibility(TAttribute<EVisibility> VisibilityAttribute);
+		UE_API virtual void SetEnabled(TAttribute<bool> EnableAttribute);
+		UE_API virtual void SetResetOverride(const TOptional<FResetToDefaultOverride>& InResetOverride);
 
 	protected:
-		void CustomizePageDefaultRows(UMetasoundEditorGraphMemberDefaultLiteral& InLiteral, IDetailLayoutBuilder& InDetailLayout);
-		virtual void BuildDefaultValueWidget(IDetailPropertyRow& ValueRow, TSharedPtr<IPropertyHandle> ValueProperty);
+		UE_API void CustomizePageDefaultRows(UMetasoundEditorGraphMemberDefaultLiteral& InLiteral, IDetailLayoutBuilder& InDetailLayout);
+		UE_API virtual void BuildDefaultValueWidget(IDetailPropertyRow& ValueRow, TSharedPtr<IPropertyHandle> ValueProperty);
 
 		TArray<TSharedPtr<IPropertyHandle>> DefaultProperties;
 
 	private:
-		TSharedRef<SWidget> BuildPageDefaultNameWidget(UMetasoundEditorGraphMemberDefaultLiteral& Literal, TSharedRef<IPropertyHandle> ElementProperty);
-		void BuildPageDefaultComboBox(UMetasoundEditorGraphMemberDefaultLiteral& Literal, FText RowName);
-		void UpdatePagePickerNames(TWeakObjectPtr<UMetasoundEditorGraphMemberDefaultLiteral> LiteralPtr);
+		UE_API TSharedRef<SWidget> BuildPageDefaultNameWidget(UMetasoundEditorGraphMemberDefaultLiteral& Literal, TSharedRef<IPropertyHandle> ElementProperty);
+		UE_API void BuildPageDefaultComboBox(UMetasoundEditorGraphMemberDefaultLiteral& Literal, FText RowName);
+		UE_API void UpdatePagePickerNames(TWeakObjectPtr<UMetasoundEditorGraphMemberDefaultLiteral> LiteralPtr);
 
 		TArray<TSharedPtr<FString>> AddablePageStringNames;
 		TSet<FName> ImplementedPageNames;
@@ -66,7 +68,7 @@ namespace Metasound::Editor
 		TOptional<FResetToDefaultOverride> ResetOverride;
 	};
 
-	class METASOUNDEDITOR_API IMemberDefaultLiteralCustomizationFactory
+	class IMemberDefaultLiteralCustomizationFactory
 	{
 	public:
 		virtual ~IMemberDefaultLiteralCustomizationFactory() = default;
@@ -74,3 +76,5 @@ namespace Metasound::Editor
 		virtual TUniquePtr<FMetasoundDefaultLiteralCustomizationBase> CreateLiteralCustomization(IDetailCategoryBuilder& DefaultCategoryBuilder) const = 0;
 	};
 } // namespace Metasound::Editor
+
+#undef UE_API
diff --git a/Engine/Plugins/Runtime/Metasound/Source/MetasoundEditor/Public/MetasoundEditorGraph.h b/Engine/Plugins/Runtime/Metasound/Source/MetasoundEditor/Public/MetasoundEditorGraph.h
--- a/Engine/Plugins/Runtime/Metasound/Source/MetasoundEditor/Public/MetasoundEditorGraph.h
+++ b/Engine/Plugins/Runtime/Metasound/Source/MetasoundEditor/Public/MetasoundEditorGraph.h
@@ -17,6 +17,8 @@
 
 #include "MetasoundEditorGraph.generated.h"
 
+#define UE_API METASOUNDEDITOR_API
+
 
 // Forward Declarations
 struct FMetaSoundFrontendDocumentBuilder;
@@ -85,8 +87,8 @@ struct FMetasoundEditorGraphVariableBreadcrumb : public FMetasoundEditorGraphMem
 DECLARE_MULTICAST_DELEGATE_OneParam(FOnMetasoundMemberNameChanged, FGuid /* NodeID */);
 DECLARE_MULTICAST_DELEGATE(FOnMetasoundMemberRenameRequested);
 
-UCLASS()
-class METASOUNDEDITOR_API UMetasoundEditorGraphMemberDefaultLiteral : public UMetaSoundFrontendMemberMetadata
+UCLASS(MinimalAPI)
+class UMetasoundEditorGraphMemberDefaultLiteral : public UMetaSoundFrontendMemberMetadata
 {
 	GENERATED_BODY()
 
@@ -105,16 +107,16 @@ public:
 		return Value;
 	}
 
-	virtual void ForceRefresh() override;
-	virtual FName GetDataType() const;
-	virtual EMetasoundFrontendLiteralType GetLiteralType() const override;
-	virtual void InitDefault(const FGuid& InPageID = Metasound::Frontend::DefaultPageID);
+	UE_API virtual void ForceRefresh() override;
+	UE_API virtual FName GetDataType() const;
+	UE_API virtual EMetasoundFrontendLiteralType GetLiteralType() const override;
+	UE_API virtual void InitDefault(const FGuid& InPageID = Metasound::Frontend::DefaultPageID);
 	virtual void Initialize() { };
-	virtual void IterateDefaults(TFunctionRef<void(const FGuid&, FMetasoundFrontendLiteral)> Iter) const;
-	virtual bool RemoveDefault(const FGuid& InPageID);
-	virtual void ResetDefaults();
-	virtual void SetFromLiteral(const FMetasoundFrontendLiteral& InLiteral, const FGuid& InPageID = Metasound::Frontend::DefaultPageID) override;
-	virtual bool TryFindDefault(FMetasoundFrontendLiteral& OutLiteral, const FGuid* InPageID = nullptr) const;
+	UE_API virtual void IterateDefaults(TFunctionRef<void(const FGuid&, FMetasoundFrontendLiteral)> Iter) const;
+	UE_API virtual bool RemoveDefault(const FGuid& InPageID);
+	UE_API virtual void ResetDefaults();
+	UE_API virtual void SetFromLiteral(const FMetasoundFrontendLiteral& InLiteral, const FGuid& InPageID = Metasound::Frontend::DefaultPageID) override;
+	UE_API virtual bool TryFindDefault(FMetasoundFrontendLiteral& OutLiteral, const FGuid* InPageID = nullptr) const;
 
 	virtual void UpdatePreviewInstance(const Metasound::FVertexName& InParameterName, TScriptInterface<IAudioParameterControllerInterface>& InParameterInterface) const
 	{
@@ -124,16 +126,16 @@ public:
 	// initial editor-only fields based on context within editor/document model.
 
 #if WITH_EDITOR
-	virtual void PostEditChangeProperty(FPropertyChangedEvent& InPropertyChangedEvent) override;
-	virtual void PostEditChangeChainProperty(FPropertyChangedChainEvent& InPropertyChangedEvent) override;
-	virtual void PostEditUndo() override;
+	UE_API virtual void PostEditChangeProperty(FPropertyChangedEvent& InPropertyChangedEvent) override;
+	UE_API virtual void PostEditChangeChainProperty(FPropertyChangedChainEvent& InPropertyChangedEvent) override;
+	UE_API virtual void PostEditUndo() override;
 #endif // WITH_EDITOR
 
 	// Synchronizes local transient editor-only member data with changes made to the
 	// associated Frontend Document member. Returns whether or not object was modified.
 	virtual bool Synchronize() { return false; };
 
-	UMetasoundEditorGraphMember* FindMember() const;
+	UE_API UMetasoundEditorGraphMember* FindMember() const;
 
 	UE_DEPRECATED(5.5, "Due to serialization optimization, literals now inherit from LiteralMetadata and are no longer parented under members. "
 		"Use 'GetDefaultsPropertyName' instead to get transient paged defaults property name")
@@ -162,13 +164,13 @@ public:
 protected:
 	virtual void ResolvePageDefaults() { }
 	virtual void SortPageDefaults() { }
-	bool TryGetPreviewPageID(FGuid& OutPreviewPageID) const;
+	UE_API bool TryGetPreviewPageID(FGuid& OutPreviewPageID) const;
 };
 
 /** UMetasoundEditorGraphMember is a base class for non-node graph level members 
  * such as inputs, outputs and variables. */
-UCLASS(Abstract)
-class METASOUNDEDITOR_API UMetasoundEditorGraphMember : public UObject
+UCLASS(MinimalAPI, Abstract)
+class UMetasoundEditorGraphMember : public UObject
 {
 	GENERATED_BODY()
 
@@ -222,17 +224,17 @@ public:
 	virtual void UpdateFrontendDefaultLiteral(bool bPostTransaction, const FGuid* InPageID = nullptr) PURE_VIRTUAL(UMetasoundEditorGraphMember::UpdateFrontendDefaultLiteral, );
 
 	// Synchronizes cached data with the frontend representation on the represented document.
-	virtual bool Synchronize();
+	UE_API virtual bool Synchronize();
 
-	FMetaSoundFrontendDocumentBuilder& GetFrontendBuilderChecked() const;
+	UE_API FMetaSoundFrontendDocumentBuilder& GetFrontendBuilderChecked() const;
 
 	/** Returns the parent MetaSound Graph. If the Outer object of the member is non
 	 * a UMetasoundEditorGraph, returns a nullptr. */
-	UMetasoundEditorGraph* GetOwningGraph();
+	UE_API UMetasoundEditorGraph* GetOwningGraph();
 
 	/** Returns the parent MetaSound Graph. If the Outer object of the member is non
 	 * a UMetasoundEditorGraph, returns a nullptr. */
-	const UMetasoundEditorGraph* GetOwningGraph() const;
+	UE_API const UMetasoundEditorGraph* GetOwningGraph() const;
 
 	/* Whether this member can be renamed. */
 	virtual bool CanRename() const PURE_VIRTUAL(UMetasoundEditorGraphMember::CanRename, return false;);
@@ -244,11 +246,11 @@ public:
 	virtual void CacheBreadcrumb() { };
 
 #if WITH_EDITOR
-	virtual void PostEditUndo() override;
+	UE_API virtual void PostEditUndo() override;
 #endif // WITH_EDITOR
 
 	/** Returns the current data type */
-	FName GetDataType() const;
+	UE_API FName GetDataType() const;
 
 	static FName GetLiteralPropertyName()
 	{
@@ -259,7 +261,7 @@ public:
 	UMetasoundEditorGraphMemberDefaultLiteral* GetLiteral() const { return Literal; }
 
 	/** Creates new literal if there is none and/or conforms literal object type to member's DataType */
-	void InitializeLiteral();
+	UE_API void InitializeLiteral();
 
 protected:
 	UE_DEPRECATED(5.5, "No longer supported in favor of internal Editor::FGraphBuilder API which is actively being transitioned to using Document Builder API")
@@ -278,8 +280,8 @@ protected:
 };
 
 /** Base class for an input or output of the graph. */
-UCLASS(Abstract)
-class METASOUNDEDITOR_API UMetasoundEditorGraphVertex : public UMetasoundEditorGraphMember
+UCLASS(MinimalAPI, Abstract)
+class UMetasoundEditorGraphVertex : public UMetasoundEditorGraphMember
 {
 	GENERATED_BODY()
 
@@ -287,17 +289,17 @@ protected:
 	UE_DEPRECATED(5.5, "EditorGraph vertices no longer generates node handles (use Builder API instead)")
 	virtual Metasound::Frontend::FNodeHandle AddNodeHandle(const FName& InNodeName, const Metasound::Editor::FCreateNodeVertexParams& InParams) PURE_VIRTUAL(UMetasoundEditorGraphVertex::AddNodeHandle, return Metasound::Frontend::INodeController::GetInvalidHandle(); )
 
-	const FMetasoundFrontendNode* GetFrontendNode() const;
+	UE_API const FMetasoundFrontendNode* GetFrontendNode() const;
 
 	/** If the vertex's name contains an interface namespace, return true and set OutInterface to that interface if provided. */
-	bool NameContainsInterfaceNamespace(FMetasoundFrontendInterface* OutInterface) const;
+	UE_API bool NameContainsInterfaceNamespace(FMetasoundFrontendInterface* OutInterface) const;
 
 	/** Implements frontend call to rename member */
 	virtual bool RenameFrontendMemberInternal(FMetaSoundFrontendDocumentBuilder& Builder, FName OldName, FName InNewName) const PURE_VIRTUAL(UMetasoundEditorGraphVertex::RenameFrontendMemberInternal, return false;)
 
 public:
 	/** Initializes all properties with the given parameters required to identify the frontend member from this editor graph member. */
-	void InitMember(FName InDataType, const FMetasoundFrontendLiteral& InDefaultLiteral, FGuid InNodeID, FMetasoundFrontendClassName&& InClassName);
+	UE_API void InitMember(FName InDataType, const FMetasoundFrontendLiteral& InDefaultLiteral, FGuid InNodeID, FMetasoundFrontendClassName&& InClassName);
 
 	/** ID of Metasound Frontend node. */
 	UPROPERTY()
@@ -308,26 +310,26 @@ public:
 	FMetasoundFrontendClassName ClassName;
 
 	/* ~Begin UMetasoundEditorGraphMember interface */
-	virtual FGuid GetMemberID() const override;
-	virtual FName GetMemberName() const override;
-	virtual FText GetDisplayName() const override;
-
-	virtual bool CanRename(const FText& InNewName, FText& OutError) const override;
-	virtual TArray<UMetasoundEditorGraphMemberNode*> GetNodes() const override;
-	virtual void SetDescription(const FText& InDescription, bool bPostTransaction) override;
-	virtual void SetMemberName(const FName& InNewName, bool bPostTransaction) override;
-	virtual bool Synchronize() override;
-	virtual void CacheBreadcrumb() override;
+	UE_API virtual FGuid GetMemberID() const override;
+	UE_API virtual FName GetMemberName() const override;
+	UE_API virtual FText GetDisplayName() const override;
+
+	UE_API virtual bool CanRename(const FText& InNewName, FText& OutError) const override;
+	UE_API virtual TArray<UMetasoundEditorGraphMemberNode*> GetNodes() const override;
+	UE_API virtual void SetDescription(const FText& InDescription, bool bPostTransaction) override;
+	UE_API virtual void SetMemberName(const FName& InNewName, bool bPostTransaction) override;
+	UE_API virtual bool Synchronize() override;
+	UE_API virtual void CacheBreadcrumb() override;
 	/* ~End UMetasoundEditorGraphMember interface */
 
 	/** Returns Breadcrumb data of this Vertex*/
 	const FMetasoundEditorGraphVertexBreadcrumb& GetBreadcrumb() const { return Breadcrumb; };
 
 	/** Version of interface membership, or invalid version if not an interface member. */
-	virtual const FMetasoundFrontendVersion& GetInterfaceVersion() const;
+	UE_API virtual const FMetasoundFrontendVersion& GetInterfaceVersion() const;
 
 	/** Returns true if member is part of an interface. If supplied interface pointer, sets pointer's data to the interface vertex is member of. */
-	virtual bool IsInterfaceMember(FMetasoundFrontendInterface* OutInterface = nullptr) const;
+	UE_API virtual bool IsInterfaceMember(FMetasoundFrontendInterface* OutInterface = nullptr) const;
 
 	/** Returns the Metasound class type of the associated node */
 	virtual EMetasoundFrontendClassType GetClassType() const PURE_VIRTUAL(UMetasoundEditorGraphVertex::GetClassType, return EMetasoundFrontendClassType::Invalid; )
@@ -348,104 +350,104 @@ public:
 
 #if WITH_EDITORONLY_DATA
 	/**Sets if the Vertex is shown as Advanced Display*/
-	bool SetIsAdvancedDisplay(const bool IsAdvancedDisplay);
+	UE_API bool SetIsAdvancedDisplay(const bool IsAdvancedDisplay);
 #endif // WITH_EDITORONLY_DATA
 
 	UE_DEPRECATED(5.6, "Node Handles are actively being deprecated, use the MetaSound Frontend Document Builder API")
-	Metasound::Frontend::FNodeHandle GetNodeHandle();
+	UE_API Metasound::Frontend::FNodeHandle GetNodeHandle();
 
 	/** Returns the node handle associated with the vertex. */
-	Metasound::Frontend::FConstNodeHandle GetConstNodeHandle() const;
+	UE_API Metasound::Frontend::FConstNodeHandle GetConstNodeHandle() const;
 
 	virtual EMetasoundFrontendVertexAccessType GetVertexAccessType() const PURE_VIRTUAL(UMetasoundEditorGraphVertex::GetVertexAccessType, return EMetasoundFrontendVertexAccessType::Reference; );
 
-	virtual bool CanRename() const override;
+	UE_API virtual bool CanRename() const override;
 
 	friend class Metasound::Editor::FEditor;
 
 protected:
-	void SetMemberNameInternal(const FName& InNewName, bool bPropagateToPinNames, bool bPostTransaction);
+	UE_API void SetMemberNameInternal(const FName& InNewName, bool bPropagateToPinNames, bool bPostTransaction);
 
 	UPROPERTY()
 	FMetasoundEditorGraphVertexBreadcrumb Breadcrumb;
 };
 
-UCLASS()
-class METASOUNDEDITOR_API UMetasoundEditorGraphInput : public UMetasoundEditorGraphVertex
+UCLASS(MinimalAPI)
+class UMetasoundEditorGraphInput : public UMetasoundEditorGraphVertex
 {
 	GENERATED_BODY()
 
 public:
-	virtual FText GetDescription() const override;
-	virtual const FMetasoundFrontendClassVertex* GetFrontendClassVertex() const override;
-	virtual int32 GetSortOrderIndex() const override;
-	virtual void SetSortOrderIndex(int32 InSortOrderIndex) override;
-	virtual TArray<UMetasoundEditorGraphMemberNode*> GetNodes() const override;
-
-	virtual const FText& GetGraphMemberLabel() const override;
-	virtual bool IsDefaultPaged() const override;
-	virtual bool IsInterfaceMember(FMetasoundFrontendInterface* OutInterface = nullptr) const override;
-	virtual void ResetToClassDefault() override;
-	virtual void SetDataType(FName InNewType, bool bPostTransaction) override;
-	virtual void SetDescription(const FText& InDescription, bool bPostTransaction) override;
-	virtual void SetDisplayName(const FText& InNewName, bool bPostTransaction) override;
-	virtual void SetMemberName(const FName& InNewName, bool bPostTransaction) override;
-	virtual void SetVertexAccessType(EMetasoundFrontendVertexAccessType InNewAccessType, bool bPostTransaction) override;
-	virtual bool Synchronize() override;
-	virtual void CacheBreadcrumb() override;
-	virtual void UpdateFrontendDefaultLiteral(bool bPostTransaction, const FGuid* InPageID = nullptr) override;
-	virtual EMetasoundFrontendVertexAccessType GetVertexAccessType() const override;
+	UE_API virtual FText GetDescription() const override;
+	UE_API virtual const FMetasoundFrontendClassVertex* GetFrontendClassVertex() const override;
+	UE_API virtual int32 GetSortOrderIndex() const override;
+	UE_API virtual void SetSortOrderIndex(int32 InSortOrderIndex) override;
+	UE_API virtual TArray<UMetasoundEditorGraphMemberNode*> GetNodes() const override;
+
+	UE_API virtual const FText& GetGraphMemberLabel() const override;
+	UE_API virtual bool IsDefaultPaged() const override;
+	UE_API virtual bool IsInterfaceMember(FMetasoundFrontendInterface* OutInterface = nullptr) const override;
+	UE_API virtual void ResetToClassDefault() override;
+	UE_API virtual void SetDataType(FName InNewType, bool bPostTransaction) override;
+	UE_API virtual void SetDescription(const FText& InDescription, bool bPostTransaction) override;
+	UE_API virtual void SetDisplayName(const FText& InNewName, bool bPostTransaction) override;
+	UE_API virtual void SetMemberName(const FName& InNewName, bool bPostTransaction) override;
+	UE_API virtual void SetVertexAccessType(EMetasoundFrontendVertexAccessType InNewAccessType, bool bPostTransaction) override;
+	UE_API virtual bool Synchronize() override;
+	UE_API virtual void CacheBreadcrumb() override;
+	UE_API virtual void UpdateFrontendDefaultLiteral(bool bPostTransaction, const FGuid* InPageID = nullptr) override;
+	UE_API virtual EMetasoundFrontendVertexAccessType GetVertexAccessType() const override;
 
 protected:
 	UE_DEPRECATED(5.5, "No longer supported in favor of internal Editor::FGraphBuilder API which is actively being transitioned to using Document Builder API")
-	virtual UMetasoundEditorGraphNode* AddNode(Metasound::Frontend::FNodeHandle InNodeHandle, bool bInSelectNewNode) override;
+	UE_API virtual UMetasoundEditorGraphNode* AddNode(Metasound::Frontend::FNodeHandle InNodeHandle, bool bInSelectNewNode) override;
 
 	UE_DEPRECATED(5.5, "EditorGraph vertices no longer generates node handles (use Builder API instead)")
-	virtual Metasound::Frontend::FNodeHandle AddNodeHandle(const FName& InNodeName, const Metasound::Editor::FCreateNodeVertexParams& InParams) override;
+	UE_API virtual Metasound::Frontend::FNodeHandle AddNodeHandle(const FName& InNodeName, const Metasound::Editor::FCreateNodeVertexParams& InParams) override;
 
-	virtual bool RenameFrontendMemberInternal(FMetaSoundFrontendDocumentBuilder& Builder, FName OldName, FName InNewName) const override;
+	UE_API virtual bool RenameFrontendMemberInternal(FMetaSoundFrontendDocumentBuilder& Builder, FName OldName, FName InNewName) const override;
 
 	virtual EMetasoundFrontendClassType GetClassType() const override { return EMetasoundFrontendClassType::Input; }
-	virtual Metasound::Editor::ENodeSection GetSectionID() const override;
+	UE_API virtual Metasound::Editor::ENodeSection GetSectionID() const override;
 };
 
 
-UCLASS()
-class METASOUNDEDITOR_API UMetasoundEditorGraphOutput : public UMetasoundEditorGraphVertex
+UCLASS(MinimalAPI)
+class UMetasoundEditorGraphOutput : public UMetasoundEditorGraphVertex
 {
 	GENERATED_BODY()
 
 public:
-	virtual FText GetDescription() const override;
-	virtual const FMetasoundFrontendClassVertex* GetFrontendClassVertex() const override;
-	virtual int32 GetSortOrderIndex() const override;
-	virtual bool IsInterfaceMember(FMetasoundFrontendInterface* OutInterface = nullptr) const override;
-	virtual void SetSortOrderIndex(int32 InSortOrderIndex) override;
-	virtual const FText& GetGraphMemberLabel() const override;
-	virtual void ResetToClassDefault() override;
-	virtual void SetDataType(FName InNewType, bool bPostTransaction) override;
-	virtual void SetDescription(const FText& InDescription, bool bPostTransaction) override;
-	virtual void SetDisplayName(const FText& InNewName, bool bPostTransaction) override;
-	virtual void SetVertexAccessType(EMetasoundFrontendVertexAccessType InNewAccessType, bool bPostTransaction) override;
-	virtual bool Synchronize() override;
-	virtual void UpdateFrontendDefaultLiteral(bool bPostTransaction, const FGuid* InPageID = nullptr) override;
-	virtual EMetasoundFrontendVertexAccessType GetVertexAccessType() const override;
+	UE_API virtual FText GetDescription() const override;
+	UE_API virtual const FMetasoundFrontendClassVertex* GetFrontendClassVertex() const override;
+	UE_API virtual int32 GetSortOrderIndex() const override;
+	UE_API virtual bool IsInterfaceMember(FMetasoundFrontendInterface* OutInterface = nullptr) const override;
+	UE_API virtual void SetSortOrderIndex(int32 InSortOrderIndex) override;
+	UE_API virtual const FText& GetGraphMemberLabel() const override;
+	UE_API virtual void ResetToClassDefault() override;
+	UE_API virtual void SetDataType(FName InNewType, bool bPostTransaction) override;
+	UE_API virtual void SetDescription(const FText& InDescription, bool bPostTransaction) override;
+	UE_API virtual void SetDisplayName(const FText& InNewName, bool bPostTransaction) override;
+	UE_API virtual void SetVertexAccessType(EMetasoundFrontendVertexAccessType InNewAccessType, bool bPostTransaction) override;
+	UE_API virtual bool Synchronize() override;
+	UE_API virtual void UpdateFrontendDefaultLiteral(bool bPostTransaction, const FGuid* InPageID = nullptr) override;
+	UE_API virtual EMetasoundFrontendVertexAccessType GetVertexAccessType() const override;
 
 protected:
 	UE_DEPRECATED(5.5, "No longer supported in favor of internal Editor::FGraphBuilder API which is actively being transitioned to using Document Builder API")
-	virtual UMetasoundEditorGraphNode* AddNode(Metasound::Frontend::FNodeHandle InNodeHandle, bool bInSelectNewNode) override;
+	UE_API virtual UMetasoundEditorGraphNode* AddNode(Metasound::Frontend::FNodeHandle InNodeHandle, bool bInSelectNewNode) override;
 
 	UE_DEPRECATED(5.5, "No longer supported in favor of internal Editor::FGraphBuilder API which is actively being transitioned to using Document Builder API")
-	virtual Metasound::Frontend::FNodeHandle AddNodeHandle(const FName& InNodeName, const Metasound::Editor::FCreateNodeVertexParams& InParams) override;
+	UE_API virtual Metasound::Frontend::FNodeHandle AddNodeHandle(const FName& InNodeName, const Metasound::Editor::FCreateNodeVertexParams& InParams) override;
 
-	virtual bool RenameFrontendMemberInternal(FMetaSoundFrontendDocumentBuilder& Builder, FName OldName, FName InNewName) const override;
+	UE_API virtual bool RenameFrontendMemberInternal(FMetaSoundFrontendDocumentBuilder& Builder, FName OldName, FName InNewName) const override;
 
 	virtual EMetasoundFrontendClassType GetClassType() const override { return EMetasoundFrontendClassType::Output; }
-	virtual Metasound::Editor::ENodeSection GetSectionID() const override;
+	UE_API virtual Metasound::Editor::ENodeSection GetSectionID() const override;
 };
 
-UCLASS()
-class METASOUNDEDITOR_API UMetasoundEditorGraphVariable : public UMetasoundEditorGraphMember
+UCLASS(MinimalAPI)
+class UMetasoundEditorGraphVariable : public UMetasoundEditorGraphMember
 {
 	GENERATED_BODY()
 
@@ -453,44 +455,44 @@ class METASOUNDEDITOR_API UMetasoundEditorGraphVariable : public UMetasoundEdito
 	FGuid VariableID;
 
 public:
-	void InitMember(FName InDataType, const FMetasoundFrontendLiteral& InDefaultLiteral, FGuid InVariableID);
+	UE_API void InitMember(FName InDataType, const FMetasoundFrontendLiteral& InDefaultLiteral, FGuid InVariableID);
 
 	/* ~Begin UMetasoundEditorGraphMember interface */
-	virtual Metasound::Editor::ENodeSection GetSectionID() const override;
-	virtual void SetDataType(FName InNewType, bool bPostTransaction = true) override;
+	UE_API virtual Metasound::Editor::ENodeSection GetSectionID() const override;
+	UE_API virtual void SetDataType(FName InNewType, bool bPostTransaction = true) override;
 
-	virtual FText GetDescription() const override;
-	virtual void SetDescription(const FText& InDescription, bool bPostTransaction) override;
+	UE_API virtual FText GetDescription() const override;
+	UE_API virtual void SetDescription(const FText& InDescription, bool bPostTransaction) override;
 
-	virtual FGuid GetMemberID() const override;
-	virtual bool CanRename(const FText& InNewName, FText& OutError) const override;
-	virtual void SetMemberName(const FName& InNewName, bool bPostTransaction) override;
-	virtual FName GetMemberName() const override;
+	UE_API virtual FGuid GetMemberID() const override;
+	UE_API virtual bool CanRename(const FText& InNewName, FText& OutError) const override;
+	UE_API virtual void SetMemberName(const FName& InNewName, bool bPostTransaction) override;
+	UE_API virtual FName GetMemberName() const override;
 
-	virtual FText GetDisplayName() const override;
-	virtual void SetDisplayName(const FText& InNewName, bool bPostTransaction) override;
+	UE_API virtual FText GetDisplayName() const override;
+	UE_API virtual void SetDisplayName(const FText& InNewName, bool bPostTransaction) override;
 
-	virtual void ResetToClassDefault() override;
-	virtual bool Synchronize() override;
-	virtual void CacheBreadcrumb() override;
-	virtual void UpdateFrontendDefaultLiteral(bool bPostTransaction, const FGuid* InPageID = nullptr) override;
+	UE_API virtual void ResetToClassDefault() override;
+	UE_API virtual bool Synchronize() override;
+	UE_API virtual void CacheBreadcrumb() override;
+	UE_API virtual void UpdateFrontendDefaultLiteral(bool bPostTransaction, const FGuid* InPageID = nullptr) override;
 
-	virtual TArray<UMetasoundEditorGraphMemberNode*> GetNodes() const override;
+	UE_API virtual TArray<UMetasoundEditorGraphMemberNode*> GetNodes() const override;
 
-	virtual const FText& GetGraphMemberLabel() const override;
+	UE_API virtual const FText& GetGraphMemberLabel() const override;
 	/* ~EndUMetasoundEditorGraphMember interface */
 
-	const FGuid& GetVariableID() const;
-	const FMetasoundFrontendVariable* GetFrontendVariable() const;
+	UE_API const FGuid& GetVariableID() const;
+	UE_API const FMetasoundFrontendVariable* GetFrontendVariable() const;
 
 	/**Returns Breadcrumb data of this Variable*/
 	const FMetasoundEditorGraphVariableBreadcrumb& GetBreadcrumb() const { return Breadcrumb; };
 
 	UE_DEPRECATED(5.6, "")
-	Metasound::Frontend::FVariableHandle GetVariableHandle();
-	Metasound::Frontend::FConstVariableHandle GetConstVariableHandle() const;
+	UE_API Metasound::Frontend::FVariableHandle GetVariableHandle();
+	UE_API Metasound::Frontend::FConstVariableHandle GetConstVariableHandle() const;
 
-	virtual bool CanRename() const override;
+	UE_API virtual bool CanRename() const override;
 
 protected:
 
@@ -499,7 +501,7 @@ protected:
 
 protected:
 	UE_DEPRECATED(5.5, "No longer supported in favor of internal Editor::FGraphBuilder API which is actively being transitioned to using Document Builder API")
-	virtual UMetasoundEditorGraphNode* AddNode(Metasound::Frontend::FNodeHandle InNodeHandle, bool bInSelectNewNode) override;
+	UE_API virtual UMetasoundEditorGraphNode* AddNode(Metasound::Frontend::FNodeHandle InNodeHandle, bool bInSelectNewNode) override;
 
 private:
 	struct FVariableEditorNodes
@@ -517,55 +519,55 @@ private:
 		TArray<FVector2D> DeferredAccessorLocations;
 	};
 
-	FVariableEditorNodes GetVariableNodes() const;
-	FVariableNodeLocations GetVariableNodeLocations() const;
-	void AddVariableNodes(UObject& InMetaSound, const FVariableNodeLocations& InNodeLocs);
+	UE_API FVariableEditorNodes GetVariableNodes() const;
+	UE_API FVariableNodeLocations GetVariableNodeLocations() const;
+	UE_API void AddVariableNodes(UObject& InMetaSound, const FVariableNodeLocations& InNodeLocs);
 };
 
-UCLASS()
-class METASOUNDEDITOR_API UMetasoundEditorGraph : public UMetasoundEditorGraphBase
+UCLASS(MinimalAPI)
+class UMetasoundEditorGraph : public UMetasoundEditorGraphBase
 {
 	GENERATED_BODY()
 
 public:
 	UE_DEPRECATED(5.5, "Input node creation is no longer supported publically via the MetasoundEditorGraph.")
-	UMetasoundEditorGraphInputNode* CreateInputNode(Metasound::Frontend::FNodeHandle InNodeHandle, bool bInSelectNewNode);
+	UE_API UMetasoundEditorGraphInputNode* CreateInputNode(Metasound::Frontend::FNodeHandle InNodeHandle, bool bInSelectNewNode);
 
 	UE_DEPRECATED(5.6, "Document Handles are actively being deprecated, use the MetaSound Frontend Document Builder API")
-	Metasound::Frontend::FDocumentHandle GetDocumentHandle();
+	UE_API Metasound::Frontend::FDocumentHandle GetDocumentHandle();
 
-	Metasound::Frontend::FConstDocumentHandle GetDocumentHandle() const;
+	UE_API Metasound::Frontend::FConstDocumentHandle GetDocumentHandle() const;
 
 	UE_DEPRECATED(5.6, "Document Handles are actively being deprecated, use the MetaSound Frontend Document Builder API")
-	Metasound::Frontend::FGraphHandle GetGraphHandle();
+	UE_API Metasound::Frontend::FGraphHandle GetGraphHandle();
 
-	Metasound::Frontend::FConstGraphHandle GetGraphHandle() const;
+	UE_API Metasound::Frontend::FConstGraphHandle GetGraphHandle() const;
 
-	virtual void PreSave(FObjectPreSaveContext InSaveContext) override;
+	UE_API virtual void PreSave(FObjectPreSaveContext InSaveContext) override;
 
-	UMetaSoundBuilderBase& GetBuilderChecked() const;
-	UObject* GetMetasound() const;
-	UObject& GetMetasoundChecked() const;
+	UE_API UMetaSoundBuilderBase& GetBuilderChecked() const;
+	UE_API UObject* GetMetasound() const;
+	UE_API UObject& GetMetasoundChecked() const;
 
-	void IterateInputs(TFunctionRef<void(UMetasoundEditorGraphInput&)> InFunction) const;
-	void IterateOutputs(TFunctionRef<void(UMetasoundEditorGraphOutput&)> InFunction) const;
-	void IterateVariables(TFunctionRef<void(UMetasoundEditorGraphVariable&)> InFunction) const;
-	void IterateMembers(TFunctionRef<void(UMetasoundEditorGraphMember&)> InFunction) const;
+	UE_API void IterateInputs(TFunctionRef<void(UMetasoundEditorGraphInput&)> InFunction) const;
+	UE_API void IterateOutputs(TFunctionRef<void(UMetasoundEditorGraphOutput&)> InFunction) const;
+	UE_API void IterateVariables(TFunctionRef<void(UMetasoundEditorGraphVariable&)> InFunction) const;
+	UE_API void IterateMembers(TFunctionRef<void(UMetasoundEditorGraphMember&)> InFunction) const;
 
-	bool ContainsInput(const UMetasoundEditorGraphInput& InInput) const;
-	bool ContainsOutput(const UMetasoundEditorGraphOutput& InOutput) const;
-	bool ContainsVariable(const UMetasoundEditorGraphVariable& InVariable) const;
+	UE_API bool ContainsInput(const UMetasoundEditorGraphInput& InInput) const;
+	UE_API bool ContainsOutput(const UMetasoundEditorGraphOutput& InOutput) const;
+	UE_API bool ContainsVariable(const UMetasoundEditorGraphVariable& InVariable) const;
 
-	void SetPreviewID(uint32 InPreviewID);
-	bool IsPreviewing() const;
-	bool IsEditable() const;
+	UE_API void SetPreviewID(uint32 InPreviewID);
+	UE_API bool IsPreviewing() const;
+	UE_API bool IsEditable() const;
 
 	// UMetasoundEditorGraphBase Implementation
-	virtual void RegisterGraphWithFrontend() override;
-	virtual void MigrateEditorDocumentData(FMetaSoundFrontendDocumentBuilder& OutBuilder) override;
+	UE_API virtual void RegisterGraphWithFrontend() override;
+	UE_API virtual void MigrateEditorDocumentData(FMetaSoundFrontendDocumentBuilder& OutBuilder) override;
 
 private:
-	void ValidateInternal(Metasound::Editor::FGraphValidationResults& OutResults);
+	UE_API void ValidateInternal(Metasound::Editor::FGraphValidationResults& OutResults);
 
 	// Preview ID is the Unique ID provided by the UObject that implements
 	// a sound's ParameterInterface when a sound begins playing.
@@ -587,28 +589,28 @@ private:
 	TArray<TObjectPtr<UMetasoundEditorGraphVariable>> Variables;
 
 public:
-	UMetasoundEditorGraphInput* FindInput(FGuid InNodeID) const;
-	UMetasoundEditorGraphInput* FindInput(FName InName) const;
-	UMetasoundEditorGraphInput* FindOrAddInput(const FGuid& InNodeID);
+	UE_API UMetasoundEditorGraphInput* FindInput(FGuid InNodeID) const;
+	UE_API UMetasoundEditorGraphInput* FindInput(FName InName) const;
+	UE_API UMetasoundEditorGraphInput* FindOrAddInput(const FGuid& InNodeID);
 
 	UE_DEPRECATED(5.6, "ConstNodeHandle is being deprecated, use overload which is provided NodeID instead")
-	UMetasoundEditorGraphInput* FindOrAddInput(Metasound::Frontend::FConstNodeHandle InNodeHandle);
+	UE_API UMetasoundEditorGraphInput* FindOrAddInput(Metasound::Frontend::FConstNodeHandle InNodeHandle);
 
-	UMetasoundEditorGraphOutput* FindOutput(FGuid InNodeID) const;
-	UMetasoundEditorGraphOutput* FindOutput(FName InName) const;
-	UMetasoundEditorGraphOutput* FindOrAddOutput(const FGuid& InNodeID);
+	UE_API UMetasoundEditorGraphOutput* FindOutput(FGuid InNodeID) const;
+	UE_API UMetasoundEditorGraphOutput* FindOutput(FName InName) const;
+	UE_API UMetasoundEditorGraphOutput* FindOrAddOutput(const FGuid& InNodeID);
 
 	UE_DEPRECATED(5.6, "ConstNodeHandle is being deprecated, use overload which is provided NodeID instead")
-	UMetasoundEditorGraphOutput* FindOrAddOutput(Metasound::Frontend::FConstNodeHandle InNodeHandle);
+	UE_API UMetasoundEditorGraphOutput* FindOrAddOutput(Metasound::Frontend::FConstNodeHandle InNodeHandle);
 
-	UMetasoundEditorGraphVariable* FindVariable(const FGuid& InVariableID) const;
-	UMetasoundEditorGraphVariable* FindOrAddVariable(FName VariableName);
+	UE_API UMetasoundEditorGraphVariable* FindVariable(const FGuid& InVariableID) const;
+	UE_API UMetasoundEditorGraphVariable* FindOrAddVariable(FName VariableName);
 
 	UE_DEPRECATED(5.6, "ConstVariableHandle is being deprecated, use overload which is provided VariableName instead")
-	UMetasoundEditorGraphVariable* FindOrAddVariable(const Metasound::Frontend::FConstVariableHandle& InVariableHandle);
+	UE_API UMetasoundEditorGraphVariable* FindOrAddVariable(const Metasound::Frontend::FConstVariableHandle& InVariableHandle);
 
-	UMetasoundEditorGraphMember* FindMember(FGuid InMemberID) const;
-	UMetasoundEditorGraphMember* FindAdjacentMember(const UMetasoundEditorGraphMember& InMember);
+	UE_API UMetasoundEditorGraphMember* FindMember(FGuid InMemberID) const;
+	UE_API UMetasoundEditorGraphMember* FindAdjacentMember(const UMetasoundEditorGraphMember& InMember);
 
 private:
 	/**
@@ -656,3 +658,5 @@ public:
 	friend class UMetaSoundSourceFactory;
 	friend class Metasound::Editor::FGraphBuilder;
 };
+
+#undef UE_API
diff --git a/Engine/Plugins/Runtime/Metasound/Source/MetasoundEditor/Public/MetasoundEditorGraphSchema.h b/Engine/Plugins/Runtime/Metasound/Source/MetasoundEditor/Public/MetasoundEditorGraphSchema.h
--- a/Engine/Plugins/Runtime/Metasound/Source/MetasoundEditor/Public/MetasoundEditorGraphSchema.h
+++ b/Engine/Plugins/Runtime/Metasound/Source/MetasoundEditor/Public/MetasoundEditorGraphSchema.h
@@ -13,6 +13,8 @@
 
 #include "MetasoundEditorGraphSchema.generated.h"
 
+#define UE_API METASOUNDEDITOR_API
+
 
 // Forward Declarations
 class UEdGraph;
@@ -38,7 +40,7 @@ namespace Metasound::Editor
 
 	const FText& GetContextGroupDisplayName(EPrimaryContextGroup InContextGroup);
 
-	class METASOUNDEDITOR_API ISchemaQueryResult
+	class ISchemaQueryResult
 	{
 	public:
 		virtual ~ISchemaQueryResult() = default;
@@ -90,7 +92,7 @@ namespace Metasound::Editor
 
 
 USTRUCT()
-struct METASOUNDEDITOR_API FMetasoundGraphSchemaAction : public FEdGraphSchemaAction
+struct FMetasoundGraphSchemaAction : public FEdGraphSchemaAction
 {
 	GENERATED_USTRUCT_BODY()
 
@@ -125,7 +127,7 @@ struct METASOUNDEDITOR_API FMetasoundGraphSchemaAction : public FEdGraphSchemaAc
 
 /** This is used to combine functionality for nodes that can have multiple outputs and should never be directly instantiated. */
 USTRUCT()
-struct METASOUNDEDITOR_API FMetasoundGraphSchemaAction_NodeWithMultipleOutputs : public FMetasoundGraphSchemaAction
+struct FMetasoundGraphSchemaAction_NodeWithMultipleOutputs : public FMetasoundGraphSchemaAction
 {
 	GENERATED_USTRUCT_BODY();
 
@@ -141,14 +143,14 @@ struct METASOUNDEDITOR_API FMetasoundGraphSchemaAction_NodeWithMultipleOutputs :
 
 	//~ Begin FEdGraphSchemaAction Interface
 	using FEdGraphSchemaAction::PerformAction; // Prevent hiding of deprecated base class function with FVector2D
-	virtual UEdGraphNode* PerformAction(UEdGraph* ParentGraph, TArray<UEdGraphPin*>& FromPins, const FVector2f& Location, bool bSelectNewNode = true) override;
+	UE_API virtual UEdGraphNode* PerformAction(UEdGraph* ParentGraph, TArray<UEdGraphPin*>& FromPins, const FVector2f& Location, bool bSelectNewNode = true) override;
 	virtual UEdGraphNode* PerformAction(UEdGraph* ParentGraph, UEdGraphPin* FromPin, const FVector2f& Location, bool bSelectNewNode = true) { return nullptr; }
 	//~ End FEdGraphSchemaAction Interface
 };
 
 /** Action to add an input reference to the graph */
 USTRUCT()
-struct METASOUNDEDITOR_API FMetasoundGraphSchemaAction_NewInput : public FMetasoundGraphSchemaAction_NodeWithMultipleOutputs
+struct FMetasoundGraphSchemaAction_NewInput : public FMetasoundGraphSchemaAction_NodeWithMultipleOutputs
 {
 	GENERATED_USTRUCT_BODY();
 
@@ -159,41 +161,41 @@ struct METASOUNDEDITOR_API FMetasoundGraphSchemaAction_NewInput : public FMetaso
 		: FMetasoundGraphSchemaAction_NodeWithMultipleOutputs()
 	{}
 
-	FMetasoundGraphSchemaAction_NewInput(FText InNodeCategory, FText InDisplayName, FGuid InInputNodeID, FText InToolTip, Metasound::Editor::EPrimaryContextGroup InGroup);
+	UE_API FMetasoundGraphSchemaAction_NewInput(FText InNodeCategory, FText InDisplayName, FGuid InInputNodeID, FText InToolTip, Metasound::Editor::EPrimaryContextGroup InGroup);
 
 	virtual ~FMetasoundGraphSchemaAction_NewInput() = default;
 
 	//~ Begin FMetasoundGraphSchemaAction Interface
-	virtual const FSlateBrush* GetIconBrush() const override;
+	UE_API virtual const FSlateBrush* GetIconBrush() const override;
 
-	virtual const FLinearColor& GetIconColor() const override;
+	UE_API virtual const FLinearColor& GetIconColor() const override;
 	//~ End FMetasoundGraphSchemaAction Interface
 
 	//~ Begin FEdGraphSchemaAction Interface
 	using FEdGraphSchemaAction::PerformAction; // Prevent hiding of deprecated base class function with FVector2D
-	virtual UEdGraphNode* PerformAction(UEdGraph* ParentGraph, UEdGraphPin* FromPin, const FVector2f& Location, bool bSelectNewNode = true) override;
+	UE_API virtual UEdGraphNode* PerformAction(UEdGraph* ParentGraph, UEdGraphPin* FromPin, const FVector2f& Location, bool bSelectNewNode = true) override;
 	//~ End FEdGraphSchemaAction Interface
 };
 
 /** Promotes an input to a graph input, using its respective literal value as the default value */
 USTRUCT()
-struct METASOUNDEDITOR_API FMetasoundGraphSchemaAction_PromoteToInput : public FMetasoundGraphSchemaAction_NodeWithMultipleOutputs
+struct FMetasoundGraphSchemaAction_PromoteToInput : public FMetasoundGraphSchemaAction_NodeWithMultipleOutputs
 {
 	GENERATED_USTRUCT_BODY();
 
-	FMetasoundGraphSchemaAction_PromoteToInput();
+	UE_API FMetasoundGraphSchemaAction_PromoteToInput();
 
 	virtual ~FMetasoundGraphSchemaAction_PromoteToInput() = default;
 
 	//~ Begin FEdGraphSchemaAction Interface
 	using FEdGraphSchemaAction::PerformAction; // Prevent hiding of deprecated base class function with FVector2D
-	virtual UEdGraphNode* PerformAction(UEdGraph* ParentGraph, UEdGraphPin* FromPin, const FVector2f& Location, bool bSelectNewNode = true) override;
+	UE_API virtual UEdGraphNode* PerformAction(UEdGraph* ParentGraph, UEdGraphPin* FromPin, const FVector2f& Location, bool bSelectNewNode = true) override;
 	//~ End FEdGraphSchemaAction Interface
 };
 
 /** Adds an output to the graph */
 USTRUCT()
-struct METASOUNDEDITOR_API FMetasoundGraphSchemaAction_NewOutput : public FMetasoundGraphSchemaAction
+struct FMetasoundGraphSchemaAction_NewOutput : public FMetasoundGraphSchemaAction
 {
 	GENERATED_USTRUCT_BODY();
 
@@ -204,41 +206,41 @@ struct METASOUNDEDITOR_API FMetasoundGraphSchemaAction_NewOutput : public FMetas
 		: FMetasoundGraphSchemaAction()
 	{}
 
-	FMetasoundGraphSchemaAction_NewOutput(FText InNodeCategory, FText InDisplayName, FGuid InOutputNodeID, FText InToolTip, Metasound::Editor::EPrimaryContextGroup InGroup);
+	UE_API FMetasoundGraphSchemaAction_NewOutput(FText InNodeCategory, FText InDisplayName, FGuid InOutputNodeID, FText InToolTip, Metasound::Editor::EPrimaryContextGroup InGroup);
 
 	virtual ~FMetasoundGraphSchemaAction_NewOutput() = default;
 
 	//~ Begin FMetasoundGraphSchemaAction Interface
-	virtual const FSlateBrush* GetIconBrush() const override;
+	UE_API virtual const FSlateBrush* GetIconBrush() const override;
 
-	virtual const FLinearColor& GetIconColor() const override;
+	UE_API virtual const FLinearColor& GetIconColor() const override;
 	//~ End FMetasoundGraphSchemaAction Interface
 
 	//~ Begin FEdGraphSchemaAction Interface
 	using FEdGraphSchemaAction::PerformAction; // Prevent hiding of deprecated base class function with FVector2D
-	virtual UEdGraphNode* PerformAction(UEdGraph* ParentGraph, UEdGraphPin* FromPin, const FVector2f& Location, bool bSelectNewNode = true) override;
+	UE_API virtual UEdGraphNode* PerformAction(UEdGraph* ParentGraph, UEdGraphPin* FromPin, const FVector2f& Location, bool bSelectNewNode = true) override;
 	//~ End FEdGraphSchemaAction Interface
 };
 
 /** Promotes a node output to a graph output */
 USTRUCT()
-struct METASOUNDEDITOR_API FMetasoundGraphSchemaAction_PromoteToOutput : public FMetasoundGraphSchemaAction
+struct FMetasoundGraphSchemaAction_PromoteToOutput : public FMetasoundGraphSchemaAction
 {
 	GENERATED_USTRUCT_BODY();
 
-	FMetasoundGraphSchemaAction_PromoteToOutput();
+	UE_API FMetasoundGraphSchemaAction_PromoteToOutput();
 
 	virtual ~FMetasoundGraphSchemaAction_PromoteToOutput() = default;
 
 	//~ Begin FEdGraphSchemaAction Interface
 	using FEdGraphSchemaAction::PerformAction; // Prevent hiding of deprecated base class function with FVector2D
-	virtual UEdGraphNode* PerformAction(UEdGraph* ParentGraph, UEdGraphPin* FromPin, const FVector2f& Location, bool bSelectNewNode = true) override;
+	UE_API virtual UEdGraphNode* PerformAction(UEdGraph* ParentGraph, UEdGraphPin* FromPin, const FVector2f& Location, bool bSelectNewNode = true) override;
 	//~ End FEdGraphSchemaAction Interface
 };
 
 /** Adds a variable node to the graph */
 USTRUCT()
-struct METASOUNDEDITOR_API FMetasoundGraphSchemaAction_NewVariableNode : public FMetasoundGraphSchemaAction_NodeWithMultipleOutputs
+struct FMetasoundGraphSchemaAction_NewVariableNode : public FMetasoundGraphSchemaAction_NodeWithMultipleOutputs
 {
 	GENERATED_USTRUCT_BODY();
 
@@ -249,19 +251,19 @@ struct METASOUNDEDITOR_API FMetasoundGraphSchemaAction_NewVariableNode : public
 		: FMetasoundGraphSchemaAction_NodeWithMultipleOutputs()
 	{}
 
-	FMetasoundGraphSchemaAction_NewVariableNode(FText InNodeCategory, FText InDisplayName, FGuid InVariableID, FText InToolTip);
+	UE_API FMetasoundGraphSchemaAction_NewVariableNode(FText InNodeCategory, FText InDisplayName, FGuid InVariableID, FText InToolTip);
 
 	virtual ~FMetasoundGraphSchemaAction_NewVariableNode() = default;
 
 	//~ Begin FEdGraphSchemaAction Interface
 	using FEdGraphSchemaAction::PerformAction; // Prevent hiding of deprecated base class function with FVector2D
-	virtual UEdGraphNode* PerformAction(UEdGraph* ParentGraph, UEdGraphPin* FromPin, const FVector2f& Location, bool bSelectNewNode = true) override;
+	UE_API virtual UEdGraphNode* PerformAction(UEdGraph* ParentGraph, UEdGraphPin* FromPin, const FVector2f& Location, bool bSelectNewNode = true) override;
 	//~ End FEdGraphSchemaAction Interface
 
 	//~ Begin FMetasoundGraphSchemaAction Interface
-	virtual const FSlateBrush* GetIconBrush() const override;
+	UE_API virtual const FSlateBrush* GetIconBrush() const override;
 
-	virtual const FLinearColor& GetIconColor() const override;
+	UE_API virtual const FLinearColor& GetIconColor() const override;
 	//~ End FMetasoundGraphSchemaAction Interface
 	
 protected:
@@ -279,7 +281,7 @@ protected:
 
 /** Adds a variable node to the graph */
 USTRUCT()
-struct METASOUNDEDITOR_API FMetasoundGraphSchemaAction_NewVariableAccessorNode : public FMetasoundGraphSchemaAction_NewVariableNode
+struct FMetasoundGraphSchemaAction_NewVariableAccessorNode : public FMetasoundGraphSchemaAction_NewVariableNode
 {
 	GENERATED_USTRUCT_BODY();
 
@@ -287,18 +289,18 @@ struct METASOUNDEDITOR_API FMetasoundGraphSchemaAction_NewVariableAccessorNode :
 		: FMetasoundGraphSchemaAction_NewVariableNode()
 	{}
 
-	FMetasoundGraphSchemaAction_NewVariableAccessorNode(FText InNodeCategory, FText InDisplayName, FGuid InVariableID, FText InToolTip);
+	UE_API FMetasoundGraphSchemaAction_NewVariableAccessorNode(FText InNodeCategory, FText InDisplayName, FGuid InVariableID, FText InToolTip);
 
 	virtual ~FMetasoundGraphSchemaAction_NewVariableAccessorNode() = default;
 
 protected:
-	virtual const FMetasoundFrontendNode* CreateFrontendVariableNode(FMetaSoundFrontendDocumentBuilder& DocBuilder) const override;
+	UE_API virtual const FMetasoundFrontendNode* CreateFrontendVariableNode(FMetaSoundFrontendDocumentBuilder& DocBuilder) const override;
 
 };
 
 /** Adds a variable node to the graph */
 USTRUCT()
-struct METASOUNDEDITOR_API FMetasoundGraphSchemaAction_NewVariableDeferredAccessorNode : public FMetasoundGraphSchemaAction_NewVariableNode
+struct FMetasoundGraphSchemaAction_NewVariableDeferredAccessorNode : public FMetasoundGraphSchemaAction_NewVariableNode
 {
 	GENERATED_USTRUCT_BODY();
 
@@ -308,15 +310,15 @@ struct METASOUNDEDITOR_API FMetasoundGraphSchemaAction_NewVariableDeferredAccess
 
 	virtual ~FMetasoundGraphSchemaAction_NewVariableDeferredAccessorNode() = default;
 
-	FMetasoundGraphSchemaAction_NewVariableDeferredAccessorNode(FText InNodeCategory, FText InDisplayName, FGuid InVariableID, FText InToolTip);
+	UE_API FMetasoundGraphSchemaAction_NewVariableDeferredAccessorNode(FText InNodeCategory, FText InDisplayName, FGuid InVariableID, FText InToolTip);
 	
 protected:
-	virtual const FMetasoundFrontendNode* CreateFrontendVariableNode(FMetaSoundFrontendDocumentBuilder& DocBuilder) const override;
+	UE_API virtual const FMetasoundFrontendNode* CreateFrontendVariableNode(FMetaSoundFrontendDocumentBuilder& DocBuilder) const override;
 };
 
 /** Adds a variable node to the graph */
 USTRUCT()
-struct METASOUNDEDITOR_API FMetasoundGraphSchemaAction_NewVariableMutatorNode : public FMetasoundGraphSchemaAction_NewVariableNode
+struct FMetasoundGraphSchemaAction_NewVariableMutatorNode : public FMetasoundGraphSchemaAction_NewVariableNode
 {
 	GENERATED_USTRUCT_BODY();
 
@@ -326,64 +328,64 @@ struct METASOUNDEDITOR_API FMetasoundGraphSchemaAction_NewVariableMutatorNode :
 
 	virtual ~FMetasoundGraphSchemaAction_NewVariableMutatorNode() = default;
 
-	FMetasoundGraphSchemaAction_NewVariableMutatorNode(FText InNodeCategory, FText InDisplayName, FGuid InVariableID, FText InToolTip);
+	UE_API FMetasoundGraphSchemaAction_NewVariableMutatorNode(FText InNodeCategory, FText InDisplayName, FGuid InVariableID, FText InToolTip);
 
 protected:
-	virtual const FMetasoundFrontendNode* CreateFrontendVariableNode(FMetaSoundFrontendDocumentBuilder& DocBuilder) const override;
+	UE_API virtual const FMetasoundFrontendNode* CreateFrontendVariableNode(FMetaSoundFrontendDocumentBuilder& DocBuilder) const override;
 };
 
 /** Promotes an input to a graph variable & respective getter node, using its respective literal value as the default value */
 USTRUCT()
-struct METASOUNDEDITOR_API FMetasoundGraphSchemaAction_PromoteToVariable_AccessorNode : public FMetasoundGraphSchemaAction_NodeWithMultipleOutputs
+struct FMetasoundGraphSchemaAction_PromoteToVariable_AccessorNode : public FMetasoundGraphSchemaAction_NodeWithMultipleOutputs
 {
 	GENERATED_USTRUCT_BODY();
 
-	FMetasoundGraphSchemaAction_PromoteToVariable_AccessorNode();
+	UE_API FMetasoundGraphSchemaAction_PromoteToVariable_AccessorNode();
 
 	virtual ~FMetasoundGraphSchemaAction_PromoteToVariable_AccessorNode() = default;
 
 	//~ Begin FEdGraphSchemaAction Interface
 	using FEdGraphSchemaAction::PerformAction; // Prevent hiding of deprecated base class function with FVector2D
-	virtual UEdGraphNode* PerformAction(UEdGraph* ParentGraph, UEdGraphPin* FromPin, const FVector2f& Location, bool bSelectNewNode = true) override;
+	UE_API virtual UEdGraphNode* PerformAction(UEdGraph* ParentGraph, UEdGraphPin* FromPin, const FVector2f& Location, bool bSelectNewNode = true) override;
 	//~ End FEdGraphSchemaAction Interface
 };
 
 /** Promotes an output to a graph variable & respective setter node, using its respective literal value as the default value */
 USTRUCT()
-struct METASOUNDEDITOR_API FMetasoundGraphSchemaAction_PromoteToVariable_MutatorNode : public FMetasoundGraphSchemaAction
+struct FMetasoundGraphSchemaAction_PromoteToVariable_MutatorNode : public FMetasoundGraphSchemaAction
 {
 	GENERATED_USTRUCT_BODY();
 
-	FMetasoundGraphSchemaAction_PromoteToVariable_MutatorNode();
+	UE_API FMetasoundGraphSchemaAction_PromoteToVariable_MutatorNode();
 
 	virtual ~FMetasoundGraphSchemaAction_PromoteToVariable_MutatorNode() = default;
 
 	//~ Begin FEdGraphSchemaAction Interface
 	using FEdGraphSchemaAction::PerformAction; // Prevent hiding of deprecated base class function with FVector2D
-	virtual UEdGraphNode* PerformAction(UEdGraph* ParentGraph, UEdGraphPin* FromPin, const FVector2f& Location, bool bSelectNewNode = true) override;
+	UE_API virtual UEdGraphNode* PerformAction(UEdGraph* ParentGraph, UEdGraphPin* FromPin, const FVector2f& Location, bool bSelectNewNode = true) override;
 	//~ End FEdGraphSchemaAction Interface
 };
 
 /** Promotes an input to a graph variable & respective deferred getter node, using its respective literal value as the default value */
 USTRUCT()
-struct METASOUNDEDITOR_API FMetasoundGraphSchemaAction_PromoteToVariable_DeferredAccessorNode : public FMetasoundGraphSchemaAction_NodeWithMultipleOutputs
+struct FMetasoundGraphSchemaAction_PromoteToVariable_DeferredAccessorNode : public FMetasoundGraphSchemaAction_NodeWithMultipleOutputs
 {
 	GENERATED_USTRUCT_BODY();
 
-	FMetasoundGraphSchemaAction_PromoteToVariable_DeferredAccessorNode();
+	UE_API FMetasoundGraphSchemaAction_PromoteToVariable_DeferredAccessorNode();
 
 	virtual ~FMetasoundGraphSchemaAction_PromoteToVariable_DeferredAccessorNode() = default;
 
 	//~ Begin FEdGraphSchemaAction Interface
 	using FEdGraphSchemaAction::PerformAction; // Prevent hiding of deprecated base class function with FVector2D
-	virtual UEdGraphNode* PerformAction(UEdGraph* ParentGraph, UEdGraphPin* FromPin, const FVector2f& Location, bool bSelectNewNode = true) override;
+	UE_API virtual UEdGraphNode* PerformAction(UEdGraph* ParentGraph, UEdGraphPin* FromPin, const FVector2f& Location, bool bSelectNewNode = true) override;
 	//~ End FEdGraphSchemaAction Interface
 };
 
 
 /** Action to add a node to the graph */
 USTRUCT()
-struct METASOUNDEDITOR_API FMetasoundGraphSchemaAction_NewNode : public FMetasoundGraphSchemaAction_NodeWithMultipleOutputs
+struct FMetasoundGraphSchemaAction_NewNode : public FMetasoundGraphSchemaAction_NodeWithMultipleOutputs
 {
 	GENERATED_USTRUCT_BODY();
 
@@ -400,20 +402,20 @@ struct METASOUNDEDITOR_API FMetasoundGraphSchemaAction_NewNode : public FMetasou
 	virtual ~FMetasoundGraphSchemaAction_NewNode() = default;
 
 	//~ Begin FMetasoundGraphSchemaAction Interface
-	virtual const FSlateBrush* GetIconBrush() const override;
+	UE_API virtual const FSlateBrush* GetIconBrush() const override;
 
-	virtual const FLinearColor& GetIconColor() const override;
+	UE_API virtual const FLinearColor& GetIconColor() const override;
 	//~ End FMetasoundGraphSchemaAction Interface
 
 	//~ Begin FEdGraphSchemaAction Interface
 	using FEdGraphSchemaAction::PerformAction; // Prevent hiding of deprecated base class function with FVector2D
-	virtual UEdGraphNode* PerformAction(UEdGraph* ParentGraph, UEdGraphPin* FromPin, const FVector2f& Location, bool bSelectNewNode = true) override;
+	UE_API virtual UEdGraphNode* PerformAction(UEdGraph* ParentGraph, UEdGraphPin* FromPin, const FVector2f& Location, bool bSelectNewNode = true) override;
 	//~ End FEdGraphSchemaAction Interface
 };
 
 /** Action to add nodes to the graph based on selected objects*/
 USTRUCT()
-struct METASOUNDEDITOR_API FMetasoundGraphSchemaAction_NewFromSelected : public FMetasoundGraphSchemaAction_NewNode
+struct FMetasoundGraphSchemaAction_NewFromSelected : public FMetasoundGraphSchemaAction_NewNode
 {
 	GENERATED_USTRUCT_BODY();
 
@@ -429,48 +431,48 @@ struct METASOUNDEDITOR_API FMetasoundGraphSchemaAction_NewFromSelected : public
 
 	//~ Begin FEdGraphSchemaAction Interface
 	using FEdGraphSchemaAction::PerformAction; // Prevent hiding of deprecated base class function with FVector2D
-	virtual UEdGraphNode* PerformAction(UEdGraph* ParentGraph, UEdGraphPin* FromPin, const FVector2f& Location, bool bSelectNewNode = true) override;
+	UE_API virtual UEdGraphNode* PerformAction(UEdGraph* ParentGraph, UEdGraphPin* FromPin, const FVector2f& Location, bool bSelectNewNode = true) override;
 	//~ End FEdGraphSchemaAction Interface
 };
 
 /** Action to create new Audio Analyzer node */
 USTRUCT()
-struct METASOUNDEDITOR_API FMetasoundGraphSchemaAction_NewAudioAnalyzer : public FMetasoundGraphSchemaAction
+struct FMetasoundGraphSchemaAction_NewAudioAnalyzer : public FMetasoundGraphSchemaAction
 {
 	GENERATED_USTRUCT_BODY();
 
-	FMetasoundGraphSchemaAction_NewAudioAnalyzer();
+	UE_API FMetasoundGraphSchemaAction_NewAudioAnalyzer();
 
 	//~ Begin FMetasoundGraphSchemaAction Interface
-	virtual const FLinearColor& GetIconColor() const override;
+	UE_API virtual const FLinearColor& GetIconColor() const override;
 	//~ End FMetasoundGraphSchemaAction Interface
 
 	//~ Begin FEdGraphSchemaAction Interface
 	using FEdGraphSchemaAction::PerformAction; // Prevent hiding of deprecated base class function with FVector2D
-	virtual UEdGraphNode* PerformAction(UEdGraph* ParentGraph, UEdGraphPin* FromPin, const FVector2f& Location, bool bSelectNewNode = true) override;
+	UE_API virtual UEdGraphNode* PerformAction(UEdGraph* ParentGraph, UEdGraphPin* FromPin, const FVector2f& Location, bool bSelectNewNode = true) override;
 	//~ End FEdGraphSchemaAction Interface
 };
 
 /** Action to create new reroute node */
 USTRUCT()
-struct METASOUNDEDITOR_API FMetasoundGraphSchemaAction_NewReroute : public FMetasoundGraphSchemaAction
+struct FMetasoundGraphSchemaAction_NewReroute : public FMetasoundGraphSchemaAction
 {
 	GENERATED_USTRUCT_BODY();
 
 	FMetasoundGraphSchemaAction_NewReroute() = default;
-	FMetasoundGraphSchemaAction_NewReroute(const FLinearColor* InIconColor, bool bInShouldTransact = true);
+	UE_API FMetasoundGraphSchemaAction_NewReroute(const FLinearColor* InIconColor, bool bInShouldTransact = true);
 
 	virtual ~FMetasoundGraphSchemaAction_NewReroute() = default;
 
 	//~ Begin FMetasoundGraphSchemaAction Interface
-	virtual const FSlateBrush* GetIconBrush() const override;
+	UE_API virtual const FSlateBrush* GetIconBrush() const override;
 
-	virtual const FLinearColor& GetIconColor() const override;
+	UE_API virtual const FLinearColor& GetIconColor() const override;
 	//~ End FMetasoundGraphSchemaAction Interface
 
 	//~ Begin FEdGraphSchemaAction Interface
 	using FEdGraphSchemaAction::PerformAction; // Prevent hiding of deprecated base class function with FVector2D
-	virtual UEdGraphNode* PerformAction(UEdGraph* ParentGraph, UEdGraphPin* FromPin, const FVector2f& Location, bool bSelectNewNode = true) override;
+	UE_API virtual UEdGraphNode* PerformAction(UEdGraph* ParentGraph, UEdGraphPin* FromPin, const FVector2f& Location, bool bSelectNewNode = true) override;
 	//~ End FEdGraphSchemaAction Interface
 
 private:
@@ -480,7 +482,7 @@ private:
 
 /** Action to create new comment */
 USTRUCT()
-struct METASOUNDEDITOR_API FMetasoundGraphSchemaAction_NewComment : public FMetasoundGraphSchemaAction
+struct FMetasoundGraphSchemaAction_NewComment : public FMetasoundGraphSchemaAction
 {
 	GENERATED_USTRUCT_BODY();
 
@@ -493,20 +495,20 @@ struct METASOUNDEDITOR_API FMetasoundGraphSchemaAction_NewComment : public FMeta
 	{}
 
 	//~ Begin FMetasoundGraphSchemaAction Interface
-	virtual const FSlateBrush* GetIconBrush() const override;
+	UE_API virtual const FSlateBrush* GetIconBrush() const override;
 
-	virtual const FLinearColor& GetIconColor() const override;
+	UE_API virtual const FLinearColor& GetIconColor() const override;
 	//~ End FMetasoundGraphSchemaAction Interface
 
 	//~ Begin FEdGraphSchemaAction Interface
 	using FEdGraphSchemaAction::PerformAction; // Prevent hiding of deprecated base class function with FVector2D
-	virtual UEdGraphNode* PerformAction(UEdGraph* ParentGraph, UEdGraphPin* FromPin, const FVector2f& Location, bool bSelectNewNode = true) override;
+	UE_API virtual UEdGraphNode* PerformAction(UEdGraph* ParentGraph, UEdGraphPin* FromPin, const FVector2f& Location, bool bSelectNewNode = true) override;
 	//~ End FEdGraphSchemaAction Interface
 };
 
 /** Action to paste clipboard contents into the graph */
 USTRUCT()
-struct METASOUNDEDITOR_API FMetasoundGraphSchemaAction_Paste : public FMetasoundGraphSchemaAction
+struct FMetasoundGraphSchemaAction_Paste : public FMetasoundGraphSchemaAction
 {
 	GENERATED_USTRUCT_BODY();
 
@@ -519,7 +521,7 @@ struct METASOUNDEDITOR_API FMetasoundGraphSchemaAction_Paste : public FMetasound
 	{}
 
 	//~ Begin FEdGraphSchemaAction Interface
-	virtual UEdGraphNode* PerformAction(UEdGraph* ParentGraph, UEdGraphPin* FromPin, const FVector2f& Location, bool bSelectNewNode = true) override;
+	UE_API virtual UEdGraphNode* PerformAction(UEdGraph* ParentGraph, UEdGraphPin* FromPin, const FVector2f& Location, bool bSelectNewNode = true) override;
 	//~ End FEdGraphSchemaAction Interface
 };
 
@@ -579,3 +581,5 @@ private:
 	/** Adds action for creating a comment */
 	void GetCommentAction(FGraphActionMenuBuilder& ActionMenuBuilder, const UEdGraph* CurrentGraph = nullptr) const;
 };
+
+#undef UE_API
diff --git a/Engine/Plugins/Runtime/Metasound/Source/MetasoundEditor/Public/MetasoundEditorModule.h b/Engine/Plugins/Runtime/Metasound/Source/MetasoundEditor/Public/MetasoundEditorModule.h
--- a/Engine/Plugins/Runtime/Metasound/Source/MetasoundEditor/Public/MetasoundEditorModule.h
+++ b/Engine/Plugins/Runtime/Metasound/Source/MetasoundEditor/Public/MetasoundEditorModule.h
@@ -14,6 +14,8 @@
 #include "Styling/SlateStyleRegistry.h"
 #include "Templates/Function.h"
 
+#define UE_API METASOUNDEDITOR_API
+
 // Forward Declarations
 struct FAudioMeterDefaultColorStyle;
 struct FAudioOscilloscopePanelStyle;
@@ -75,7 +77,7 @@ namespace Metasound
 		DECLARE_DELEGATE_RetVal_OneParam(TSharedRef<SWidget>, FOnCreateGraphNodeVisualizationWidget, const FCreateGraphNodeVisualizationWidgetParams&);
 
 
-		class METASOUNDEDITOR_API IMetasoundEditorModule : public IModuleInterface
+		class IMetasoundEditorModule : public IModuleInterface
 		{
 		public:
 			UE_DEPRECATED(5.3, "IsExplicitProxyClass is deprecated, use Metasound::Frontend::FDataTypeRegistryInfo::bIsExplicit")
@@ -91,10 +93,10 @@ namespace Metasound
 			virtual void PrimeAssetRegistryAsync() { }
 
 			UE_DEPRECATED(5.5, "Prime no longer supported (nor recommended as AssetTag updates now enable querying MetaSound assets for interface information without loading full asset.")
-			virtual EAssetPrimeStatus GetAssetRegistryPrimeStatus() const;
+			UE_API virtual EAssetPrimeStatus GetAssetRegistryPrimeStatus() const;
 
 			UE_DEPRECATED(5.5, "Use FMetaSoundAssetManager::IsInitialAssetScanComplete instead")
-			virtual EAssetScanStatus GetAssetRegistryScanStatus() const;
+			UE_API virtual EAssetScanStatus GetAssetRegistryScanStatus() const;
 
 			virtual TUniquePtr<FMetasoundDefaultLiteralCustomizationBase> CreateMemberDefaultLiteralCustomization(UClass& InClass, IDetailCategoryBuilder& DefaultCategoryBuilder) const = 0;
 
@@ -136,3 +138,5 @@ namespace Metasound
 		};
 	} // namespace Editor
 } // namespace Metasound
+
+#undef UE_API
diff --git a/Engine/Plugins/Runtime/Metasound/Source/MetasoundEditor/Public/MetasoundEditorSubsystem.h b/Engine/Plugins/Runtime/Metasound/Source/MetasoundEditor/Public/MetasoundEditorSubsystem.h
--- a/Engine/Plugins/Runtime/Metasound/Source/MetasoundEditor/Public/MetasoundEditorSubsystem.h
+++ b/Engine/Plugins/Runtime/Metasound/Source/MetasoundEditor/Public/MetasoundEditorSubsystem.h
@@ -10,6 +10,8 @@
 
 #include "MetasoundEditorSubsystem.generated.h"
 
+#define UE_API METASOUNDEDITOR_API
+
 
 // Forward Declarations
 class UMetaSoundBuilderBase;
@@ -21,8 +23,8 @@ struct FMetaSoundPageSettings;
 
 
 /** The subsystem in charge of editor MetaSound functionality */
-UCLASS(meta = (DisplayName = "MetaSound Editor Subsystem"))
-class METASOUNDEDITOR_API UMetaSoundEditorSubsystem : public UEditorSubsystem
+UCLASS(MinimalAPI, meta = (DisplayName = "MetaSound Editor Subsystem"))
+class UMetaSoundEditorSubsystem : public UEditorSubsystem
 {
 	GENERATED_BODY()
 
@@ -31,7 +33,7 @@ public:
 	// reference to given member (asserts that existing literal is of similar subclass provided).  If 
 	// it does not exist, or an optional template object is provided, metadata is generated then bound.
 	// Returns true if new literal metadata was generated, false if not. Asserts if bind failed.
-	bool BindMemberMetadata(
+	UE_API bool BindMemberMetadata(
 		FMetaSoundFrontendDocumentBuilder& Builder,
 		UMetasoundEditorGraphMember& InMember,
 		TSubclassOf<UMetasoundEditorGraphMemberDefaultLiteral> LiteralClass,
@@ -44,7 +46,7 @@ public:
 	// @param TemplateSoundWave - SoundWave settings such as attenuation, modulation, and sound class will be copied from the optional TemplateSoundWave.
 	// For preset builders, TemplateSoundWave will override the template values from the referenced asset.
 	UFUNCTION(BlueprintCallable, Category = "Audio|MetaSound|Builder|Editor", meta = (WorldContext = "Parent", ExpandEnumAsExecs = "OutResult"))
-	UPARAM(DisplayName = "MetaSound Asset") TScriptInterface<IMetaSoundDocumentInterface> BuildToAsset(
+	UE_API UPARAM(DisplayName = "MetaSound Asset") TScriptInterface<IMetaSoundDocumentInterface> BuildToAsset(
 		UPARAM(DisplayName = "Builder") UMetaSoundBuilderBase* InBuilder,
 		const FString& Author,
 		const FString& AssetName,
@@ -53,7 +55,7 @@ public:
 		UPARAM(DisplayName = "Template SoundWave") const USoundWave* TemplateSoundWave = nullptr);
 
 	// Creates new member metadata for a member of a given builder, copying data from the referenced asset in the case of preset inherited inputs
-	UMetasoundEditorGraphMemberDefaultLiteral* CreateMemberMetadata(
+	UE_API UMetasoundEditorGraphMemberDefaultLiteral* CreateMemberMetadata(
 		FMetaSoundFrontendDocumentBuilder& Builder,
 		FName InMemberName,
 		TSubclassOf<UMetasoundEditorGraphMemberDefaultLiteral> LiteralClass) const;
@@ -61,18 +63,18 @@ public:
 	// Returns a builder for the given MetaSound asset. Returns null if provided a transient MetaSound. For finding builders for transient
 	// MetaSounds, use the UMetaSoundBuilderSubsystem's API (FindPatchBuilder, FindSourceBuilder, FindBuilderByName etc.)
 	UFUNCTION(BlueprintCallable, Category = "Audio|MetaSound|Builder|Editor", meta = (DisplayName = "Find Or Begin Building MetaSound Asset", ExpandEnumAsExecs = "OutResult"))
-	UPARAM(DisplayName = "Builder") UMetaSoundBuilderBase* FindOrBeginBuilding(TScriptInterface<IMetaSoundDocumentInterface> MetaSound, EMetaSoundBuilderResult& OutResult) const;
+	UE_API UPARAM(DisplayName = "Builder") UMetaSoundBuilderBase* FindOrBeginBuilding(TScriptInterface<IMetaSoundDocumentInterface> MetaSound, EMetaSoundBuilderResult& OutResult) const;
 	
 	// Find graph input metadata (which includes editor only range information for floats) for a given input. If the metadata does not exist, create it. 
 	UFUNCTION(BlueprintCallable, Category = "Audio|MetaSound|Builder", meta = (ExpandEnumAsExecs = "OutResult"))
-	UPARAM(DisplayName = "Input Metadata") UMetaSoundFrontendMemberMetadata* FindOrCreateGraphInputMetadata(UPARAM(DisplayName = "Builder") UMetaSoundBuilderBase* InBuilder, FName InputName, EMetaSoundBuilderResult& OutResult);
+	UE_API UPARAM(DisplayName = "Input Metadata") UMetaSoundFrontendMemberMetadata* FindOrCreateGraphInputMetadata(UPARAM(DisplayName = "Builder") UMetaSoundBuilderBase* InBuilder, FName InputName, EMetaSoundBuilderResult& OutResult);
 
 	// Returns the corresponding literal class for a given type
-	TSubclassOf<UMetasoundEditorGraphMemberDefaultLiteral> GetLiteralClassForType(FName TypeName) const;
+	UE_API TSubclassOf<UMetasoundEditorGraphMemberDefaultLiteral> GetLiteralClassForType(FName TypeName) const;
 
 	// Sets the visual location to InLocation of a given node InNode of a given builder's document.
 	UFUNCTION(BlueprintCallable, Category = "Audio|MetaSound|Builder|Editor", meta = (ExpandEnumAsExecs = "OutResult"))
-	void SetNodeLocation(
+	UE_API void SetNodeLocation(
 		UPARAM(DisplayName = "Builder") UMetaSoundBuilderBase* InBuilder,
 		UPARAM(DisplayName = "Node Handle") const FMetaSoundNodeHandle& InNode,
 		UPARAM(DisplayName = "Location") const FVector2D& InLocation,
@@ -81,61 +83,63 @@ public:
 	// Initialize the UObject asset
 	// with an optional MetaSound to be referenced if the asset is a preset
 	// and optionally clearing the existing MetaSound document (for the case of duplicated assets) 
-	void InitAsset(UObject& InNewMetaSound, UObject* InReferencedMetaSound = nullptr, const bool bClearDocument = false);
+	UE_API void InitAsset(UObject& InNewMetaSound, UObject* InReferencedMetaSound = nullptr, const bool bClearDocument = false);
 
 	UE_DEPRECATED(5.5, "EdGraph is now transiently generated and privately managed for asset editor use only.")
-	void InitEdGraph(UObject& InMetaSound);
+	UE_API void InitEdGraph(UObject& InMetaSound);
 
 	// Returns whether or not a page with the given name both exists and is set as
 	// a valid, cooked target for the currently set audition platform in editor.
-	bool IsPageAuditionPlatformCookTarget(FName InPageName) const;
-	bool IsPageAuditionPlatformCookTarget(const FGuid& InPageID) const;
+	UE_API bool IsPageAuditionPlatformCookTarget(FName InPageName) const;
+	UE_API bool IsPageAuditionPlatformCookTarget(const FGuid& InPageID) const;
 
 	// Add a builder listener for a builder which is used to add and remove custom editor builder delegates.
 	UFUNCTION(BlueprintCallable, Category = "Audio|MetaSound|Builder|Editor", meta = (ExpandEnumAsExecs = "OutResult"))
-	UPARAM(DisplayName = "Builder Listener") UMetaSoundEditorBuilderListener* AddBuilderDelegateListener(UMetaSoundBuilderBase* InBuilder, EMetaSoundBuilderResult& OutResult);
+	UE_API UPARAM(DisplayName = "Builder Listener") UMetaSoundEditorBuilderListener* AddBuilderDelegateListener(UMetaSoundBuilderBase* InBuilder, EMetaSoundBuilderResult& OutResult);
 	
 	// Wraps RegisterGraphWithFrontend logic in Frontend with any additional logic required to refresh editor & respective editor object state.
 	// @param InMetaSound - MetaSound to register
 	// @param bInForceSynchronize - Forces the synchronize flag for all open graphs being registered by this call (all referenced graphs and
 	// referencing graphs open in editors)
-	void RegisterGraphWithFrontend(UObject& InMetaSound, bool bInForceViewSynchronization = false) const;
+	UE_API void RegisterGraphWithFrontend(UObject& InMetaSound, bool bInForceViewSynchronization = false) const;
 
 	// Register toolbar extender that will be displayed in the MetaSound Asset Editor.
-	void RegisterToolbarExtender(TSharedRef<FExtender> InExtender);
+	UE_API void RegisterToolbarExtender(TSharedRef<FExtender> InExtender);
 
 	// If the given page name is implemented on the provided builder, sets the focused page of
 	// the provided builder to the associated page and sets the audition page to
 	// the provided name. If the given builder has an asset editor open, optionally opens or brings
 	// that editor's associated PageID into user focus.
 	UFUNCTION(BlueprintCallable, Category = "Audio|MetaSound|Builder|Editor", meta = (ExpandEnumAsExecs = "OutResult"))
-	void SetFocusedPage(UMetaSoundBuilderBase* Builder, FName PageName, bool bOpenEditor, EMetaSoundBuilderResult& OutResult) const;
+	UE_API void SetFocusedPage(UMetaSoundBuilderBase* Builder, FName PageName, bool bOpenEditor, EMetaSoundBuilderResult& OutResult) const;
 
 	// If the given PageID is implemented on the provided builder, sets the focused page of
 	// the provided builder to the associated page and sets the audition target page to
 	// the provided ID. If the given builder has an asset editor open, optionally opens or brings
 	// that editor's associated PageID into user focus. Returns whether or not the audition page
 	// was set to the provided focus page.
-	bool SetFocusedPage(UMetaSoundBuilderBase& Builder, const FGuid& InPageID, bool bOpenEditor, bool bPostTransaction = true) const;
+	UE_API bool SetFocusedPage(UMetaSoundBuilderBase& Builder, const FGuid& InPageID, bool bOpenEditor, bool bPostTransaction = true) const;
 
 	// Unregisters toolbar extender that is displayed in the MetaSound Asset Editor.
-	bool UnregisterToolbarExtender(TSharedRef<FExtender> InExtender);
+	UE_API bool UnregisterToolbarExtender(TSharedRef<FExtender> InExtender);
 
 	// Get the default author for a MetaSound asset
-	const FString GetDefaultAuthor();
+	UE_API const FString GetDefaultAuthor();
 
 	// Returns all currently toolbar extenders registered to be displayed within the MetaSound Asset Editor.
-	const TArray<TSharedRef<FExtender>>& GetToolbarExtenders() const;
+	UE_API const TArray<TSharedRef<FExtender>>& GetToolbarExtenders() const;
 
-	static UMetaSoundEditorSubsystem& GetChecked();
-	static const UMetaSoundEditorSubsystem& GetConstChecked();
+	static UE_API UMetaSoundEditorSubsystem& GetChecked();
+	static UE_API const UMetaSoundEditorSubsystem& GetConstChecked();
 
 private:
-	bool SetFocusedPageInternal(FName PageName, const FGuid& InPageID, UMetaSoundBuilderBase& Builder, bool bOpenEditor, bool bPostTransaction) const;
+	UE_API bool SetFocusedPageInternal(FName PageName, const FGuid& InPageID, UMetaSoundBuilderBase& Builder, bool bOpenEditor, bool bPostTransaction) const;
 
 	// Copy over sound wave settings such as attenuation, modulation, and sound class from the template sound wave to the MetaSound
-	void SetSoundWaveSettingsFromTemplate(USoundWave& NewMetasound, const USoundWave& TemplateSoundWave) const;
+	UE_API void SetSoundWaveSettingsFromTemplate(USoundWave& NewMetasound, const USoundWave& TemplateSoundWave) const;
 
 	// Editor Toolbar Extenders
 	TArray<TSharedRef<FExtender>> EditorToolbarExtenders;
 };
+
+#undef UE_API
diff --git a/Engine/Plugins/Runtime/Metasound/Source/MetasoundEditor/Public/MetasoundNodeConfigurationCustomization.h b/Engine/Plugins/Runtime/Metasound/Source/MetasoundEditor/Public/MetasoundNodeConfigurationCustomization.h
--- a/Engine/Plugins/Runtime/Metasound/Source/MetasoundEditor/Public/MetasoundNodeConfigurationCustomization.h
+++ b/Engine/Plugins/Runtime/Metasound/Source/MetasoundEditor/Public/MetasoundNodeConfigurationCustomization.h
@@ -3,6 +3,8 @@
 
 #include "InstancedStructDetails.h"
 
+#define UE_API METASOUNDEDITOR_API
+
 // Forward Declarations
 struct FPropertyChangedEvent;
 class FInstancedStructDataDetails;
@@ -12,7 +14,7 @@ class UMetasoundEditorGraphNode;
 
 namespace Metasound::Editor
 {
-	class METASOUNDEDITOR_API FMetaSoundNodeConfigurationDataDetails : public FInstancedStructDataDetails
+	class FMetaSoundNodeConfigurationDataDetails : public FInstancedStructDataDetails
 	{
 	public:
 		FMetaSoundNodeConfigurationDataDetails(TSharedPtr<IPropertyHandle> InStructProperty, TWeakObjectPtr<UMetasoundEditorGraphNode> InNode)
@@ -21,11 +23,13 @@ namespace Metasound::Editor
 		{
 		}
 
-		virtual void OnChildRowAdded(IDetailPropertyRow& ChildRow) override;
+		UE_API virtual void OnChildRowAdded(IDetailPropertyRow& ChildRow) override;
 
 	protected:
-		void OnChildPropertyChanged(const FPropertyChangedEvent& InPropertyChangedEvent);
+		UE_API void OnChildPropertyChanged(const FPropertyChangedEvent& InPropertyChangedEvent);
 
 		TWeakObjectPtr<UMetasoundEditorGraphNode> GraphNode;
 	};
 } // namespace Metasound::Editor
+
+#undef UE_API
diff --git a/Engine/Plugins/Runtime/Metasound/Source/MetasoundEngineTest/Public/NodeTestGraphBuilder.h b/Engine/Plugins/Runtime/Metasound/Source/MetasoundEngineTest/Public/NodeTestGraphBuilder.h
--- a/Engine/Plugins/Runtime/Metasound/Source/MetasoundEngineTest/Public/NodeTestGraphBuilder.h
+++ b/Engine/Plugins/Runtime/Metasound/Source/MetasoundEngineTest/Public/NodeTestGraphBuilder.h
@@ -7,13 +7,15 @@
 #include "MetasoundGenerator.h"
 #include "MetasoundNodeInterface.h"
 
+#define UE_API METASOUNDENGINETEST_API
+
 namespace Metasound::Test
 {
 	/** Helper to make testing nodes simpler. */
-	class METASOUNDENGINETEST_API FNodeTestGraphBuilder
+	class FNodeTestGraphBuilder
 	{
 	public:
-		FNodeTestGraphBuilder();
+		UE_API FNodeTestGraphBuilder();
 
 		/** 
 		* Add a node to the graph 
@@ -22,7 +24,7 @@ namespace Metasound::Test
 		* @param MajorVersion	Major version of the node to add
 		* @return				NodeHandle of the added node
 		*/
-		Frontend::FNodeHandle AddNode(const FNodeClassName& ClassName, int32 MajorVersion) const;
+		UE_API Frontend::FNodeHandle AddNode(const FNodeClassName& ClassName, int32 MajorVersion) const;
 
 		/**
 		* Add an input node from the graph by name
@@ -30,7 +32,7 @@ namespace Metasound::Test
 		* @param Name			The name of the input node to get
 		* @return				NodeHandle, which can be invalid if an input with the given name doesn't exist in the graph
 		*/
-		Frontend::FNodeHandle GetInput(const FName& Name) const;
+		UE_API Frontend::FNodeHandle GetInput(const FName& Name) const;
 
 		/**
 		* Add an output node from the graph by name
@@ -38,7 +40,7 @@ namespace Metasound::Test
 		* @param Name			The name of the output node to get
 		* @return				NodeHandle, which can be invalid if an input with the given name doesn't exist in the graph
 		*/
-		Frontend::FNodeHandle GetOutput(const FName& Name) const;
+		UE_API Frontend::FNodeHandle GetOutput(const FName& Name) const;
 
 		/**
 		* Add an input node from the graph by name, type, and access type
@@ -48,7 +50,7 @@ namespace Metasound::Test
 		* @param AccessType		Whether the node should be a Value, Reference, or Unset type
 		* @return				NodeHandle of the added node, can be invalid if failed to add
 		*/
-		Frontend::FNodeHandle AddInput(
+		UE_API Frontend::FNodeHandle AddInput(
 			const FName& InputName,
 			const FName& TypeName,
 			EMetasoundFrontendVertexAccessType AccessType = EMetasoundFrontendVertexAccessType::Reference) const;
@@ -81,7 +83,7 @@ namespace Metasound::Test
 		* @param TypeName		The type for the node to reference
 		* @return				NodeHandle of the added node, can be invalid if failed to add
 		*/
-		Frontend::FNodeHandle AddOutput(const FName& OutputName, const FName& TypeName);
+		UE_API Frontend::FNodeHandle AddOutput(const FName& OutputName, const FName& TypeName);
 
 		/** 
 		* Adds a constructor input node to the graph and connects it as an input to the NodeToConnect
@@ -144,7 +146,7 @@ namespace Metasound::Test
 		* @param NodeHandle		The node to connect the input reference nodes to
 		* @return				whether all the nodes were successfully created and connected
 		*/
-		bool AddAndConnectDataReferenceInputs(const Frontend::FNodeHandle& NodeHandle);
+		UE_API bool AddAndConnectDataReferenceInputs(const Frontend::FNodeHandle& NodeHandle);
 
 		/** 
 		* Adds a data reference input node with InputName and connects it to the node's input with the same name
@@ -155,7 +157,7 @@ namespace Metasound::Test
 		* @param NodeName		Optional name to use for the created data reference input node
 		* @return				NodeHandle of the added node, can be invalid if failed to add
 		*/
-		Frontend::FNodeHandle AddAndConnectDataReferenceInput(const Frontend::FNodeHandle& NodeToConnect, const FName& InputName, const FName& TypeName, const FName& NodeName = NAME_None) const;
+		UE_API Frontend::FNodeHandle AddAndConnectDataReferenceInput(const Frontend::FNodeHandle& NodeToConnect, const FName& InputName, const FName& TypeName, const FName& NodeName = NAME_None) const;
 
 		/** Adds a data reference output node with OutputName and connects it to the node's output with the same name
 		* 	
@@ -165,7 +167,7 @@ namespace Metasound::Test
 		* @param NodeName		Optoinal name to use for the created data reference output node
 		* @return				NodeHandle of the added node, can be invalid if failed to add
 		*/
-		Frontend::FNodeHandle AddAndConnectDataReferenceOutput(const Frontend::FNodeHandle& NodeToConnect, const FName& OutputName, const FName& TypeName, const FName& NodeName = NAME_None);
+		UE_API Frontend::FNodeHandle AddAndConnectDataReferenceOutput(const Frontend::FNodeHandle& NodeToConnect, const FName& OutputName, const FName& TypeName, const FName& NodeName = NAME_None);
 
 		/** 
 		* Call last to create a generator from the graph constructed using the other method. 
@@ -174,7 +176,7 @@ namespace Metasound::Test
 		* @param SamplesPerBlock	The number of samples to render per block
 		* @return					A unique pointer to the Metasound Generator created
 		*/
-		TUniquePtr<FMetasoundGenerator> BuildGenerator(FSampleRate SampleRate = 48000, int32 SamplesPerBlock = 256) const;
+		UE_API TUniquePtr<FMetasoundGenerator> BuildGenerator(FSampleRate SampleRate = 48000, int32 SamplesPerBlock = 256) const;
 
 		/** 
 		* Helper that will add a single node, wire up the node's inputs and outputs, and hand back the graph's operator 
@@ -185,7 +187,7 @@ namespace Metasound::Test
 		* @param SamplesPerBlock	The number of samples to render per block
 		* @return					A unique pointer to the Metasound Generator created
 		*/
-		static TUniquePtr<FMetasoundGenerator> MakeSingleNodeGraph(
+		static UE_API TUniquePtr<FMetasoundGenerator> MakeSingleNodeGraph(
 			const FNodeClassName& ClassName,
 			int32 MajorVersion,
 			FSampleRate SampleRate = 48000,
@@ -200,7 +202,7 @@ namespace Metasound::Test
 		* @param InputName		The name of the input pin to connect to the right node
 		* @bool					true if able to connect the nodes (ie. all nodes were valid, all input/outputs were valid)
 		*/
-		static bool ConnectNodes(const Frontend::FNodeHandle& LeftNode, const FName& OutputName, const Frontend::FNodeHandle& RightNode, const FName& InputName);
+		static UE_API bool ConnectNodes(const Frontend::FNodeHandle& LeftNode, const FName& OutputName, const Frontend::FNodeHandle& RightNode, const FName& InputName);
 
 		/**
 		* Connects an Input (Left) Node's output with name OutputName to another (Right) node's input with name InputName
@@ -210,7 +212,7 @@ namespace Metasound::Test
 		* @param InputOutputName	The name of the input pin to connect to the right node
 		* @bool						true if able to connect the nodes (ie. all nodes were valid, all input/outputs were valid)
 		*/
-		static bool ConnectNodes(const Frontend::FNodeHandle& LeftNode, const Frontend::FNodeHandle& RightNode, const FName& InputOutputName);
+		static UE_API bool ConnectNodes(const Frontend::FNodeHandle& LeftNode, const Frontend::FNodeHandle& RightNode, const FName& InputOutputName);
 
 	private:
 		FMetasoundFrontendDocument Document;
@@ -219,3 +221,5 @@ namespace Metasound::Test
 		TArray<FVertexName> AudioOutputNames;
 	};
 }
+
+#undef UE_API
diff --git a/Engine/Plugins/Runtime/ModelViewViewModel/Source/ModelViewViewModelBlueprint/Public/Bindings/MVVMCompiledBindingLibraryCompiler.h b/Engine/Plugins/Runtime/ModelViewViewModel/Source/ModelViewViewModelBlueprint/Public/Bindings/MVVMCompiledBindingLibraryCompiler.h
--- a/Engine/Plugins/Runtime/ModelViewViewModel/Source/ModelViewViewModelBlueprint/Public/Bindings/MVVMCompiledBindingLibraryCompiler.h
+++ b/Engine/Plugins/Runtime/ModelViewViewModel/Source/ModelViewViewModelBlueprint/Public/Bindings/MVVMCompiledBindingLibraryCompiler.h
@@ -4,6 +4,8 @@
 
 #include "Bindings/MVVMCompiledBindingLibrary.h"
 
+#define UE_API MODELVIEWVIEWMODELBLUEPRINT_API
+
 namespace UE::MVVM { struct FMVVMConstFieldVariant; }
 template <typename T> class TSubclassOf;
 
@@ -18,7 +20,7 @@ namespace UE::MVVM
 {
 
 /**  */
-class MODELVIEWVIEWMODELBLUEPRINT_API FCompiledBindingLibraryCompiler
+class FCompiledBindingLibraryCompiler
 {
 public:
 	/** */
@@ -148,32 +150,32 @@ public:
 	};
 
 public:
-	FCompiledBindingLibraryCompiler(UBlueprint* Context);
+	UE_API FCompiledBindingLibraryCompiler(UBlueprint* Context);
 
 public:
 	/** */
-	TValueOrError<FFieldIdHandle, FText> AddFieldId(const UClass* SourceClass, FName FieldName);
+	UE_API TValueOrError<FFieldIdHandle, FText> AddFieldId(const UClass* SourceClass, FName FieldName);
 	
 	/** */
-	TValueOrError<FFieldPathHandle, FText> AddFieldPath(TArrayView<const UE::MVVM::FMVVMConstFieldVariant> FieldPath, bool bRead);
+	UE_API TValueOrError<FFieldPathHandle, FText> AddFieldPath(TArrayView<const UE::MVVM::FMVVMConstFieldVariant> FieldPath, bool bRead);
 
 	/** */
-	TValueOrError<FFieldPathHandle, FText> AddObjectFieldPath(TArrayView<const UE::MVVM::FMVVMConstFieldVariant> FieldPath, const UClass* ExpectedType, bool bRead);
+	UE_API TValueOrError<FFieldPathHandle, FText> AddObjectFieldPath(TArrayView<const UE::MVVM::FMVVMConstFieldVariant> FieldPath, const UClass* ExpectedType, bool bRead);
 
 	/** */
-	TValueOrError<FFieldPathHandle, FText> AddConversionFunctionFieldPath(const UClass* SourceClass, const UFunction* Function);
+	UE_API TValueOrError<FFieldPathHandle, FText> AddConversionFunctionFieldPath(const UClass* SourceClass, const UFunction* Function);
 
 	/** */
-	TValueOrError<FFieldPathHandle, FText> AddDelegateSignatureFieldPath(const UClass* SourceClass, const UFunction* Function);
+	UE_API TValueOrError<FFieldPathHandle, FText> AddDelegateSignatureFieldPath(const UClass* SourceClass, const UFunction* Function);
 
 	/** */
-	TValueOrError<FBindingHandle, FText> AddBinding(FFieldPathHandle Source, FFieldPathHandle Destination);
+	UE_API TValueOrError<FBindingHandle, FText> AddBinding(FFieldPathHandle Source, FFieldPathHandle Destination);
 
 	/** */
-	TValueOrError<FBindingHandle, FText> AddBinding(FFieldPathHandle Source, FFieldPathHandle Destination, FFieldPathHandle ConversionFunction);
+	UE_API TValueOrError<FBindingHandle, FText> AddBinding(FFieldPathHandle Source, FFieldPathHandle Destination, FFieldPathHandle ConversionFunction);
 
 	/** */
-	TValueOrError<FBindingHandle, FText> AddComplexBinding(FFieldPathHandle Destination, FFieldPathHandle ConversionFunction);
+	UE_API TValueOrError<FBindingHandle, FText> AddComplexBinding(FFieldPathHandle Destination, FFieldPathHandle ConversionFunction);
 
 	struct FCompileResult
 	{
@@ -185,15 +187,17 @@ public:
 	};
 
 	/** */
-	TValueOrError<FCompileResult, FText> Compile(FGuid LibraryId);
+	UE_API TValueOrError<FCompileResult, FText> Compile(FGuid LibraryId);
 
 private:
 	/** */
-	TValueOrError<FFieldPathHandle, FText> AddFieldPathImpl(TArrayView<const UE::MVVM::FMVVMConstFieldVariant> FieldPath, bool bRead);
-	TValueOrError<FBindingHandle, FText> AddBindingImpl(FFieldPathHandle Source, FFieldPathHandle Destination, FFieldPathHandle ConversionFunction, bool bIsComplexBinding);
-	TValueOrError<FFieldPathHandle, FText> AddFunctionFieldPathImpl(const UClass* InSourceClass, const UFunction* InFunction, TFunctionRef<FText(const UFunction*)> ValidateFunctionCallback);
+	UE_API TValueOrError<FFieldPathHandle, FText> AddFieldPathImpl(TArrayView<const UE::MVVM::FMVVMConstFieldVariant> FieldPath, bool bRead);
+	UE_API TValueOrError<FBindingHandle, FText> AddBindingImpl(FFieldPathHandle Source, FFieldPathHandle Destination, FFieldPathHandle ConversionFunction, bool bIsComplexBinding);
+	UE_API TValueOrError<FFieldPathHandle, FText> AddFunctionFieldPathImpl(const UClass* InSourceClass, const UFunction* InFunction, TFunctionRef<FText(const UFunction*)> ValidateFunctionCallback);
 
 	TPimplPtr<Private::FCompiledBindingLibraryCompilerImpl> Impl;
 };
 
 } //namespace UE::MVVM
+
+#undef UE_API
diff --git a/Engine/Plugins/Runtime/ModelViewViewModel/Source/ModelViewViewModelBlueprint/Public/Bindings/MVVMConversionFunctionHelper.h b/Engine/Plugins/Runtime/ModelViewViewModel/Source/ModelViewViewModelBlueprint/Public/Bindings/MVVMConversionFunctionHelper.h
--- a/Engine/Plugins/Runtime/ModelViewViewModel/Source/ModelViewViewModelBlueprint/Public/Bindings/MVVMConversionFunctionHelper.h
+++ b/Engine/Plugins/Runtime/ModelViewViewModel/Source/ModelViewViewModelBlueprint/Public/Bindings/MVVMConversionFunctionHelper.h
@@ -47,7 +47,7 @@ namespace UE::MVVM::ConversionFunctionHelper
 	 */
 	MODELVIEWVIEWMODELBLUEPRINT_API TValueOrError<void, FText> CanCreateSetterGraph(UBlueprint* WidgetBlueprint, const FMVVMBlueprintPropertyPath& PropertyPath);
 
-	struct MODELVIEWVIEWMODELBLUEPRINT_API FCreateGraphResult
+	struct FCreateGraphResult
 	{
 		/** The new graph created. */
 		UEdGraph* NewGraph = nullptr;
@@ -59,7 +59,7 @@ namespace UE::MVVM::ConversionFunctionHelper
 		bool bIsUbergraphPage = false;
 	};
 
-	struct MODELVIEWVIEWMODELBLUEPRINT_API FCreateGraphParams
+	struct FCreateGraphParams
 	{
 		bool bIsConst = false;
 		bool bTransient = false;
diff --git a/Engine/Plugins/Runtime/ModelViewViewModel/Source/ModelViewViewModelBlueprint/Public/Extensions/MVVMViewBlueprintListViewBaseExtension.h b/Engine/Plugins/Runtime/ModelViewViewModel/Source/ModelViewViewModelBlueprint/Public/Extensions/MVVMViewBlueprintListViewBaseExtension.h
--- a/Engine/Plugins/Runtime/ModelViewViewModel/Source/ModelViewViewModelBlueprint/Public/Extensions/MVVMViewBlueprintListViewBaseExtension.h
+++ b/Engine/Plugins/Runtime/ModelViewViewModel/Source/ModelViewViewModelBlueprint/Public/Extensions/MVVMViewBlueprintListViewBaseExtension.h
@@ -6,6 +6,8 @@
 #include "MVVMBlueprintViewExtension.h"
 #include "MVVMViewBlueprintListViewBaseExtension.generated.h"
 
+#define UE_API MODELVIEWVIEWMODELBLUEPRINT_API
+
 class UWidgetBlueprintGeneratedClass;
 class UMVVMBlueprintView;
 class UMVVMViewClass;
@@ -22,16 +24,16 @@ namespace UE::MVVM::Compiler
 	class IMVVMBlueprintViewCompile;
 }
 
-UCLASS()
-class MODELVIEWVIEWMODELBLUEPRINT_API UMVVMBlueprintViewExtension_ListViewBase : public UMVVMBlueprintViewExtension
+UCLASS(MinimalAPI)
+class UMVVMBlueprintViewExtension_ListViewBase : public UMVVMBlueprintViewExtension
 {
 	GENERATED_BODY()
 
 public:
 	//~ Begin UMVVMBlueprintViewExtension overrides
-	virtual void Precompile(UE::MVVM::Compiler::IMVVMBlueprintViewPrecompile* Compiler, UWidgetBlueprintGeneratedClass* Class) override;
-	virtual void Compile(UE::MVVM::Compiler::IMVVMBlueprintViewCompile* Compiler, UWidgetBlueprintGeneratedClass* Class, UMVVMViewClass* ViewExtension) override;
-	virtual bool WidgetRenamed(FName OldName, FName NewName) override;
+	UE_API virtual void Precompile(UE::MVVM::Compiler::IMVVMBlueprintViewPrecompile* Compiler, UWidgetBlueprintGeneratedClass* Class) override;
+	UE_API virtual void Compile(UE::MVVM::Compiler::IMVVMBlueprintViewCompile* Compiler, UWidgetBlueprintGeneratedClass* Class, UMVVMViewClass* ViewExtension) override;
+	UE_API virtual bool WidgetRenamed(FName OldName, FName NewName) override;
 	//~ End UMVVMBlueprintViewExtension overrides
 
 	FGuid GetEntryViewModelId() const
@@ -40,7 +42,7 @@ public:
 	}
 
 private:
-	const UMVVMBlueprintView* GetEntryWidgetBlueprintView(const UUserWidget* EntryUserWidget) const;
+	UE_API const UMVVMBlueprintView* GetEntryWidgetBlueprintView(const UUserWidget* EntryUserWidget) const;
 
 private:
 	UPROPERTY()
@@ -53,3 +55,5 @@ private:
 
 	friend UE::MVVM::FMVVMListViewBaseExtensionCustomizationExtender;
 };
+
+#undef UE_API
diff --git a/Engine/Plugins/Runtime/ModelViewViewModel/Source/ModelViewViewModelBlueprint/Public/Extensions/MVVMViewBlueprintPanelWidgetExtension.h b/Engine/Plugins/Runtime/ModelViewViewModel/Source/ModelViewViewModelBlueprint/Public/Extensions/MVVMViewBlueprintPanelWidgetExtension.h
--- a/Engine/Plugins/Runtime/ModelViewViewModel/Source/ModelViewViewModelBlueprint/Public/Extensions/MVVMViewBlueprintPanelWidgetExtension.h
+++ b/Engine/Plugins/Runtime/ModelViewViewModel/Source/ModelViewViewModelBlueprint/Public/Extensions/MVVMViewBlueprintPanelWidgetExtension.h
@@ -8,6 +8,8 @@
 #include "Components/PanelSlot.h"
 #include "MVVMViewBlueprintPanelWidgetExtension.generated.h"
 
+#define UE_API MODELVIEWVIEWMODELBLUEPRINT_API
+
 class UMVVMBlueprintView;
 class UMVVMViewClass;
 class UMVVMViewPanelWidgetClassExtension;
@@ -26,19 +28,19 @@ namespace UE::MVVM::Compiler
 	struct FBlueprintViewUserWidgetProperty;
 }
 
-UCLASS()
-class MODELVIEWVIEWMODELBLUEPRINT_API UMVVMBlueprintViewExtension_PanelWidget : public UMVVMBlueprintViewExtension
+UCLASS(MinimalAPI)
+class UMVVMBlueprintViewExtension_PanelWidget : public UMVVMBlueprintViewExtension
 {
 	GENERATED_BODY()
 
 public:
 	//~ Begin UMVVMBlueprintViewExtension overrides
-	virtual TArray<UE::MVVM::Compiler::FBlueprintViewUserWidgetProperty> AddProperties() override;
-	virtual TArray<UE::MVVM::Compiler::FBlueprintViewUserWidgetWidgetProperty> AddWidgetProperties() override;
-	virtual void Precompile(UE::MVVM::Compiler::IMVVMBlueprintViewPrecompile* Compiler, UWidgetBlueprintGeneratedClass* Class) override;
-	virtual void Compile(UE::MVVM::Compiler::IMVVMBlueprintViewCompile* Compiler, UWidgetBlueprintGeneratedClass* Class, UMVVMViewClass* ViewExtension) override;
-	virtual bool WidgetRenamed(FName OldName, FName NewName) override;
-	virtual void OnPreviewContentChanged(TSharedRef<SWidget> NewContent) override;
+	UE_API virtual TArray<UE::MVVM::Compiler::FBlueprintViewUserWidgetProperty> AddProperties() override;
+	UE_API virtual TArray<UE::MVVM::Compiler::FBlueprintViewUserWidgetWidgetProperty> AddWidgetProperties() override;
+	UE_API virtual void Precompile(UE::MVVM::Compiler::IMVVMBlueprintViewPrecompile* Compiler, UWidgetBlueprintGeneratedClass* Class) override;
+	UE_API virtual void Compile(UE::MVVM::Compiler::IMVVMBlueprintViewCompile* Compiler, UWidgetBlueprintGeneratedClass* Class, UMVVMViewClass* ViewExtension) override;
+	UE_API virtual bool WidgetRenamed(FName OldName, FName NewName) override;
+	UE_API virtual void OnPreviewContentChanged(TSharedRef<SWidget> NewContent) override;
 	//~ End UMVVMBlueprintViewExtension overrides
 
 	FGuid GetEntryViewModelId() const
@@ -47,9 +49,9 @@ public:
 	}
 
 private:
-	const UMVVMBlueprintView* GetEntryWidgetBlueprintView(const UUserWidget* EntryUserWidget) const;
+	UE_API const UMVVMBlueprintView* GetEntryWidgetBlueprintView(const UUserWidget* EntryUserWidget) const;
 
-	static void RefreshDesignerPreviewEntries(UPanelWidget* PanelWidget, TSubclassOf<UUserWidget> EntryWidgetClass, UPanelSlot* SlotTemplate, int32 NumDesignerPreviewEntries, bool bFullRebuild);
+	static UE_API void RefreshDesignerPreviewEntries(UPanelWidget* PanelWidget, TSubclassOf<UUserWidget> EntryWidgetClass, UPanelSlot* SlotTemplate, int32 NumDesignerPreviewEntries, bool bFullRebuild);
 
 private:
 	UPROPERTY()
@@ -77,3 +79,5 @@ private:
 
 	friend UE::MVVM::FMVVMPanelWidgetExtensionCustomizationExtender;
 };
+
+#undef UE_API
diff --git a/Engine/Plugins/Runtime/ModelViewViewModel/Source/ModelViewViewModelBlueprint/Public/MVVMBlueprintInstancedViewModel.h b/Engine/Plugins/Runtime/ModelViewViewModel/Source/ModelViewViewModelBlueprint/Public/MVVMBlueprintInstancedViewModel.h
--- a/Engine/Plugins/Runtime/ModelViewViewModel/Source/ModelViewViewModelBlueprint/Public/MVVMBlueprintInstancedViewModel.h
+++ b/Engine/Plugins/Runtime/ModelViewViewModel/Source/ModelViewViewModelBlueprint/Public/MVVMBlueprintInstancedViewModel.h
@@ -9,18 +9,20 @@
 #include "ViewModel/MVVMInstancedViewModelGeneratedClass.h"
 #include "MVVMBlueprintInstancedViewModel.generated.h"
 
+#define UE_API MODELVIEWVIEWMODELBLUEPRINT_API
+
 /**
  *
  */
-UCLASS(Abstract, Within = MVVMBlueprintView)
-class MODELVIEWVIEWMODELBLUEPRINT_API UMVVMBlueprintInstancedViewModelBase : public UObject
+UCLASS(MinimalAPI, Abstract, Within = MVVMBlueprintView)
+class UMVVMBlueprintInstancedViewModelBase : public UObject
 {
 	GENERATED_BODY()
 
 public:
-	UMVVMBlueprintInstancedViewModelBase();
+	UE_API UMVVMBlueprintInstancedViewModelBase();
 
-	void GenerateClass(bool bForceGeneration);
+	UE_API void GenerateClass(bool bForceGeneration);
 
 	UClass* GetGeneratedClass() const
 	{
@@ -28,13 +30,13 @@ public:
 	}
 
 protected:
-	virtual void PreloadObjectsForCompilation();
-	virtual bool IsClassDirty() const;
-	virtual void CleanClass();
-	virtual void AddProperties();
-	virtual void ConstructClass();
-	virtual void SetDefaultValues();
-	virtual void ClassGenerated();
+	UE_API virtual void PreloadObjectsForCompilation();
+	UE_API virtual bool IsClassDirty() const;
+	UE_API virtual void CleanClass();
+	UE_API virtual void AddProperties();
+	UE_API virtual void ConstructClass();
+	UE_API virtual void SetDefaultValues();
+	UE_API virtual void ClassGenerated();
 
 protected:
 	struct FInitializePropertyArgs
@@ -45,16 +47,16 @@ protected:
 		bool bPrivate = false;
 	};
 
-	bool IsValidFieldName(const FName NewPropertyName) const;
-	bool IsValidFieldName(const FName NewPropertyName, UStruct* NewOwner) const;
-	FProperty* CreateProperty(const FProperty* FromProperty, UStruct* NewOwner);
-	FProperty* CreateProperty(const FProperty* FromProperty, UStruct* NewOwner, FName NewPropertyName);
-	void InitializeProperty(FProperty* NewProperty, FInitializePropertyArgs& Args);
-	void LinkProperty(FProperty* NewProperty) const;
-	void LinkProperty(FProperty* NewProperty, UStruct* NewOwner) const;
-	void AddOnRepFunction(FProperty* NewProperty);
-	void SafeRename(UObject* Object);
-	void SetDefaultValue(const FProperty* SourceProperty, void const* SourceValuePtr, const FProperty* DestinationProperty);
+	UE_API bool IsValidFieldName(const FName NewPropertyName) const;
+	UE_API bool IsValidFieldName(const FName NewPropertyName, UStruct* NewOwner) const;
+	UE_API FProperty* CreateProperty(const FProperty* FromProperty, UStruct* NewOwner);
+	UE_API FProperty* CreateProperty(const FProperty* FromProperty, UStruct* NewOwner, FName NewPropertyName);
+	UE_API void InitializeProperty(FProperty* NewProperty, FInitializePropertyArgs& Args);
+	UE_API void LinkProperty(FProperty* NewProperty) const;
+	UE_API void LinkProperty(FProperty* NewProperty, UStruct* NewOwner) const;
+	UE_API void AddOnRepFunction(FProperty* NewProperty);
+	UE_API void SafeRename(UObject* Object);
+	UE_API void SetDefaultValue(const FProperty* SourceProperty, void const* SourceValuePtr, const FProperty* DestinationProperty);
 
 public:
 	/** The base object of the generated class. */
@@ -73,8 +75,8 @@ protected:
 /**
  *
  */
-UCLASS()
-class MODELVIEWVIEWMODELBLUEPRINT_API UMVVMBlueprintInstancedViewModel_PropertyBag : public UMVVMBlueprintInstancedViewModelBase
+UCLASS(MinimalAPI)
+class UMVVMBlueprintInstancedViewModel_PropertyBag : public UMVVMBlueprintInstancedViewModelBase
 {
 	GENERATED_BODY()
 
@@ -90,11 +92,11 @@ public:
 	}
 
 protected:
-	virtual bool IsClassDirty() const override;
-	virtual void CleanClass() override;
-	virtual void AddProperties() override;
-	virtual void SetDefaultValues() override;
-	virtual void ClassGenerated() override;
+	UE_API virtual bool IsClassDirty() const override;
+	UE_API virtual void CleanClass() override;
+	UE_API virtual void AddProperties() override;
+	UE_API virtual void SetDefaultValues() override;
+	UE_API virtual void ClassGenerated() override;
 
 private:
 	UPROPERTY(EditAnywhere, Category = "Viewmodel")
@@ -107,6 +109,8 @@ private:
 
 public:
 #if WITH_EDITOR
-	virtual void PostEditChangeChainProperty(FPropertyChangedChainEvent& PropertyChangedEvent) override;
+	UE_API virtual void PostEditChangeChainProperty(FPropertyChangedChainEvent& PropertyChangedEvent) override;
 #endif
 };
+
+#undef UE_API
