Fix relative render cost not being forwarded to cached MetaSound operator

diff --git a/Engine/Plugins/Runtime/Metasound/Source/MetasoundGenerator/Private/MetasoundGeneratorBuilder.cpp b/Engine/Plugins/Runtime/Metasound/Source/MetasoundGenerator/Private/MetasoundGeneratorBuilder.cpp
--- a/Engine/Plugins/Runtime/Metasound/Source/MetasoundGenerator/Private/MetasoundGeneratorBuilder.cpp
+++ b/Engine/Plugins/Runtime/Metasound/Source/MetasoundGenerator/Private/MetasoundGeneratorBuilder.cpp
@@ -284,7 +284,7 @@ namespace Metasound
 			);
 			
 			OpAndInputs.Operator = Builder.BuildGraphOperator(BuildParams, OutBuildResults);
-
+			OpAndInputs.RenderCost = InInitParams.GraphRenderCost;
 			return OpAndInputs;
 		}
 
diff --git a/Engine/Plugins/Runtime/Metasound/Source/MetasoundGenerator/Private/MetasoundOperatorCache.cpp b/Engine/Plugins/Runtime/Metasound/Source/MetasoundGenerator/Private/MetasoundOperatorCache.cpp
--- a/Engine/Plugins/Runtime/Metasound/Source/MetasoundGenerator/Private/MetasoundOperatorCache.cpp
+++ b/Engine/Plugins/Runtime/Metasound/Source/MetasoundGenerator/Private/MetasoundOperatorCache.cpp
@@ -518,6 +518,10 @@ namespace Metasound
 				}
 
 				FBuildResults BuildResults;
+				if (!PreCacheData->InitParams.GraphRenderCost)
+				{
+					PreCacheData->InitParams.GraphRenderCost = FGraphRenderCost::MakeGraphRenderCost();
+				}
 				FOperatorAndInputs OperatorAndInputs = GeneratorBuilder::BuildGraphOperator(PreCacheData->InitParams.OperatorSettings, PreCacheData->InitParams, BuildResults);
 				GeneratorBuilder::LogBuildErrors(PreCacheData->InitParams.MetaSoundName, BuildResults);
 
