[MetasoundDocument]
- Fix race condition when multiple documents are registered at the same time

diff --git a/Engine/Plugins/Runtime/Metasound/Source/MetasoundEngine/Public/MetasoundDocumentBuilderRegistry.h b/Engine/Plugins/Runtime/Metasound/Source/MetasoundEngine/Public/MetasoundDocumentBuilderRegistry.h
--- a/Engine/Plugins/Runtime/Metasound/Source/MetasoundEngine/Public/MetasoundDocumentBuilderRegistry.h
+++ b/Engine/Plugins/Runtime/Metasound/Source/MetasoundEngine/Public/MetasoundDocumentBuilderRegistry.h
@@ -74,6 +74,7 @@ namespace Metasound::Engine
 			TScriptInterface<IMetaSoundDocumentInterface> DocInterface = &InMetaSoundObject;
 			check(DocInterface.GetObject());
 
+			FScopeLock Lock(&BuildersCriticalSection);
 			if (UMetaSoundBuilderBase* Builder = FindBuilderObject(&InMetaSoundObject))
 			{
 				return *CastChecked<BuilderClass>(Builder);
@@ -85,23 +86,21 @@ namespace Metasound::Engine
 			const UClass& BuilderUClass = DocInterface->GetBuilderUClass();
 			const FName NewName = MakeUniqueObjectName(nullptr, &BuilderUClass, FName(*BuilderName));
 
-			TObjectPtr<UMetaSoundBuilderBase> NewBuilder;
+			TObjectPtr<UMetaSoundBuilderBase> NewBuilder = 
+				CastChecked<UMetaSoundBuilderBase>(NewObject<UObject>(GetTransientPackage(), &BuilderUClass, NewName, RF_Transactional | RF_Transient));
+
+			FMetaSoundFrontendDocumentBuilder& BuilderRef = NewBuilder->GetBuilder();
+			BuilderRef = FMetaSoundFrontendDocumentBuilder(DocInterface);
+			const FMetasoundFrontendDocument& Document = DocInterface->GetConstDocument();
+			const FMetasoundFrontendClassName& ClassName = Document.RootGraph.Metadata.GetClassName();
+			if (!ClassName.IsValid())
 			{
-				FScopeLock Lock(&BuildersCriticalSection);
-				NewBuilder = CastChecked<UMetaSoundBuilderBase>(NewObject<UObject>(GetTransientPackage(), &BuilderUClass, NewName, RF_Transactional | RF_Transient));
-				FMetaSoundFrontendDocumentBuilder& BuilderRef = NewBuilder->GetBuilder();
-				BuilderRef = FMetaSoundFrontendDocumentBuilder(DocInterface);
-				const FMetasoundFrontendDocument& Document = DocInterface->GetConstDocument();
-				const FMetasoundFrontendClassName& ClassName = Document.RootGraph.Metadata.GetClassName();
-				if (!ClassName.IsValid())
-				{
-					BuilderRef.InitDocument();
-				}
-
-				checkf(ClassName.IsValid(), TEXT("Document initialization must result in a valid class name being generated"));
-				AddBuilderInternal(ClassName, NewBuilder);
+				BuilderRef.InitDocument();
 			}
 
+			checkf(ClassName.IsValid(), TEXT("Document initialization must result in a valid class name being generated"));
+			AddBuilderInternal(ClassName, NewBuilder);
+
 			return *CastChecked<BuilderClass>(NewBuilder);
 		}
 #endif // WITH_EDITORONLY_DATA
