Fix crash due to audio thread access in Standalone Game
#rnx

diff --git a/Engine/Plugins/Runtime/Metasound/Source/MetasoundEngine/Private/MetasoundBuilderPrivate.h b/Engine/Plugins/Runtime/Metasound/Source/MetasoundEngine/Private/MetasoundBuilderPrivate.h
--- a/Engine/Plugins/Runtime/Metasound/Source/MetasoundEngine/Private/MetasoundBuilderPrivate.h
+++ b/Engine/Plugins/Runtime/Metasound/Source/MetasoundEngine/Private/MetasoundBuilderPrivate.h
@@ -14,7 +14,7 @@ namespace Metasound::Engine
 	template<typename ElementType>
 	const ElementType* FindPreferredPage_ThreadSafe(const TArray<ElementType>& InElements)
 	{
-		if (IsInGameThread())
+		if (IsInGameThread() || IsInAudioThread())
 		{
 			return Frontend::FindPreferredPage(InElements, UMetaSoundSettings::GetPageOrder());
 		}
diff --git a/Engine/Plugins/Runtime/Metasound/Source/MetasoundEngine/Private/MetasoundSettings.cpp b/Engine/Plugins/Runtime/Metasound/Source/MetasoundEngine/Private/MetasoundSettings.cpp
--- a/Engine/Plugins/Runtime/Metasound/Source/MetasoundEngine/Private/MetasoundSettings.cpp
+++ b/Engine/Plugins/Runtime/Metasound/Source/MetasoundEngine/Private/MetasoundSettings.cpp
@@ -430,9 +430,25 @@ TArray<FGuid> UMetaSoundSettings::GeneratePageOrder() const
 
 TArrayView<const FGuid> UMetaSoundSettings::GetPageOrder()
 {
-	check(IsInGameThread());
+	// Thread access for the CachedPageOrder is a bit hairy due to the existence
+	// of the AudioThread. 
+	//
+	// For non-editor builds, the page order is cached during PostInitProperties
+	// and then never altered. For these builds it is safe to access as long as
+	// it has been initialized. 
+	//
+	// For editor builds, access to the cached paged order is allowed from the 
+	// game thread or the audio thread. The audio thread runs in a handful of 
+	// editor builds, primarily when running StandaloneGame. In these builds, 
+	// there is a codepath where the CachedPageOrder could be altered on the 
+	// game thread and read on the audio thread, but this is generally disallowed
+	// as StandaloneGame does not expose changes to the UMetaSoundSettings.
+#if WITH_EDITOR
+	check(IsInGameThread() || IsInAudioThread());
+#endif
 	if (const UMetaSoundSettings* Settings = GetDefault<UMetaSoundSettings>())
 	{
+		checkf(Settings->bIsPageOrderCached, TEXT("Attempt to get PageOrder before UMetaSoundSettings have finished initializing"));
 		return Settings->CachedPageOrder;
 	}
 	return {};
@@ -600,7 +616,9 @@ void UMetaSoundSettings::PostInitProperties()
 		}
 	}
 
-	UpdateCachedPageOrder();
+	// Initialize the cached page order
+	CachedPageOrder = GeneratePageOrder();
+	bIsPageOrderCached = true;
 }
 
 bool UMetaSoundSettings::SetTargetPage(FName PageName)
@@ -676,12 +694,15 @@ TArray<FName> UMetaSoundSettings::GetQualityNames()
 
 	return { };
 }
-#endif // WITH_EDITOR
 
-void UMetaSoundSettings::UpdateCachedPageOrder()
+void UMetaSoundSettings::UpdateCachedPageOrder() 
 {
+	// In editor builds, we should only be updating the CachedPageOrder on the 
+	// GameThread. 
+	check(IsInGameThread());
 	CachedPageOrder = GeneratePageOrder();
 }
+#endif // WITH_EDITOR
 
 TArray<FGuid> UMetaSoundSettings::GeneratePageOrderInternal(const FName& InPlatformName, const FGuid& InTargetPageID) const
 {
diff --git a/Engine/Plugins/Runtime/Metasound/Source/MetasoundEngine/Public/MetasoundSettings.h b/Engine/Plugins/Runtime/Metasound/Source/MetasoundEngine/Public/MetasoundSettings.h
--- a/Engine/Plugins/Runtime/Metasound/Source/MetasoundEngine/Public/MetasoundSettings.h
+++ b/Engine/Plugins/Runtime/Metasound/Source/MetasoundEngine/Public/MetasoundSettings.h
@@ -228,9 +228,12 @@ private:
 	// Generates the page order for a given platform and desired target page ID
 	TArray<FGuid> GeneratePageOrderInternal(const FName& InPlatformName, const FGuid& InTargetPageID) const;
 
-	// Updates the cached page order with the latest settings.
+	std::atomic<bool> bIsPageOrderCached = false;	
 	TArray<FGuid> CachedPageOrder;
+#if WITH_EDITOR
+	// Updates the cached page order with the latest settings.
 	void UpdateCachedPageOrder();
+#endif
 
 public:
 	// Returns a cached global version of the PageOrder from the UMetaSoundSettings Default Object.
