Fix MetaSound graph render cost incorrectly being shared when using operator cache

diff --git a/Engine/Plugins/Runtime/Metasound/Source/MetasoundGenerator/Private/MetasoundOperatorCache.cpp b/Engine/Plugins/Runtime/Metasound/Source/MetasoundGenerator/Private/MetasoundOperatorCache.cpp
--- a/Engine/Plugins/Runtime/Metasound/Source/MetasoundGenerator/Private/MetasoundOperatorCache.cpp
+++ b/Engine/Plugins/Runtime/Metasound/Source/MetasoundGenerator/Private/MetasoundOperatorCache.cpp
@@ -518,10 +518,13 @@ namespace Metasound
 				}
 
 				FBuildResults BuildResults;
-				if (!PreCacheData->InitParams.GraphRenderCost)
+				// First operator to add will use the given graph render cost if set, 
+				// but subsequent ones must be recreated so they aren't shared between multiple instances
+				if (i != 0 || !PreCacheData->InitParams.GraphRenderCost)
 				{
 					PreCacheData->InitParams.GraphRenderCost = FGraphRenderCost::MakeGraphRenderCost();
 				}
+
 				FOperatorAndInputs OperatorAndInputs = GeneratorBuilder::BuildGraphOperator(PreCacheData->InitParams.OperatorSettings, PreCacheData->InitParams, BuildResults);
 				GeneratorBuilder::LogBuildErrors(PreCacheData->InitParams.MetaSoundName, BuildResults);
 
