Replace some usages of FORCEINLINE with inline in Metasound modules.

For context about this change, see this doc: https://docs.google.com/document/d/1LjpT7NpzTecfXdFNyv5KM6sCu67827BRljNi1iaWVSE/edit?tab=t.0#heading=h.fvtz8gjukz7x

Includes modules: MetasoundEngine, MetasoundFrontend, MetasoundGenerator, MetasoundGraphCore, MetasoundMediaStreamer, UnderscoreCore

#rnx

diff --git a/Engine/Plugins/Runtime/Metasound/Source/MetasoundEngine/Public/MetasoundAudioBus.h b/Engine/Plugins/Runtime/Metasound/Source/MetasoundEngine/Public/MetasoundAudioBus.h
--- a/Engine/Plugins/Runtime/Metasound/Source/MetasoundEngine/Public/MetasoundAudioBus.h
+++ b/Engine/Plugins/Runtime/Metasound/Source/MetasoundEngine/Public/MetasoundAudioBus.h
@@ -32,7 +32,7 @@ namespace Metasound
 			return AudioBusProxy;
 		}
 
-		friend FORCEINLINE uint32 GetTypeHash(const Metasound::FAudioBusAsset& InAudioBusAsset)
+		friend inline uint32 GetTypeHash(const Metasound::FAudioBusAsset& InAudioBusAsset)
 		{
 			const FAudioBusProxyPtr& Proxy = InAudioBusAsset.GetAudioBusProxy();
 			if (Proxy.IsValid())
diff --git a/Engine/Plugins/Runtime/Metasound/Source/MetasoundEngine/Public/MetasoundBuilderBase.h b/Engine/Plugins/Runtime/Metasound/Source/MetasoundEngine/Public/MetasoundBuilderBase.h
--- a/Engine/Plugins/Runtime/Metasound/Source/MetasoundEngine/Public/MetasoundBuilderBase.h
+++ b/Engine/Plugins/Runtime/Metasound/Source/MetasoundEngine/Public/MetasoundBuilderBase.h
@@ -213,7 +213,7 @@ public:
 
 	// Connects a given node's named output to the graph output with the given name.
 	UFUNCTION(BlueprintCallable, Category = "Audio|MetaSound|Builder", meta = (ExpandEnumAsExecs = "OutResult"))
-	FORCEINLINE void ConnectNamedNodeOutputToNamedGraphOutput(const FMetaSoundNodeHandle& SourceNode, const FName& NodeOutputName, const FName& GraphOutputName, EMetaSoundBuilderResult& OutResult)
+	inline void ConnectNamedNodeOutputToNamedGraphOutput(const FMetaSoundNodeHandle& SourceNode, const FName& NodeOutputName, const FName& GraphOutputName, EMetaSoundBuilderResult& OutResult)
 	{
 		ConnectNodeToGraphOutput(SourceNode, NodeOutputName, GraphOutputName, OutResult);
 	}
@@ -221,7 +221,7 @@ public:
 
 	// Connects a given node's named output to the graph output with the given handle.
 	UFUNCTION(BlueprintCallable, Category = "Audio|MetaSound|Builder", meta = (ExpandEnumAsExecs = "OutResult"))
-	FORCEINLINE void ConnectNamedNodeOutputToGraphOutput(const FMetaSoundNodeHandle& SourceNode, const FName& NodeOutputName, const FMetaSoundBuilderNodeInputHandle& GraphOutput, EMetaSoundBuilderResult& OutResult)
+	inline void ConnectNamedNodeOutputToGraphOutput(const FMetaSoundNodeHandle& SourceNode, const FName& NodeOutputName, const FMetaSoundBuilderNodeInputHandle& GraphOutput, EMetaSoundBuilderResult& OutResult)
 	{
 		ConnectNodeToGraphOutput(SourceNode, NodeOutputName, GraphOutput, OutResult);
 	}
diff --git a/Engine/Plugins/Runtime/Metasound/Source/MetasoundEngine/Public/MetasoundWave.h b/Engine/Plugins/Runtime/Metasound/Source/MetasoundEngine/Public/MetasoundWave.h
--- a/Engine/Plugins/Runtime/Metasound/Source/MetasoundEngine/Public/MetasoundWave.h
+++ b/Engine/Plugins/Runtime/Metasound/Source/MetasoundEngine/Public/MetasoundWave.h
@@ -53,7 +53,7 @@ namespace Metasound
 			return SoundWaveProxy.Get();
 		}
 
-		friend FORCEINLINE uint32 GetTypeHash(const Metasound::FWaveAsset& InWaveAsset)
+		friend inline uint32 GetTypeHash(const Metasound::FWaveAsset& InWaveAsset)
 		{
 			if (InWaveAsset.IsSoundWaveValid())
 			{
diff --git a/Engine/Plugins/Runtime/Metasound/Source/MetasoundEngine/Public/MetasoundWaveTable.h b/Engine/Plugins/Runtime/Metasound/Source/MetasoundEngine/Public/MetasoundWaveTable.h
--- a/Engine/Plugins/Runtime/Metasound/Source/MetasoundEngine/Public/MetasoundWaveTable.h
+++ b/Engine/Plugins/Runtime/Metasound/Source/MetasoundEngine/Public/MetasoundWaveTable.h
@@ -77,7 +77,7 @@ namespace Metasound
 			return Proxy.Get();
 		}
 		
-		friend FORCEINLINE uint32 GetTypeHash(const Metasound::FWaveTableBankAsset& InWaveTableBankAsset)
+		friend inline uint32 GetTypeHash(const Metasound::FWaveTableBankAsset& InWaveTableBankAsset)
 		{
 			if (InWaveTableBankAsset.IsValid())
 			{
diff --git a/Engine/Plugins/Runtime/Metasound/Source/MetasoundFrontend/Public/MetasoundAssetKey.h b/Engine/Plugins/Runtime/Metasound/Source/MetasoundFrontend/Public/MetasoundAssetKey.h
--- a/Engine/Plugins/Runtime/Metasound/Source/MetasoundFrontend/Public/MetasoundAssetKey.h
+++ b/Engine/Plugins/Runtime/Metasound/Source/MetasoundFrontend/Public/MetasoundAssetKey.h
@@ -31,12 +31,12 @@ struct FMetaSoundAssetKey
 	UE_API FMetaSoundAssetKey(const Metasound::Frontend::FNodeClassRegistryKey& RegKey);
 	UE_API FMetaSoundAssetKey(const FMetasoundFrontendClassMetadata& InMetadata);
 
-	FORCEINLINE friend bool operator==(const FMetaSoundAssetKey& InLHS, const FMetaSoundAssetKey& InRHS)
+	inline friend bool operator==(const FMetaSoundAssetKey& InLHS, const FMetaSoundAssetKey& InRHS)
 	{
 		return (InLHS.ClassName == InRHS.ClassName) && (InLHS.Version == InRHS.Version);
 	}
 
-	FORCEINLINE friend bool operator<(const FMetaSoundAssetKey& InLHS, const FMetaSoundAssetKey& InRHS)
+	inline friend bool operator<(const FMetaSoundAssetKey& InLHS, const FMetaSoundAssetKey& InRHS)
 	{
 		if (InLHS.ClassName == InRHS.ClassName)
 		{
@@ -46,7 +46,7 @@ struct FMetaSoundAssetKey
 		return InLHS.ClassName < InRHS.ClassName;
 	}
 
-	friend FORCEINLINE uint32 GetTypeHash(const FMetaSoundAssetKey& InKey)
+	friend inline uint32 GetTypeHash(const FMetaSoundAssetKey& InKey)
 	{
 		return HashCombineFast(GetTypeHash(InKey.ClassName), GetTypeHash(InKey.Version));
 	}
diff --git a/Engine/Plugins/Runtime/Metasound/Source/MetasoundFrontend/Public/MetasoundAssetManager.h b/Engine/Plugins/Runtime/Metasound/Source/MetasoundFrontend/Public/MetasoundAssetManager.h
--- a/Engine/Plugins/Runtime/Metasound/Source/MetasoundFrontend/Public/MetasoundAssetManager.h
+++ b/Engine/Plugins/Runtime/Metasound/Source/MetasoundFrontend/Public/MetasoundAssetManager.h
@@ -207,12 +207,12 @@ namespace Metasound::Frontend
 			FMetaSoundAssetKey Key;
 			FTopLevelAssetPath Path;
 
-			FORCEINLINE friend bool operator==(const FAssetRef& InLHS, const FAssetRef& InRHS)
+			inline friend bool operator==(const FAssetRef& InLHS, const FAssetRef& InRHS)
 			{
 				return (InLHS.Key == InRHS.Key) && (InLHS.Path == InRHS.Path);
 			}
 
-			FORCEINLINE friend uint32 GetTypeHash(const FAssetRef& InInfo)
+			inline friend uint32 GetTypeHash(const FAssetRef& InInfo)
 			{
 				return HashCombineFast(GetTypeHash(InInfo.Key), GetTypeHash(InInfo.Path));
 			}
@@ -223,12 +223,12 @@ namespace Metasound::Frontend
 			FNodeRegistryKey RegistryKey;
 			FSoftObjectPath AssetPath;
 
-			FORCEINLINE friend bool operator==(const FAssetInfo& InLHS, const FAssetInfo& InRHS)
+			inline friend bool operator==(const FAssetInfo& InLHS, const FAssetInfo& InRHS)
 			{
 				return (InLHS.RegistryKey == InRHS.RegistryKey) && (InLHS.AssetPath == InRHS.AssetPath);
 			}
 
-			FORCEINLINE friend uint32 GetTypeHash(const FAssetInfo& InInfo)
+			inline friend uint32 GetTypeHash(const FAssetInfo& InInfo)
 			{
 				return HashCombineFast(GetTypeHash(InInfo.RegistryKey), GetTypeHash(InInfo.AssetPath));
 			}
diff --git a/Engine/Plugins/Runtime/Metasound/Source/MetasoundFrontend/Public/MetasoundFrontendDocument.h b/Engine/Plugins/Runtime/Metasound/Source/MetasoundFrontend/Public/MetasoundFrontendDocument.h
--- a/Engine/Plugins/Runtime/Metasound/Source/MetasoundFrontend/Public/MetasoundFrontendDocument.h
+++ b/Engine/Plugins/Runtime/Metasound/Source/MetasoundFrontend/Public/MetasoundFrontendDocument.h
@@ -332,7 +332,7 @@ struct FMetasoundFrontendVersionNumber
 
 	UE_API FString ToString() const;
 
-	friend FORCEINLINE uint32 GetTypeHash(const FMetasoundFrontendVersionNumber& InNumber)
+	friend inline uint32 GetTypeHash(const FMetasoundFrontendVersionNumber& InNumber)
 	{
 		return HashCombineFast(GetTypeHash(InNumber.Major), GetTypeHash(InNumber.Minor));
 	}
@@ -408,7 +408,7 @@ struct FMetasoundFrontendVersion
 		return InLHS == InRHS || InLHS < InRHS;
 	}
 
-	friend FORCEINLINE uint32 GetTypeHash(const FMetasoundFrontendVersion& InVersion)
+	friend inline uint32 GetTypeHash(const FMetasoundFrontendVersion& InVersion)
 	{
 		return HashCombineFast(GetTypeHash(InVersion.Name), GetTypeHash(InVersion.Number));
 	}
@@ -468,7 +468,7 @@ public:
 		return InLHS.NodeID != InRHS.NodeID || InLHS.VertexID != InRHS.VertexID;
 	}
 
-	friend FORCEINLINE uint32 GetTypeHash(const FMetasoundFrontendVertexHandle& InHandle)
+	friend inline uint32 GetTypeHash(const FMetasoundFrontendVertexHandle& InHandle)
 	{
 		return HashCombineFast(InHandle.NodeID.A, InHandle.VertexID.D);
 	}
@@ -716,7 +716,7 @@ struct FMetasoundFrontendEdge
 		return !(InLHS == InRHS);
 	}
 
-	friend FORCEINLINE uint32 GetTypeHash(const FMetasoundFrontendEdge& InEdge)
+	friend inline uint32 GetTypeHash(const FMetasoundFrontendEdge& InEdge)
 	{
 		const int32 FromHash = HashCombineFast(InEdge.FromNodeID.A, InEdge.FromVertexID.B);
 		const int32 ToHash = HashCombineFast(InEdge.ToNodeID.C, InEdge.ToVertexID.D);
@@ -1289,7 +1289,7 @@ struct FMetasoundFrontendInterfaceVertexBinding
 		return FString::Format(TEXT("{0}->{1}"), { OutputName.ToString(), InputName.ToString() });
 	}
 
-	friend FORCEINLINE uint32 GetTypeHash(const FMetasoundFrontendInterfaceVertexBinding& InBinding)
+	friend inline uint32 GetTypeHash(const FMetasoundFrontendInterfaceVertexBinding& InBinding)
 	{
 		return HashCombineFast(GetTypeHash(InBinding.OutputName), GetTypeHash(InBinding.InputName));
 	}
@@ -1434,18 +1434,18 @@ struct FMetasoundFrontendClassName
 	// Parses string into class name.  For deserialization and debug use only.
 	static UE_API bool Parse(const FString& InClassName, FMetasoundFrontendClassName& OutClassName);
 
-	friend FORCEINLINE uint32 GetTypeHash(const FMetasoundFrontendClassName& ClassName)
+	friend inline uint32 GetTypeHash(const FMetasoundFrontendClassName& ClassName)
 	{
 		const int32 NameHash = HashCombineFast(GetTypeHash(ClassName.Namespace), GetTypeHash(ClassName.Name));
 		return HashCombineFast(NameHash, GetTypeHash(ClassName.Variant));
 	}
 
-	friend FORCEINLINE bool operator==(const FMetasoundFrontendClassName& InLHS, const FMetasoundFrontendClassName& InRHS)
+	friend inline bool operator==(const FMetasoundFrontendClassName& InLHS, const FMetasoundFrontendClassName& InRHS)
 	{
 		return (InLHS.Namespace == InRHS.Namespace) && (InLHS.Name == InRHS.Name) && (InLHS.Variant == InRHS.Variant);
 	}
 
-	friend FORCEINLINE bool operator<(const FMetasoundFrontendClassName& InLHS, const FMetasoundFrontendClassName& InRHS)
+	friend inline bool operator<(const FMetasoundFrontendClassName& InLHS, const FMetasoundFrontendClassName& InRHS)
 	{
 		if (InLHS.Namespace == InRHS.Namespace)
 		{
diff --git a/Engine/Plugins/Runtime/Metasound/Source/MetasoundFrontend/Public/MetasoundFrontendDocumentBuilder.h b/Engine/Plugins/Runtime/Metasound/Source/MetasoundFrontend/Public/MetasoundFrontendDocumentBuilder.h
--- a/Engine/Plugins/Runtime/Metasound/Source/MetasoundFrontend/Public/MetasoundFrontendDocumentBuilder.h
+++ b/Engine/Plugins/Runtime/Metasound/Source/MetasoundFrontend/Public/MetasoundFrontendDocumentBuilder.h
@@ -62,7 +62,7 @@ namespace Metasound::Frontend
 			return !(InLHS == InRHS);
 		}
 
-		friend FORCEINLINE uint32 GetTypeHash(const FNamedEdge& InBinding)
+		friend inline uint32 GetTypeHash(const FNamedEdge& InBinding)
 		{
 			const int32 NameHash = HashCombineFast(GetTypeHash(InBinding.OutputName), GetTypeHash(InBinding.InputName));
 			const int32 GuidHash = HashCombineFast(GetTypeHash(InBinding.OutputNodeID), GetTypeHash(InBinding.InputNodeID));
diff --git a/Engine/Plugins/Runtime/Metasound/Source/MetasoundFrontend/Public/MetasoundFrontendNodeClassRegistry.h b/Engine/Plugins/Runtime/Metasound/Source/MetasoundFrontend/Public/MetasoundFrontendNodeClassRegistry.h
--- a/Engine/Plugins/Runtime/Metasound/Source/MetasoundFrontend/Public/MetasoundFrontendNodeClassRegistry.h
+++ b/Engine/Plugins/Runtime/Metasound/Source/MetasoundFrontend/Public/MetasoundFrontendNodeClassRegistry.h
@@ -124,7 +124,7 @@ namespace Metasound::Frontend
 		// The datatype one would like to convert to.
 		FName ToDataType;
 
-		FORCEINLINE bool operator==(const FConverterNodeClassRegistryKey& Other) const
+		inline bool operator==(const FConverterNodeClassRegistryKey& Other) const
 		{
 			return FromDataType == Other.FromDataType && ToDataType == Other.ToDataType;
 		}
@@ -147,7 +147,7 @@ namespace Metasound::Frontend
 		// The key for this node in the node registry.
 		FNodeClassRegistryKey NodeKey;
 
-		FORCEINLINE bool operator==(const FConverterNodeClassInfo& Other) const
+		inline bool operator==(const FConverterNodeClassInfo& Other) const
 		{
 			return NodeKey == Other.NodeKey;
 		}
diff --git a/Engine/Plugins/Runtime/Metasound/Source/MetasoundFrontend/Public/MetasoundFrontendRegistryKey.h b/Engine/Plugins/Runtime/Metasound/Source/MetasoundFrontend/Public/MetasoundFrontendRegistryKey.h
--- a/Engine/Plugins/Runtime/Metasound/Source/MetasoundFrontend/Public/MetasoundFrontendRegistryKey.h
+++ b/Engine/Plugins/Runtime/Metasound/Source/MetasoundFrontend/Public/MetasoundFrontendRegistryKey.h
@@ -78,12 +78,12 @@ namespace Metasound::Frontend
 		UE_API FNodeClassRegistryKey(const FNodeClassInfo& InClassInfo);
 		UE_API FNodeClassRegistryKey(const FMetaSoundAssetKey& AssetKey);
 
-		FORCEINLINE friend bool operator==(const FNodeClassRegistryKey& InLHS, const FNodeClassRegistryKey& InRHS)
+		inline friend bool operator==(const FNodeClassRegistryKey& InLHS, const FNodeClassRegistryKey& InRHS)
 		{
 			return (InLHS.Type == InRHS.Type) && (InLHS.ClassName == InRHS.ClassName) && (InLHS.Version == InRHS.Version);
 		}
 
-		FORCEINLINE friend bool operator<(const FNodeClassRegistryKey& InLHS, const FNodeClassRegistryKey& InRHS)
+		inline friend bool operator<(const FNodeClassRegistryKey& InLHS, const FNodeClassRegistryKey& InRHS)
 		{
 			if (static_cast<uint8>(InLHS.Type) == static_cast<uint8>(InRHS.Type))
 			{
@@ -98,7 +98,7 @@ namespace Metasound::Frontend
 			return static_cast<uint8>(InLHS.Type) < static_cast<uint8>(InRHS.Type);
 		}
 
-		friend FORCEINLINE uint32 GetTypeHash(const FNodeClassRegistryKey& InKey)
+		friend inline uint32 GetTypeHash(const FNodeClassRegistryKey& InKey)
 		{
 			const int32 Hash = HashCombineFast(static_cast<uint32>(InKey.Type), GetTypeHash(InKey.ClassName));
 			return HashCombineFast(Hash, GetTypeHash(InKey.Version));
@@ -137,12 +137,12 @@ namespace Metasound::Frontend
 
 		UE_API bool IsValid() const;
 
-		FORCEINLINE friend bool operator==(const FGraphClassRegistryKey& InLHS, const FGraphClassRegistryKey& InRHS)
+		inline friend bool operator==(const FGraphClassRegistryKey& InLHS, const FGraphClassRegistryKey& InRHS)
 		{
 			return (InLHS.NodeKey == InRHS.NodeKey) && (InLHS.AssetPath == InRHS.AssetPath);
 		}
 
-		friend FORCEINLINE uint32 GetTypeHash(const FGraphClassRegistryKey& InKey)
+		friend inline uint32 GetTypeHash(const FGraphClassRegistryKey& InKey)
 		{
 			const int32 Hash = HashCombineFast(GetTypeHash(InKey.NodeKey), GetTypeHash(InKey.AssetPath));
 			return Hash;
diff --git a/Engine/Plugins/Runtime/Metasound/Source/MetasoundFrontend/Public/MetasoundParameterPack.h b/Engine/Plugins/Runtime/Metasound/Source/MetasoundFrontend/Public/MetasoundParameterPack.h
--- a/Engine/Plugins/Runtime/Metasound/Source/MetasoundFrontend/Public/MetasoundParameterPack.h
+++ b/Engine/Plugins/Runtime/Metasound/Source/MetasoundFrontend/Public/MetasoundParameterPack.h
@@ -509,7 +509,7 @@ public:
 		return ParameterStoragePtr.Get();
 	}
 
-	friend FORCEINLINE uint32 GetTypeHash(const FMetasoundParameterStorageWrapper& InParameterStorageWrapper)
+	friend inline uint32 GetTypeHash(const FMetasoundParameterStorageWrapper& InParameterStorageWrapper)
 	{
 		if (InParameterStorageWrapper.IsPackValid())
 		{
diff --git a/Engine/Plugins/Runtime/Metasound/Source/MetasoundFrontend/Public/MetasoundParameterPackFixedArray.h b/Engine/Plugins/Runtime/Metasound/Source/MetasoundFrontend/Public/MetasoundParameterPackFixedArray.h
--- a/Engine/Plugins/Runtime/Metasound/Source/MetasoundFrontend/Public/MetasoundParameterPackFixedArray.h
+++ b/Engine/Plugins/Runtime/Metasound/Source/MetasoundFrontend/Public/MetasoundParameterPackFixedArray.h
@@ -37,12 +37,12 @@ namespace Metasound
 			}
 		}
 
-		FORCEINLINE const ElementType* GetData() const
+		inline const ElementType* GetData() const
 		{
 			return (ElementType*)InlineData;
 		}
 
-		FORCEINLINE ElementType* GetData()
+		inline ElementType* GetData()
 		{
 			return (ElementType*)InlineData;
 		}
diff --git a/Engine/Plugins/Runtime/Metasound/Source/MetasoundFrontend/Public/MetasoundRouter.h b/Engine/Plugins/Runtime/Metasound/Source/MetasoundFrontend/Public/MetasoundRouter.h
--- a/Engine/Plugins/Runtime/Metasound/Source/MetasoundFrontend/Public/MetasoundRouter.h
+++ b/Engine/Plugins/Runtime/Metasound/Source/MetasoundFrontend/Public/MetasoundRouter.h
@@ -175,7 +175,7 @@ namespace Metasound
 		uint64 InstanceID = INDEX_NONE;
 	};
 
-	FORCEINLINE uint32 GetTypeHash(const Metasound::FTransmissionAddress& InAddress)
+	inline uint32 GetTypeHash(const Metasound::FTransmissionAddress& InAddress)
 	{
 		return InAddress.GetHash();
 	}
diff --git a/Engine/Plugins/Runtime/Metasound/Source/MetasoundFrontend/Public/MetasoundTime.h b/Engine/Plugins/Runtime/Metasound/Source/MetasoundFrontend/Public/MetasoundTime.h
--- a/Engine/Plugins/Runtime/Metasound/Source/MetasoundFrontend/Public/MetasoundTime.h
+++ b/Engine/Plugins/Runtime/Metasound/Source/MetasoundFrontend/Public/MetasoundTime.h
@@ -359,7 +359,7 @@ namespace Metasound
 				return FTime(static_cast<TimeType>(InMicroseconds * 1e-6));
 			}
 			
-			friend FORCEINLINE uint32 GetTypeHash(const FTime& InTime)
+			friend inline uint32 GetTypeHash(const FTime& InTime)
 			{
 				return ::GetTypeHash(InTime.Time);
 			}
diff --git a/Engine/Plugins/Runtime/Metasound/Source/MetasoundGenerator/Public/MetasoundGenerator.h b/Engine/Plugins/Runtime/Metasound/Source/MetasoundGenerator/Public/MetasoundGenerator.h
--- a/Engine/Plugins/Runtime/Metasound/Source/MetasoundGenerator/Public/MetasoundGenerator.h
+++ b/Engine/Plugins/Runtime/Metasound/Source/MetasoundGenerator/Public/MetasoundGenerator.h
@@ -96,7 +96,7 @@ namespace Metasound
 
 		METASOUNDGENERATOR_API friend bool operator<(const FOperatorPoolEntryID& InLHS, const FOperatorPoolEntryID& InRHS);
 		METASOUNDGENERATOR_API friend bool operator==(const FOperatorPoolEntryID& InLHS, const FOperatorPoolEntryID& InRHS);
-		friend FORCEINLINE uint32 GetTypeHash(const FOperatorPoolEntryID& InID)
+		friend inline uint32 GetTypeHash(const FOperatorPoolEntryID& InID)
 		{
 			return HashCombineFast(GetTypeHash(InID.OperatorID), GetTypeHash(InID.OperatorSettings));
 		}
diff --git a/Engine/Plugins/Runtime/Metasound/Source/MetasoundGraphCore/Public/MetasoundAudioBuffer.h b/Engine/Plugins/Runtime/Metasound/Source/MetasoundGraphCore/Public/MetasoundAudioBuffer.h
--- a/Engine/Plugins/Runtime/Metasound/Source/MetasoundGraphCore/Public/MetasoundAudioBuffer.h
+++ b/Engine/Plugins/Runtime/Metasound/Source/MetasoundGraphCore/Public/MetasoundAudioBuffer.h
@@ -59,7 +59,7 @@ namespace Metasound
 			}
 
 			/** Return a pointer to the audio float data. */
-			FORCEINLINE const float* GetData() const
+			inline const float* GetData() const
 			{
 #if METASOUNDGRAPHCORE_CHECKAUDIONUM
 				UE_CLOG(InitialNum != Buffer.Num(), LogMetaSound, Error, TEXT("MetaSound audio buffer size change detected.  Audio Buffers should not be resized."));
@@ -68,7 +68,7 @@ namespace Metasound
 			}
 
 			/** Return a pointer to the audio float data. */
-			FORCEINLINE float* GetData()
+			inline float* GetData()
 			{
 #if METASOUNDGRAPHCORE_CHECKAUDIONUM
 				UE_CLOG(InitialNum != Buffer.Num(), LogMetaSound, Error, TEXT("MetaSound audio buffer size change detected.  Audio Buffers should not be resized."));
@@ -77,7 +77,7 @@ namespace Metasound
 			}
 
 			/** Return the number of samples in the audio buffer. */
-			FORCEINLINE int32 Num() const
+			inline int32 Num() const
 			{
 #if METASOUNDGRAPHCORE_CHECKAUDIONUM
 				UE_CLOG(InitialNum != Buffer.Num(), LogMetaSound, Error, TEXT("MetaSound audio buffer size change detected.  Audio Buffers should not be resized."));
@@ -86,7 +86,7 @@ namespace Metasound
 			}
 
 			/** Implicit conversion to Audio::FAlignedFloatBuffer */
-			FORCEINLINE operator const Audio::FAlignedFloatBuffer& () const
+			inline operator const Audio::FAlignedFloatBuffer& () const
 			{
 #if METASOUNDGRAPHCORE_CHECKAUDIONUM
 				UE_CLOG(InitialNum != Buffer.Num(), LogMetaSound, Error, TEXT("MetaSound audio buffer size change detected.  Audio Buffers should not be resized."));
@@ -98,7 +98,7 @@ return Buffer;
 			 *
 			 * WARNING: if the buffer is resized, it will cause errors.
 			 */
-			FORCEINLINE operator Audio::FAlignedFloatBuffer& ()
+			inline operator Audio::FAlignedFloatBuffer& ()
 			{
 #if METASOUNDGRAPHCORE_CHECKAUDIONUM
 				UE_CLOG(InitialNum != Buffer.Num(), LogMetaSound, Error, TEXT("MetaSound audio buffer size change detected.  Audio Buffers should not be resized."));
@@ -107,7 +107,7 @@ return Buffer;
 			}
 
 			/** Implicit conversion to Audio::FAlignedFloatBuffer */
-			FORCEINLINE operator TArrayView<const float> () const
+			inline operator TArrayView<const float> () const
 			{
 #if METASOUNDGRAPHCORE_CHECKAUDIONUM
 				UE_CLOG(InitialNum != Buffer.Num(), LogMetaSound, Error, TEXT("MetaSound audio buffer size change detected.  Audio Buffers should not be resized."));
@@ -119,7 +119,7 @@ return Buffer;
 			 *
 			 * WARNING: if the buffer is resized, it will cause errors.
 			 */
-			FORCEINLINE operator TArrayView<float> ()
+			inline operator TArrayView<float> ()
 			{
 #if METASOUNDGRAPHCORE_CHECKAUDIONUM
 				UE_CLOG(InitialNum != Buffer.Num(), LogMetaSound, Error, TEXT("MetaSound audio buffer size change detected.  Audio Buffers should not be resized."));
@@ -127,7 +127,7 @@ return Buffer;
 				return Buffer;
 			}
 
-			FORCEINLINE void Zero()
+			inline void Zero()
 			{
 				if (Buffer.Num() > 0)
 				{
diff --git a/Engine/Plugins/Runtime/Metasound/Source/MetasoundGraphCore/Public/MetasoundDataReference.h b/Engine/Plugins/Runtime/Metasound/Source/MetasoundGraphCore/Public/MetasoundDataReference.h
--- a/Engine/Plugins/Runtime/Metasound/Source/MetasoundGraphCore/Public/MetasoundDataReference.h
+++ b/Engine/Plugins/Runtime/Metasound/Source/MetasoundGraphCore/Public/MetasoundDataReference.h
@@ -134,7 +134,7 @@ namespace Metasound
 	};
 
 	/** Return the ID of the data reference. */
-	FORCEINLINE FDataReferenceID GetDataReferenceID(const IDataReference& InDataReference)
+	inline FDataReferenceID GetDataReferenceID(const IDataReference& InDataReference)
 	{
 		return static_cast<FDataReferenceID>(InDataReference.GetRaw());
 	}
@@ -296,18 +296,18 @@ namespace Metasound
 		}
 
 		/** Const access to the underlying parameter object. */
-		FORCEINLINE const DataType& operator*() const
+		inline const DataType& operator*() const
 		{
 			return *TDataReference<DataType>::ObjectReference;
 		}
 
 		/** Const access to the underlying parameter object. */
-		FORCEINLINE const DataType* operator->() const
+		inline const DataType* operator->() const
 		{
 			return TDataReference<DataType>::ObjectReference.operator->();
 		}
 
-		FORCEINLINE const DataType* Get() const
+		inline const DataType* Get() const
 		{
 			return TDataReference<DataType>::ObjectReference.operator->();
 		}
@@ -398,19 +398,19 @@ namespace Metasound
 		}
 
 		/** Non-const access to the underlying parameter object. */
-		FORCEINLINE DataType& operator*() const
+		inline DataType& operator*() const
 		{
 			return *TDataReference<DataType>::ObjectReference;
 		}
 
 		/** Non-const access to the underlying parameter object. */
-		FORCEINLINE DataType* operator->() const
+		inline DataType* operator->() const
 		{
 			return TDataReference<DataType>::ObjectReference.operator->();
 		}
 
 		/** Non-const access to the underlying parameter object. */
-		FORCEINLINE DataType* Get() const
+		inline DataType* Get() const
 		{
 			return TDataReference<DataType>::ObjectReference.operator->();
 		}
@@ -493,19 +493,19 @@ namespace Metasound
 		TDataReadReference<DataType>& operator=(TDataReadReference<DataType>&& Other) = default;
 
 		/** Const access to the underlying parameter object. */
-		FORCEINLINE const DataType& operator*() const
+		inline const DataType& operator*() const
 		{
 			return *TDataReference<DataType>::ObjectReference;
 		}
 
 		/** Const access to the underlying parameter object. */
-		FORCEINLINE const DataType* operator->() const
+		inline const DataType* operator->() const
 		{
 			return TDataReference<DataType>::ObjectReference.operator->();
 		}
 
 		/** Non-const access to the underlying parameter object. */
-		FORCEINLINE const DataType* Get() const
+		inline const DataType* Get() const
 		{
 			return TDataReference<DataType>::ObjectReference.operator->();
 		}
diff --git a/Engine/Plugins/Runtime/Metasound/Source/MetasoundGraphCore/Public/MetasoundLiteral.h b/Engine/Plugins/Runtime/Metasound/Source/MetasoundGraphCore/Public/MetasoundLiteral.h
--- a/Engine/Plugins/Runtime/Metasound/Source/MetasoundGraphCore/Public/MetasoundLiteral.h
+++ b/Engine/Plugins/Runtime/Metasound/Source/MetasoundGraphCore/Public/MetasoundLiteral.h
@@ -171,74 +171,74 @@ namespace Metasound
 		}
 
 		template<>
-		FORCEINLINE ELiteralType GetLiteralArgTypeFromDecayed<>()
+		inline ELiteralType GetLiteralArgTypeFromDecayed<>()
 		{
 			checkNoEntry();
 			return ELiteralType::None;
 		}
 
 		template<>
-		FORCEINLINE ELiteralType GetLiteralArgTypeFromDecayed<bool>()
+		inline ELiteralType GetLiteralArgTypeFromDecayed<bool>()
 		{
 			return ELiteralType::Boolean;
 		}
 
 		template<>
-		FORCEINLINE ELiteralType GetLiteralArgTypeFromDecayed<int32>()
+		inline ELiteralType GetLiteralArgTypeFromDecayed<int32>()
 		{
 			return ELiteralType::Integer;
 		}
 
 		template<>
-		FORCEINLINE ELiteralType GetLiteralArgTypeFromDecayed<float>()
+		inline ELiteralType GetLiteralArgTypeFromDecayed<float>()
 		{
 			return ELiteralType::Float;
 		}
 
 		template<>
-		FORCEINLINE ELiteralType GetLiteralArgTypeFromDecayed<FString>()
+		inline ELiteralType GetLiteralArgTypeFromDecayed<FString>()
 		{
 			return ELiteralType::String;
 		}
 
 		template<>
-		FORCEINLINE ELiteralType GetLiteralArgTypeFromDecayed<TSharedPtr<Audio::IProxyData>>()
+		inline ELiteralType GetLiteralArgTypeFromDecayed<TSharedPtr<Audio::IProxyData>>()
 		{
 			return ELiteralType::UObjectProxy;
 		}
 
 		template<>
-		FORCEINLINE ELiteralType GetLiteralArgTypeFromDecayed<TArray<FLiteral::FNone>>()
+		inline ELiteralType GetLiteralArgTypeFromDecayed<TArray<FLiteral::FNone>>()
 		{
 			return ELiteralType::NoneArray;
 		}
 
 		template<>
-		FORCEINLINE ELiteralType GetLiteralArgTypeFromDecayed<TArray<bool>>()
+		inline ELiteralType GetLiteralArgTypeFromDecayed<TArray<bool>>()
 		{
 			return ELiteralType::BooleanArray;
 		}
 
 		template<>
-		FORCEINLINE ELiteralType GetLiteralArgTypeFromDecayed<TArray<int32>>()
+		inline ELiteralType GetLiteralArgTypeFromDecayed<TArray<int32>>()
 		{
 			return ELiteralType::IntegerArray;
 		}
 
 		template<>
-		FORCEINLINE ELiteralType GetLiteralArgTypeFromDecayed<TArray<float>>()
+		inline ELiteralType GetLiteralArgTypeFromDecayed<TArray<float>>()
 		{
 			return ELiteralType::FloatArray;
 		}
 
 		template<>
-		FORCEINLINE ELiteralType GetLiteralArgTypeFromDecayed<TArray<FString>>()
+		inline ELiteralType GetLiteralArgTypeFromDecayed<TArray<FString>>()
 		{
 			return ELiteralType::StringArray;
 		}
 
 		template<>
-		FORCEINLINE ELiteralType GetLiteralArgTypeFromDecayed<TArray<TSharedPtr<Audio::IProxyData>>>()
+		inline ELiteralType GetLiteralArgTypeFromDecayed<TArray<TSharedPtr<Audio::IProxyData>>>()
 		{
 			return ELiteralType::UObjectProxyArray;
 		}
diff --git a/Engine/Plugins/Runtime/Metasound/Source/MetasoundGraphCore/Public/MetasoundNodeInterface.h b/Engine/Plugins/Runtime/Metasound/Source/MetasoundGraphCore/Public/MetasoundNodeInterface.h
--- a/Engine/Plugins/Runtime/Metasound/Source/MetasoundGraphCore/Public/MetasoundNodeInterface.h
+++ b/Engine/Plugins/Runtime/Metasound/Source/MetasoundGraphCore/Public/MetasoundNodeInterface.h
@@ -110,7 +110,7 @@ namespace Metasound
 		/** Invalid form of node class name (i.e. empty namespace, name, and variant)*/
 		static UE_API const FNodeClassName InvalidNodeClassName;
 
-		friend FORCEINLINE bool operator==(const FNodeClassName& InLHS, const FNodeClassName& InRHS)
+		friend inline bool operator==(const FNodeClassName& InLHS, const FNodeClassName& InRHS)
 		{
 			return (InLHS.Namespace == InRHS.Namespace) && (InLHS.Name == InRHS.Name) && (InLHS.Variant == InRHS.Variant);
 		}
@@ -317,12 +317,12 @@ namespace Metasound
 	typedef TMap<FNodeDataVertexKey, FInputDataDestination> FInputDataDestinationCollection;
 
 	/** Make a FNodeDataVertexKey from an FOutputDataSource. */
-	FORCEINLINE FNodeDataVertexKey MakeSourceDataVertexKey(const FOutputDataSource& InSource)
+	inline FNodeDataVertexKey MakeSourceDataVertexKey(const FOutputDataSource& InSource)
 	{
 		return FNodeDataVertexKey(InSource.Node, InSource.Vertex.VertexName);
 	}
 
-	FORCEINLINE FNodeDataVertexKey MakeDestinationDataVertexKey(const FInputDataDestination& InDestination)
+	inline FNodeDataVertexKey MakeDestinationDataVertexKey(const FInputDataDestination& InDestination)
 	{
 		return FNodeDataVertexKey(InDestination.Node, InDestination.Vertex.VertexName);
 	}
diff --git a/Engine/Plugins/Runtime/Metasound/Source/MetasoundGraphCore/Public/MetasoundOperatorSettings.h b/Engine/Plugins/Runtime/Metasound/Source/MetasoundGraphCore/Public/MetasoundOperatorSettings.h
--- a/Engine/Plugins/Runtime/Metasound/Source/MetasoundGraphCore/Public/MetasoundOperatorSettings.h
+++ b/Engine/Plugins/Runtime/Metasound/Source/MetasoundGraphCore/Public/MetasoundOperatorSettings.h
@@ -51,7 +51,7 @@ namespace Metasound
 			UE_API void SetSampleRate(FSampleRate InSampleRate);
 
 			/** Get the audio sample rate in Hz. */
-			FORCEINLINE float GetSampleRate() const
+			inline float GetSampleRate() const
 			{
 				return SampleRate;
 			}
@@ -81,7 +81,7 @@ namespace Metasound
 			 * enables the vast majority of Metasound IOperators to take 
 			 * advantage of SIMD hardware accelerations.
 			 */
-			FORCEINLINE int32 GetNumFramesPerBlock() const
+			inline int32 GetNumFramesPerBlock() const
 			{
 				return NumFramesPerBlock;
 			}
@@ -92,7 +92,7 @@ namespace Metasound
 			METASOUNDGRAPHCORE_API friend bool operator<(const FOperatorSettings& InLHS, const FOperatorSettings& InRHS);
 			METASOUNDGRAPHCORE_API friend bool operator==(const FOperatorSettings& InLHS, const FOperatorSettings& InRHS);
 
-			friend FORCEINLINE uint32 GetTypeHash(const FOperatorSettings& InSettings)
+			friend inline uint32 GetTypeHash(const FOperatorSettings& InSettings)
 			{
 				// Ignore TargetBlockRate in comparison because it is not used in any external calculations
 				// Ignore ActualBlockRate because it is derived from NumFramesPerBlock
diff --git a/Engine/Plugins/Runtime/Metasound/Source/MetasoundGraphCore/Public/MetasoundVertex.h b/Engine/Plugins/Runtime/Metasound/Source/MetasoundGraphCore/Public/MetasoundVertex.h
--- a/Engine/Plugins/Runtime/Metasound/Source/MetasoundGraphCore/Public/MetasoundVertex.h
+++ b/Engine/Plugins/Runtime/Metasound/Source/MetasoundGraphCore/Public/MetasoundVertex.h
@@ -50,7 +50,7 @@ namespace Metasound
 			{
 			}
 
-			FORCEINLINE bool operator()(const VertexType& InOther) const
+			inline bool operator()(const VertexType& InOther) const
 			{
 				return InOther.VertexName == NameRef;
 			}
diff --git a/Engine/Plugins/Runtime/Metasound/Source/MetasoundGraphCore/Public/MetasoundVertexData.h b/Engine/Plugins/Runtime/Metasound/Source/MetasoundGraphCore/Public/MetasoundVertexData.h
--- a/Engine/Plugins/Runtime/Metasound/Source/MetasoundGraphCore/Public/MetasoundVertexData.h
+++ b/Engine/Plugins/Runtime/Metasound/Source/MetasoundGraphCore/Public/MetasoundVertexData.h
@@ -30,7 +30,7 @@ namespace Metasound
 #if ENABLE_METASOUND_ACCESS_TYPE_COMPATIBILITY_TEST
 		void METASOUNDGRAPHCORE_API CheckAccessTypeCompatibility(const FDataVertex& InDataVertex, const FAnyDataReference& InDataReference);
 #else
-		FORCEINLINE void CheckAccessTypeCompatibility(const FDataVertex& InDataVertex, const FAnyDataReference& InDataReference) {}
+		inline void CheckAccessTypeCompatibility(const FDataVertex& InDataVertex, const FAnyDataReference& InDataReference) {}
 #endif // #if DO_CHECK
 
 		METASOUNDGRAPHCORE_API EVertexAccessType DataReferenceAccessTypeToVertexAccessType(EDataReferenceAccessType InReferenceAccessType);
@@ -162,7 +162,7 @@ namespace Metasound
 		struct TTriggerDeprecatedBindAPIWarning<const T>
 		{
 			UE_DEPRECATED(5.6, "Please use a non-const lvalue reference for data references passed to this function. Const data references can sometimes be created by passing in temporary objects and/or triggering implicit conversion.")
-			FORCEINLINE void Trigger() 
+			inline void Trigger() 
 			{
 			}
 		};
