- Part 1 of cutting down FName table bloat & string passing, copying, & parsing by moving implementation to use AssetPath wherever possible
- Misc clean-up and removal of old code

#rnx
#tests AQ_OperatorCache, AudioUnitTests, Editor Playback, PIE

diff --git a/Engine/Plugins/Runtime/Metasound/Source/MetasoundEngine/Private/MetasoundBuilderSubsystem.cpp b/Engine/Plugins/Runtime/Metasound/Source/MetasoundEngine/Private/MetasoundBuilderSubsystem.cpp
--- a/Engine/Plugins/Runtime/Metasound/Source/MetasoundEngine/Private/MetasoundBuilderSubsystem.cpp
+++ b/Engine/Plugins/Runtime/Metasound/Source/MetasoundEngine/Private/MetasoundBuilderSubsystem.cpp
@@ -487,7 +487,8 @@ void UMetaSoundSourceBuilder::OnNodeAdded(int32 NodeIndex) const
 		using namespace Metasound;
 		using namespace Metasound::Frontend;
 
-		const FMetasoundFrontendGraphClass& OwningGraphClass = Builder.GetConstDocumentChecked().RootGraph;
+		const IMetaSoundDocumentInterface& DocInterface = Builder.GetConstDocumentInterfaceChecked();
+		const FMetasoundFrontendGraphClass& OwningGraphClass = DocInterface.GetConstDocument().RootGraph;
 		const FMetasoundFrontendGraph& OwningGraph = GetConstTargetPageGraphChecked();
 		const FMetasoundFrontendNode& AddedNode = OwningGraph.Nodes[NodeIndex];
 		const FMetasoundFrontendClass* NodeClass = Builder.FindDependency(AddedNode.ClassID);
@@ -495,11 +496,11 @@ void UMetaSoundSourceBuilder::OnNodeAdded(int32 NodeIndex) const
 
 		FGraphBuilder::FCreateNodeParams CreateNodeParams
 		{
-			GetName(), 
-			AddedNode,
-			*NodeClass,
-			OwningGraph,
-			OwningGraphClass
+			.FrontendNode = AddedNode,
+			.FrontendNodeClass = *NodeClass,
+			.OwningFrontendGraph = OwningGraph,
+			.OwningFrontendGraphClass = OwningGraphClass,
+			.OwningAssetPath = DocInterface.GetAssetPathChecked()
 		};
 		TUniquePtr<INode> NewNode = FGraphBuilder::CreateNode(CreateNodeParams);
 
diff --git a/Engine/Plugins/Runtime/Metasound/Source/MetasoundEngine/Private/MetasoundOperatorCacheSubsystem.cpp b/Engine/Plugins/Runtime/Metasound/Source/MetasoundEngine/Private/MetasoundOperatorCacheSubsystem.cpp
--- a/Engine/Plugins/Runtime/Metasound/Source/MetasoundEngine/Private/MetasoundOperatorCacheSubsystem.cpp
+++ b/Engine/Plugins/Runtime/Metasound/Source/MetasoundEngine/Private/MetasoundOperatorCacheSubsystem.cpp
@@ -18,7 +18,7 @@ namespace Metasound::OperatorCachePrivate
 		TEXT("If precaching metasound operators via the UMetaSoundCacheSubsystem is enabled.")
 	);
 
-	TOptional<FMetasoundGeneratorInitParams> CreateInitParams(UMetaSoundSource& InMetaSound, const FSoundGeneratorInitParams& InParams)
+	TOptional<FGeneratorInitParams> CreateInitParams(UMetaSoundSource& InMetaSound, const FSoundGeneratorInitParams& InParams)
 	{
 		using namespace Metasound::Frontend;
 		using namespace Metasound::Engine;
@@ -37,17 +37,15 @@ namespace Metasound::OperatorCachePrivate
 		// Graph analyzer currently only enabled for preview sounds (but can theoretically be supported for all sounds)
 		BuilderSettings.bPopulateInternalDataReferences = InParams.bIsPreviewSound;
 
-		return FMetasoundGeneratorInitParams
+		return FGeneratorInitParams
 		{
-			InSettings,
-			MoveTemp(BuilderSettings),
-			{}, // Graph, retrieved from the FrontEnd Registry in FOperatorPool::BuildAndAddOperator()
-			Environment,
-			InMetaSound.GetName(),
-			InMetaSound.GetOutputAudioChannelOrder(),
-			{}, // DefaultParameters
-			true, // bBuildSynchronous
-			{} // DataChannel
+			.OperatorSettings = InSettings,
+			.BuilderSettings = MoveTemp(BuilderSettings),
+			.Graph = { }, // Will be retrieved from the FrontEnd Registry in FOperatorPool::BuildAndAddOperator()
+			.Environment = Environment,
+			.AudioOutputNames = InMetaSound.GetOutputAudioChannelOrder(),
+			.bBuildSynchronous = true,
+			.AssetPath = FTopLevelAssetPath(&InMetaSound)
 		};
 	}
 } // namespace Metasound::OperatorCachePrivate
@@ -125,7 +123,7 @@ void UMetaSoundCacheSubsystem::PrecacheMetaSoundInternal(UMetaSoundSource* InMet
 	InMetaSound->InitResources();
 
 	BuildParams.GraphName = InMetaSound->GetOwningAssetName();
-	TOptional<FMetasoundGeneratorInitParams> InitParams = OperatorCachePrivate::CreateInitParams(*InMetaSound, BuildParams);
+	TOptional<FGeneratorInitParams> InitParams = OperatorCachePrivate::CreateInitParams(*InMetaSound, BuildParams);
 	if (!InitParams.IsSet())
 	{
 		return;
diff --git a/Engine/Plugins/Runtime/Metasound/Source/MetasoundEngine/Private/MetasoundSource.cpp b/Engine/Plugins/Runtime/Metasound/Source/MetasoundEngine/Private/MetasoundSource.cpp
--- a/Engine/Plugins/Runtime/Metasound/Source/MetasoundEngine/Private/MetasoundSource.cpp
+++ b/Engine/Plugins/Runtime/Metasound/Source/MetasoundEngine/Private/MetasoundSource.cpp
@@ -1073,19 +1073,19 @@ ISoundGeneratorPtr UMetaSoundSource::CreateSoundGenerator(const FSoundGeneratorI
 		TSharedRef<FGraphRenderCost> GraphRenderCost = FGraphRenderCost::MakeGraphRenderCost();
 		FMetasoundDynamicGraphGeneratorInitParams InitParams
 		{
+			FGeneratorInitParams
 			{
-				InSettings,
-				MoveTemp(BuilderSettings),
-				MakeShared<FGraph>(DynamicTransactor->GetGraph()), // Make a copy of the graph.
-				Environment,
-				GetName(),
-				GetOutputAudioChannelOrder(),
-				MoveTemp(InDefaultParameters),
-				bBuildSynchronous,
-				DataChannel,
-				GraphRenderCost,
-				FName(GetClass()->GetName()),
-				GetAssetPathChecked().GetPackageName()
+				.OperatorSettings = InSettings,
+				.BuilderSettings = MoveTemp(BuilderSettings),
+				.Graph = MakeShared<FGraph>(DynamicTransactor->GetGraph()), // Make a copy of the graph.
+				.Environment = Environment,
+				.AudioOutputNames = GetOutputAudioChannelOrder(),
+				.DefaultParameters = MoveTemp(InDefaultParameters),
+				.bBuildSynchronous = bBuildSynchronous,
+				.DataChannel = DataChannel,
+				.GraphRenderCost = GraphRenderCost,
+				.ClassName = GetClass()->GetFName(),
+				.AssetPath = GetAssetPathChecked()
 			},
 			DynamicTransactor->CreateTransformQueue(InSettings, Environment, GraphRenderCost) // Create transaction queue
 		};
@@ -1114,22 +1114,19 @@ ISoundGeneratorPtr UMetaSoundSource::CreateSoundGenerator(const FSoundGeneratorI
 		if (MetasoundGraph.IsValid())
 		{
 			// Create generator.
-			FMetasoundGeneratorInitParams InitParams
+			FGeneratorInitParams InitParams
 			{
-				InSettings,
-				MoveTemp(BuilderSettings),
-				MetasoundGraph,
-				Environment,
-				GetName(),
-				GetOutputAudioChannelOrder(),
-				MoveTemp(InDefaultParameters),
-				bBuildSynchronous,
-				DataChannel,
-				{}, // RelativeRenderCost
-				FName(GetClass()->GetName()),
-				GetAssetPathChecked().GetPackageName()
+				.OperatorSettings = InSettings,
+				.BuilderSettings = MoveTemp(BuilderSettings),
+				.Graph = MetasoundGraph,
+				.Environment = Environment,
+				.AudioOutputNames = GetOutputAudioChannelOrder(),
+				.DefaultParameters = MoveTemp(InDefaultParameters),
+				.bBuildSynchronous = bBuildSynchronous,
+				.DataChannel = DataChannel,
+				.ClassName = GetClass()->GetFName()
 			};
-
+			InitParams.AssetPath = GetAssetPathChecked();
 			Generator = MakeShared<FMetasoundConstGraphGenerator>(MoveTemp(InitParams));
 		}
 	}
diff --git a/Engine/Plugins/Runtime/Metasound/Source/MetasoundEngineTest/Private/MetasoundDynamicGeneratorTests.cpp b/Engine/Plugins/Runtime/Metasound/Source/MetasoundEngineTest/Private/MetasoundDynamicGeneratorTests.cpp
--- a/Engine/Plugins/Runtime/Metasound/Source/MetasoundEngineTest/Private/MetasoundDynamicGeneratorTests.cpp
+++ b/Engine/Plugins/Runtime/Metasound/Source/MetasoundEngineTest/Private/MetasoundDynamicGeneratorTests.cpp
@@ -44,19 +44,19 @@ namespace Metasound::Test::Generator::Dynamic
 			
 			FMetasoundDynamicGraphGeneratorInitParams InitParams
 			{
+				FGeneratorInitParams
 				{
-					OperatorSettings,
-					MoveTemp(BuilderSettings),
-					MakeShared<FGraph>(Transactor.GetGraph()),
-					Environment,
-					"TestMetaSoundGenerator",
-					GetAudioOutputVertexNames(EMetaSoundOutputAudioFormat::Mono),
-					{},
-					true
+					.OperatorSettings = OperatorSettings,
+					.BuilderSettings = MoveTemp(BuilderSettings),
+					.Graph = MakeShared<FGraph>(Transactor.GetGraph()),
+					.Environment = Environment,
+					.AudioOutputNames = GetAudioOutputVertexNames(EMetaSoundOutputAudioFormat::Mono),
+					.bBuildSynchronous = true,
+					.AssetPath = FTopLevelAssetPath { FName(), FName("TestMetaSoundGenerator") }
 				},
 				Transactor.CreateTransformQueue(OperatorSettings, Environment, nullptr) // Create transaction queue
 			};
-			
+
 			Generator.Init(MoveTemp(InitParams));
 		}
 		
diff --git a/Engine/Plugins/Runtime/Metasound/Source/MetasoundEngineTest/Private/NodeTestGraphBuilder.cpp b/Engine/Plugins/Runtime/Metasound/Source/MetasoundEngineTest/Private/NodeTestGraphBuilder.cpp
--- a/Engine/Plugins/Runtime/Metasound/Source/MetasoundEngineTest/Private/NodeTestGraphBuilder.cpp
+++ b/Engine/Plugins/Runtime/Metasound/Source/MetasoundEngineTest/Private/NodeTestGraphBuilder.cpp
@@ -89,8 +89,9 @@ namespace Metasound::Test
 		// Otherwise we get warnings when we run our tests.
 		AddInput(SourceInterface::Inputs::OnPlay, Metasound::GetMetasoundDataTypeName<FTrigger>());
 		
-		const TSet<FName> TransmittableInputNames;
-		const FString UnknownAsset = TEXT("UnknownAsset");
+		UPackage* TransientPackage = GetTransientPackage();
+		check(TransientPackage);
+		const FTopLevelAssetPath UnknownAsset(TransientPackage->GetFName(), FName("UnknownAsset"));
 		if (const TUniquePtr<FFrontendGraph> Graph = Frontend::FGraphBuilder::CreateGraph(Document, UnknownAsset))
 		{
 			FOperatorBuilderSettings BuilderSettings;
@@ -107,17 +108,16 @@ namespace Metasound::Test
 			FMetasoundEnvironment Environment;
 			Environment.SetValue<uint64>(CoreInterface::Environment::InstanceID, 123);
 			
-			FMetasoundGeneratorInitParams GeneratorInitParams{
-				OperatorSettings,
-				BuilderSettings,
-				MakeShared<const FFrontendGraph, ESPMode::ThreadSafe>(*Graph),
-				Environment,
-				"TestMetasound",
-				AudioOutputNames,
-				TArray<FAudioParameter>(),
-				true // bBuildSynchronous, so we don't have to do latent tasks
+			FGeneratorInitParams GeneratorInitParams
+			{
+				.OperatorSettings = OperatorSettings,
+				.BuilderSettings = MoveTemp(BuilderSettings),
+				.Graph = MakeShared<const FFrontendGraph, ESPMode::ThreadSafe>(*Graph),
+				.Environment = Environment,
+				.AudioOutputNames = AudioOutputNames,
+				.bBuildSynchronous = true, // Don't need to wait for latent tasks
+				.AssetPath = FTopLevelAssetPath(FName(), "TestMetasound")
 			};
-
 			return MakeUnique<FMetasoundConstGraphGenerator>(MoveTemp(GeneratorInitParams));
 		}
 
diff --git a/Engine/Plugins/Runtime/Metasound/Source/MetasoundFrontend/Private/MetasoundFrontendGraph.cpp b/Engine/Plugins/Runtime/Metasound/Source/MetasoundFrontend/Private/MetasoundFrontendGraph.cpp
--- a/Engine/Plugins/Runtime/Metasound/Source/MetasoundFrontend/Private/MetasoundFrontendGraph.cpp
+++ b/Engine/Plugins/Runtime/Metasound/Source/MetasoundFrontend/Private/MetasoundFrontendGraph.cpp
@@ -10,7 +10,14 @@
 namespace Metasound
 {
 	FFrontendGraph::FFrontendGraph(const FString& InInstanceName, const FGuid& InInstanceID)
-	:	FGraph(InInstanceName, InInstanceID)
+PRAGMA_DISABLE_DEPRECATION_WARNINGS
+	:	FGraph(FTopLevelAssetPath(InInstanceName), InInstanceID)
+PRAGMA_ENABLE_DEPRECATION_WARNINGS
+	{
+	}
+
+	FFrontendGraph::FFrontendGraph(const FTopLevelAssetPath& InAssetPath, const FGuid& InInstanceID)
+		: FGraph(InAssetPath, InInstanceID)
 	{
 	}
 
diff --git a/Engine/Plugins/Runtime/Metasound/Source/MetasoundFrontend/Private/MetasoundFrontendGraphBuilder.cpp b/Engine/Plugins/Runtime/Metasound/Source/MetasoundFrontend/Private/MetasoundFrontendGraphBuilder.cpp
--- a/Engine/Plugins/Runtime/Metasound/Source/MetasoundFrontend/Private/MetasoundFrontendGraphBuilder.cpp
+++ b/Engine/Plugins/Runtime/Metasound/Source/MetasoundFrontend/Private/MetasoundFrontendGraphBuilder.cpp
@@ -7,7 +7,6 @@
 #include "Algo/AnyOf.h"
 #include "Algo/TopologicalSort.h"
 #include "Algo/Transform.h"
-#include "CoreMinimal.h"
 #include "CoreGlobals.h"
 #include "MetasoundDocumentInterface.h"
 #include "MetasoundFrontend.h"
@@ -21,6 +20,7 @@
 #include "MetasoundLiteralNode.h"
 #include "MetasoundLog.h"
 #include "MetasoundNodeInterface.h"
+#include "UObject/TopLevelAssetPath.h"
 
 namespace Metasound::Frontend
 {
@@ -35,7 +35,7 @@ namespace Metasound::Frontend
 		// (for both a root and nested subgraphs)
 		struct FBuildContext
 		{
-			const FString& DebugAssetName;
+			const FTopLevelAssetPath& AssetPath;
 			const Frontend::IDataTypeRegistry& DataTypeRegistry;
 			const Frontend::FProxyDataCache* ProxyDataCache;
 			TArrayView<const FGuid> PageOrder;
@@ -231,7 +231,7 @@ namespace Metasound::Frontend
 									// long ago, but this log will help identify any issues.
 									// This log is added in UE 5.6 and if no one is hitting
 									// this issue, can be removed in UE 5.7.
-									UE_LOG(LogMetaSound, Error, TEXT("The MetaSound %s has an malformed input node %s. Please replace the input node and resave"), *InContext.DebugAssetName, *InNode.Name.ToString());
+									UE_LOG(LogMetaSound, Error, TEXT("The MetaSound %s has an malformed input node %s. Please replace the input node and resave"), *InContext.AssetPath.ToString(), *InNode.Name.ToString());
 								}
 							}
 						}
@@ -348,7 +348,7 @@ namespace Metasound::Frontend
 									// long ago, but this log will help identify any issues.
 									// This log is added in UE 5.6 and if no one is hitting
 									// this issue, can be removed in UE 5.7.
-									UE_LOG(LogMetaSound, Error, TEXT("The MetaSound %s has an malformed output node %s. Please replace the output node and resave"), *InContext.DebugAssetName, *InNode.Name.ToString());
+									UE_LOG(LogMetaSound, Error, TEXT("The MetaSound %s has an malformed output node %s. Please replace the output node and resave"), *InContext.AssetPath.ToString(), *InNode.Name.ToString());
 								}
 							}
 						}
@@ -463,7 +463,7 @@ namespace Metasound::Frontend
 			// Check that the frontend node and class is correct.
 			if (InNode.Interface.Inputs.Num() != 1)
 			{
-				UE_LOG(LogMetaSound, Error, TEXT("MetaSound %s contains invalid number of inputs (%d) on input node %s"), *InContext.DebugAssetName, InNode.Interface.Inputs.Num(), *InNode.Name.ToString());
+				UE_LOG(LogMetaSound, Error, TEXT("MetaSound %s contains invalid number of inputs (%d) on input node %s"), *InContext.AssetPath.ToString(), InNode.Interface.Inputs.Num(), *InNode.Name.ToString());
 				return TUniquePtr<INode>(nullptr);
 			}
 
@@ -480,7 +480,7 @@ namespace Metasound::Frontend
 			}
 			else
 			{
-				UE_LOG(LogMetaSound, Error, TEXT("MetaSound %s contains invalid input access type on input %s"), *InContext.DebugAssetName, *InNode.Name.ToString());
+				UE_LOG(LogMetaSound, Error, TEXT("MetaSound %s contains invalid input access type on input %s"), *InContext.AssetPath.ToString(), *InNode.Name.ToString());
 			}
 
 			return TUniquePtr<INode>(nullptr);
@@ -494,7 +494,7 @@ namespace Metasound::Frontend
 
 			if (InNode.Interface.Outputs.Num() != 1)
 			{
-				UE_LOG(LogMetaSound, Error, TEXT("MetaSound %s contains invalid number of outputs (%d) on output node %s"), *InContext.DebugAssetName, InNode.Interface.Outputs.Num(), *InNode.Name.ToString());
+				UE_LOG(LogMetaSound, Error, TEXT("MetaSound %s contains invalid number of outputs (%d) on output node %s"), *InContext.AssetPath.ToString(), InNode.Interface.Outputs.Num(), *InNode.Name.ToString());
 				return TUniquePtr<INode>(nullptr);
 			}
 
@@ -510,7 +510,7 @@ namespace Metasound::Frontend
 			}
 			else
 			{
-				UE_LOG(LogMetaSound, Error, TEXT("MetaSound %s contains invalid output access type on output %s"), *InContext.DebugAssetName, *InNode.Name.ToString());
+				UE_LOG(LogMetaSound, Error, TEXT("MetaSound %s contains invalid output access type on output %s"), *InContext.AssetPath.ToString(), *InNode.Name.ToString());
 			}
 
 			return TUniquePtr<INode>(nullptr);
@@ -587,21 +587,21 @@ namespace Metasound::Frontend
 				{
 					const FGraphBuilder::FCreateNodeParams CreateNodeParams
 					{
-						InGraphContext.BuildContext.DebugAssetName,
-						FrontendNode,
-						*FrontendNodeClass,
-						InGraphContext.PagedGraph,
-						InGraphContext.GraphClass,
-						InGraphContext.BuildContext.ProxyDataCache, 
-						&InGraphContext.BuildContext.DataTypeRegistry,
-						&InGraphContext.BuildDocumentContext.Graphs,
-						InGraphContext.BuildContext.PageOrder
+						.FrontendNode = FrontendNode,
+						.FrontendNodeClass = *FrontendNodeClass,
+						.OwningFrontendGraph = InGraphContext.PagedGraph,
+						.OwningFrontendGraphClass = InGraphContext.GraphClass,
+						.ProxyDataCache = InGraphContext.BuildContext.ProxyDataCache,
+						.DataTypeRegistry= &InGraphContext.BuildContext.DataTypeRegistry,
+						.Subgraphs = &InGraphContext.BuildDocumentContext.Graphs,
+						.PageOrder = InGraphContext.BuildContext.PageOrder,
+						.OwningAssetPath = InGraphContext.BuildContext.AssetPath
 					};
 					TUniquePtr<INode> Node = FGraphBuilder::CreateNode(CreateNodeParams);
 
 					if (!Node.IsValid())
 					{
-						LogFailedToCreateNode(InGraphContext.BuildContext.DebugAssetName, FrontendNode, *FrontendNodeClass);
+						LogFailedToCreateNode(InGraphContext.BuildContext.AssetPath.ToString(), FrontendNode, *FrontendNodeClass);
 						continue;
 					}
 
@@ -658,7 +658,7 @@ namespace Metasound::Frontend
 				const INode* CoreNode = InGraphContext.Graph->FindNode(Node.GetID());
 				if (nullptr == CoreNode)
 				{
-					UE_LOG(LogMetaSound, Warning, TEXT("Metasound '%s': Could not find referenced node [Name:%s, NodeID:%s]"), *InGraphContext.BuildContext.DebugAssetName, *Node.Name.ToString(), *Node.GetID().ToString());
+					UE_LOG(LogMetaSound, Warning, TEXT("Metasound '%s': Could not find referenced node [Name:%s, NodeID:%s]"), *InGraphContext.BuildContext.AssetPath.ToString(), *Node.Name.ToString(), *Node.GetID().ToString());
 					return false;
 				}
 
@@ -680,13 +680,13 @@ namespace Metasound::Frontend
 
 				if (nullptr == DestinationNodeAndVertex)
 				{
-					UE_LOG(LogMetaSound, Error, TEXT("MetaSound '%s': Failed to add edge. Could not find destination [NodeID:%s, VertexID:%s]"), *InGraphContext.BuildContext.DebugAssetName, *Edge.ToNodeID.ToString(), *Edge.ToVertexID.ToString());
+					UE_LOG(LogMetaSound, Error, TEXT("MetaSound '%s': Failed to add edge. Could not find destination [NodeID:%s, VertexID:%s]"), *InGraphContext.BuildContext.AssetPath.ToString(), *Edge.ToNodeID.ToString(), *Edge.ToVertexID.ToString());
 					return false;
 				}
 
 				if (nullptr == DestinationNodeAndVertex->Node)
 				{
-					UE_LOG(LogMetaSound, Warning, TEXT("MetaSound '%s': 'Failed to add edge. Null destination node [NodeID:%s]"), *InGraphContext.BuildContext.DebugAssetName, *Edge.ToNodeID.ToString());
+					UE_LOG(LogMetaSound, Warning, TEXT("MetaSound '%s': 'Failed to add edge. Null destination node [NodeID:%s]"), *InGraphContext.BuildContext.AssetPath.ToString(), *Edge.ToNodeID.ToString());
 					return false;
 				}
 
@@ -695,13 +695,13 @@ namespace Metasound::Frontend
 
 				if (nullptr == SourceNodeAndVertex)
 				{
-					UE_LOG(LogMetaSound, Error, TEXT("MetaSound '%s': Failed to add edge. Could not find source [NodeID:%s, VertexID:%s]"), *InGraphContext.BuildContext.DebugAssetName, *Edge.FromNodeID.ToString(), *Edge.FromVertexID.ToString());
+					UE_LOG(LogMetaSound, Error, TEXT("MetaSound '%s': Failed to add edge. Could not find source [NodeID:%s, VertexID:%s]"), *InGraphContext.BuildContext.AssetPath.ToString(), *Edge.FromNodeID.ToString(), *Edge.FromVertexID.ToString());
 					return false;
 				}
 
 				if (nullptr == SourceNodeAndVertex->Node)
 				{
-					UE_LOG(LogMetaSound, Warning, TEXT("MetaSound '%s': Skipping edge. Null source node [NodeID:%s]"), *InGraphContext.BuildContext.DebugAssetName, *Edge.FromNodeID.ToString());
+					UE_LOG(LogMetaSound, Warning, TEXT("MetaSound '%s': Skipping edge. Null source node [NodeID:%s]"), *InGraphContext.BuildContext.AssetPath.ToString(), *Edge.FromNodeID.ToString());
 					return false;
 				}
 
@@ -716,7 +716,7 @@ namespace Metasound::Frontend
 				// If succeeded, remove input as viable vertex to construct default variable, as it has been superceded by a connection.
 				if (!bSuccess)
 				{
-					UE_LOG(LogMetaSound, Error, TEXT("MetaSound '%s': Failed to connect edge from [NodeID:%s, VertexID:%s] to [NodeID:%s, VertexID:%s]"), *InGraphContext.BuildContext.DebugAssetName, *Edge.FromNodeID.ToString(), *Edge.FromVertexID.ToString(), *Edge.ToNodeID.ToString(), *Edge.ToVertexID.ToString());
+					UE_LOG(LogMetaSound, Error, TEXT("MetaSound '%s': Failed to connect edge from [NodeID:%s, VertexID:%s] to [NodeID:%s, VertexID:%s]"), *InGraphContext.BuildContext.AssetPath.ToString(), *Edge.FromNodeID.ToString(), *Edge.FromVertexID.ToString(), *Edge.ToNodeID.ToString(), *Edge.ToVertexID.ToString());
 					return false;
 				}
 			}
@@ -774,8 +774,6 @@ namespace Metasound::Frontend
 
 		TUniquePtr<FFrontendGraph> CreateGraphInternal(FBuildContext& InContext, const FBuildDocumentContext& InDocumentContext, const FMetasoundFrontendGraphClass& InGraphClass, const FGuid& InGraphId)
 		{
-			const FString GraphName = InContext.DebugAssetName;
-
 			const FMetasoundFrontendGraph* PageGraph = FindPreferredPage(InGraphClass.GetConstGraphPages(), InContext.PageOrder);
 			if (!PageGraph)
 			{
@@ -785,7 +783,7 @@ namespace Metasound::Frontend
 
 			FBuildGraphContext BuildGraphContext
 			{
-				MakeUnique<FFrontendGraph>(GraphName, InGraphId),
+				MakeUnique<FFrontendGraph>(InContext.AssetPath, InGraphId),
 				InGraphClass,
 				*PageGraph,
 				InContext,
@@ -823,18 +821,30 @@ namespace Metasound::Frontend
 
 		if (ClassInterface.Inputs.Num() != InParams.FrontendNode.Interface.Inputs.Num())
 		{
-			UE_LOG(LogMetaSound, Error, TEXT("Cannot create node. MetaSound %s contains mismatched number of inputs (%d / %d) on node %s class %s"), *InParams.DebugAssetName, InParams.FrontendNode.Interface.Inputs.Num(), ClassInterface.Inputs.Num(), *InParams.FrontendNode.Name.ToString(), *InParams.FrontendNodeClass.Metadata.GetClassName().ToString());
+			UE_LOG(LogMetaSound, Error,
+				TEXT("Cannot create node. MetaSound %s contains mismatched number of inputs (%d / %d) on node %s class %s"),
+				*InParams.OwningAssetPath.ToString(),
+				InParams.FrontendNode.Interface.Inputs.Num(),
+				ClassInterface.Inputs.Num(),
+				*InParams.FrontendNode.Name.ToString(),
+				*InParams.FrontendNodeClass.Metadata.GetClassName().ToString());
 			return Node;
 		}
 		else if (ClassInterface.Outputs.Num() != InParams.FrontendNode.Interface.Outputs.Num())
 		{
-			UE_LOG(LogMetaSound, Error, TEXT("Cannot create node. MetaSound %s contains mismatched number of outputs (%d / %d) on node %s class %s"), *InParams.DebugAssetName, InParams.FrontendNode.Interface.Outputs.Num(), ClassInterface.Outputs.Num(), *InParams.FrontendNode.Name.ToString(), *InParams.FrontendNodeClass.Metadata.GetClassName().ToString());
+			UE_LOG(LogMetaSound, Error,
+				TEXT("Cannot create node. MetaSound %s contains mismatched number of outputs (%d / %d) on node %s class %s"),
+				*InParams.OwningAssetPath.ToString(),
+				InParams.FrontendNode.Interface.Outputs.Num(),
+				ClassInterface.Outputs.Num(),
+				*InParams.FrontendNode.Name.ToString(),
+				*InParams.FrontendNodeClass.Metadata.GetClassName().ToString());
 			return Node;
 		}
 
 		FBuildContext BuildContext
 		{
-			InParams.DebugAssetName,
+			InParams.OwningAssetPath,
 			InParams.DataTypeRegistry ? *InParams.DataTypeRegistry : IDataTypeRegistry::Get(),
 			InParams.ProxyDataCache,
 			InParams.PageOrder
@@ -853,7 +863,12 @@ namespace Metasound::Frontend
 				else
 				{
 					const FString GraphClassIDString = InParams.OwningFrontendGraphClass.ID.ToString();
-					UE_LOG(LogMetaSound, Error, TEXT("MetaSound '%s': Failed to match input node [NodeID:%s, NodeName:%s] to owning graph [ClassID:%s] input."), *InParams.DebugAssetName, *InParams.FrontendNode.GetID().ToString(), *InParams.FrontendNode.Name.ToString(), *GraphClassIDString);
+					UE_LOG(LogMetaSound, Error,
+						TEXT("MetaSound '%s': Failed to match input node [NodeID:%s, NodeName:%s] to owning graph [ClassID:%s] input."),
+						*InParams.OwningAssetPath.ToString(),
+						*InParams.FrontendNode.GetID().ToString(),
+						*InParams.FrontendNode.Name.ToString(),
+						*GraphClassIDString);
 				}
 			}
 			break;
@@ -868,7 +883,12 @@ namespace Metasound::Frontend
 				else
 				{
 					const FString GraphClassIDString = InParams.OwningFrontendGraphClass.ID.ToString();
-					UE_LOG(LogMetaSound, Error, TEXT("MetaSound '%s': Failed to match output node [NodeID:%s, NodeName:%s] to owning graph [ClassID:%s] output."), *InParams.DebugAssetName, *InParams.FrontendNode.GetID().ToString(), *InParams.FrontendNode.Name.ToString(), *GraphClassIDString);
+					UE_LOG(LogMetaSound, Error,
+						TEXT("MetaSound '%s': Failed to match output node [NodeID:%s, NodeName:%s] to owning graph [ClassID:%s] output."),
+						*InParams.OwningAssetPath.ToString(),
+						*InParams.FrontendNode.GetID().ToString(),
+						*InParams.FrontendNode.Name.ToString(),
+						*GraphClassIDString);
 				}
 			}
 			break;
@@ -877,8 +897,12 @@ namespace Metasound::Frontend
 			{
 				if (InParams.Subgraphs == nullptr)
 				{
-					UE_LOG(LogMetaSound, Error, TEXT("MetaSound '%s': Failed to find subgraph for node [NodeID:%s, NodeName:%s, ClassID:%s] because of missing subgraph map"), *InParams.DebugAssetName, *InParams.FrontendNode.GetID().ToString(), *InParams.FrontendNode.Name.ToString(), *InParams.FrontendNode.ClassID.ToString());
-
+					UE_LOG(LogMetaSound, Error,
+						TEXT("MetaSound '%s': Failed to find subgraph for node [NodeID:%s, NodeName:%s, ClassID:%s] because of missing subgraph map"),
+						*InParams.OwningAssetPath.ToString(),
+						*InParams.FrontendNode.GetID().ToString(),
+						*InParams.FrontendNode.Name.ToString(),
+						*InParams.FrontendNode.ClassID.ToString());
 				}
 				else
 				{
@@ -889,7 +913,12 @@ namespace Metasound::Frontend
 					}
 					else
 					{
-						UE_LOG(LogMetaSound, Error, TEXT("MetaSound '%s': Found invalid subgraph for node [NodeID:%s, NodeName:%s, ClassID:%s]"), *InParams.DebugAssetName, *InParams.FrontendNode.GetID().ToString(), *InParams.FrontendNode.Name.ToString(), *InParams.FrontendNode.ClassID.ToString());
+						UE_LOG(LogMetaSound, Error,
+							TEXT("MetaSound '%s': Found invalid subgraph for node [NodeID:%s, NodeName:%s, ClassID:%s]"),
+							*InParams.OwningAssetPath.ToString(),
+							*InParams.FrontendNode.GetID().ToString(),
+							*InParams.FrontendNode.Name.ToString(),
+							*InParams.FrontendNode.ClassID.ToString());
 
 					}
 				}
@@ -898,7 +927,9 @@ namespace Metasound::Frontend
 
 			case EMetasoundFrontendClassType::Literal:
 			{
-				UE_LOG(LogMetaSound, Warning, TEXT("MetaSound '%s': Adding literal nodes is being deprecated. Please set literal values on the connected node interface directly."), *InParams.DebugAssetName);
+				UE_LOG(LogMetaSound, Warning,
+					TEXT("MetaSound '%s': Adding literal nodes is being deprecated. Please set literal values on the connected node interface directly."),
+					*InParams.OwningAssetPath.ToString());
 				const FName& DataTypeName = InParams.FrontendNodeClass.Metadata.GetClassName().Name;
 PRAGMA_DISABLE_DEPRECATION_WARNINGS
 				FDefaultLiteralNodeConstructorParams InitData { InParams.FrontendNode.Name, InParams.FrontendNode.GetID(), BuildContext.DataTypeRegistry.CreateDefaultLiteral(DataTypeName)};
@@ -931,16 +962,23 @@ PRAGMA_ENABLE_DEPRECATION_WARNINGS
 	}
 
 
-	TUniquePtr<FFrontendGraph> FGraphBuilder::CreateGraph(const FMetasoundFrontendGraphClass& InGraph, const TArray<FMetasoundFrontendGraphClass>& InSubgraphs, const TArray<FMetasoundFrontendClass>& InDependencies, const Frontend::FProxyDataCache& InProxyDataCache, const FString& InDebugAssetName, const FGuid InGraphId, TArrayView<const FGuid> InPageOrder)
+	TUniquePtr<FFrontendGraph> FGraphBuilder::CreateGraph(
+		const FMetasoundFrontendGraphClass& InGraph,
+		const TArray<FMetasoundFrontendGraphClass>& InSubgraphs,
+		const TArray<FMetasoundFrontendClass>& InDependencies,
+		const Frontend::FProxyDataCache& InProxyDataCache,
+		const FTopLevelAssetPath& InAssetPath,
+		const FGuid InGraphId,
+		TArrayView<const FGuid> InPageOrder)
 	{
 		using namespace GraphBuilderPrivate;
 
 		FBuildContext Context
 		{
-			InDebugAssetName, 					// DebugAssetName
-			Frontend::IDataTypeRegistry::Get(), // DataTypeRegistry
-			&InProxyDataCache, 					// ProxyDataCache
-			InPageOrder                         // Page resolution order
+			InAssetPath,
+			Frontend::IDataTypeRegistry::Get(),
+			&InProxyDataCache,
+			InPageOrder
 		};
 
 		FBuildDocumentContext DocumentContext;
@@ -962,7 +1000,7 @@ PRAGMA_ENABLE_DEPRECATION_WARNINGS
 		bool bSuccess = SortSubgraphDependencies(FrontendSubgraphPtrs);
 		if (!bSuccess)
 		{
-			UE_LOG(LogMetaSound, Error, TEXT("Failed to create graph due to failed subgraph ordering in asset '%s'."), *InDebugAssetName);
+			UE_LOG(LogMetaSound, Error, TEXT("Failed to create graph due to failed subgraph ordering in asset '%s'."), *InAssetPath.ToString());
 			return TUniquePtr<FFrontendGraph>(nullptr);
 		}
 
@@ -972,7 +1010,7 @@ PRAGMA_ENABLE_DEPRECATION_WARNINGS
 			TSharedPtr<const IGraph> Subgraph(CreateGraphInternal(Context, DocumentContext, *FrontendSubgraphPtr, FrontendSubgraphPtr->ID).Release());
 			if (!Subgraph.IsValid())
 			{
-				UE_LOG(LogMetaSound, Warning, TEXT("Failed to create subgraph [SubgraphName: %s] in asset '%s'"), *FrontendSubgraphPtr->Metadata.GetClassName().ToString(), *InDebugAssetName);
+				UE_LOG(LogMetaSound, Warning, TEXT("Failed to create subgraph [SubgraphName: %s] in asset '%s'"), *FrontendSubgraphPtr->Metadata.GetClassName().ToString(), *InAssetPath.ToString());
 			}
 			else
 			{
@@ -985,16 +1023,24 @@ PRAGMA_ENABLE_DEPRECATION_WARNINGS
 		return CreateGraphInternal(Context, DocumentContext, InGraph, InGraphId);
 	}
 
-	TUniquePtr<FFrontendGraph> FGraphBuilder::CreateGraph(const FMetasoundFrontendDocument& InDocument, const FString& InDebugAssetName, TArrayView<const FGuid> InPageOrder)
+	TUniquePtr<FFrontendGraph> FGraphBuilder::CreateGraph(
+		const FMetasoundFrontendDocument& InDocument,
+		const FTopLevelAssetPath& InAssetPath,
+		TArrayView<const FGuid> InPageOrder)
 	{
 		// Create proxies before creating graph.
 		Frontend::FProxyDataCache ProxyDataCache;
 		ProxyDataCache.CreateAndCacheProxies(InDocument, InPageOrder);
 		
-		return CreateGraph(InDocument, ProxyDataCache, InDebugAssetName, Frontend::CreateLocallyUniqueId(), InPageOrder);
+		return CreateGraph(InDocument, ProxyDataCache, InAssetPath, Frontend::CreateLocallyUniqueId(), InPageOrder);
 	}
 
-	TUniquePtr<FFrontendGraph> FGraphBuilder::CreateGraph(const FMetasoundFrontendGraphClass& InGraph, const TArray<FMetasoundFrontendGraphClass>& InSubgraphs, const TArray<FMetasoundFrontendClass>& InDependencies, const FString& InDebugAssetName, TArrayView<const FGuid> InPageOrder)
+	TUniquePtr<FFrontendGraph> FGraphBuilder::CreateGraph(
+		const FMetasoundFrontendGraphClass& InGraph,
+		const TArray<FMetasoundFrontendGraphClass>& InSubgraphs,
+		const TArray<FMetasoundFrontendClass>& InDependencies,
+		const FTopLevelAssetPath& InAssetPath,
+		TArrayView<const FGuid> InPageOrder)
 	{
 		// Create proxies before building graph
 		Frontend::FProxyDataCache ProxyDataCache;
@@ -1010,12 +1056,66 @@ PRAGMA_ENABLE_DEPRECATION_WARNINGS
 			ProxyDataCache.CreateAndCacheProxies(DependencyClass, InPageOrder);
 		}
 
-		return CreateGraph(InGraph, InSubgraphs, InDependencies, ProxyDataCache, InDebugAssetName, Frontend::CreateLocallyUniqueId(), InPageOrder);
+		return CreateGraph(InGraph, InSubgraphs, InDependencies, ProxyDataCache, InAssetPath, Frontend::CreateLocallyUniqueId(), InPageOrder);
+	}
+
+	TUniquePtr<FFrontendGraph> FGraphBuilder::CreateGraph(
+		const FMetasoundFrontendDocument& InDocument,
+		const Frontend::FProxyDataCache& InProxyDataCache,
+		const FTopLevelAssetPath& InAssetPath,
+		const FGuid InGraphId,
+		TArrayView<const FGuid> InPageOrder)
+	{
+		return CreateGraph(InDocument.RootGraph, InDocument.Subgraphs, InDocument.Dependencies, InProxyDataCache, InAssetPath, InGraphId, InPageOrder);
+	}
+
+	TUniquePtr<FFrontendGraph> FGraphBuilder::CreateGraph(const FMetasoundFrontendDocument& InDocument, const FString& InDebugAssetName, TArrayView<const FGuid> InPageOrder)
+	{
+PRAGMA_DISABLE_DEPRECATION_WARNINGS
+		const FName AssetName = FName(*InDebugAssetName);
+		return CreateGraph(InDocument, FTopLevelAssetPath(AssetName), InPageOrder);
+PRAGMA_ENABLE_DEPRECATION_WARNINGS
+	}
+
+	TUniquePtr<FFrontendGraph> FGraphBuilder::CreateGraph(
+		const FMetasoundFrontendDocument& InDocument,
+		const Frontend::FProxyDataCache& InProxies,
+		const FString& InDebugAssetName,
+		const FGuid InGraphId,
+		TArrayView<const FGuid> InPageOrder)
+	{
+PRAGMA_DISABLE_DEPRECATION_WARNINGS
+		const FName AssetName = FName(*InDebugAssetName);
+		return CreateGraph(InDocument, InProxies, FTopLevelAssetPath(AssetName), InGraphId, InPageOrder);
+PRAGMA_ENABLE_DEPRECATION_WARNINGS
+	}
+
+	TUniquePtr<FFrontendGraph> FGraphBuilder::CreateGraph(
+		const FMetasoundFrontendGraphClass& InGraph,
+		const TArray<FMetasoundFrontendGraphClass>& InSubgraphs,
+		const TArray<FMetasoundFrontendClass>& InDependencies,
+		const FString& InDebugAssetName,
+		TArrayView<const FGuid> InPageOrder)
+	{
+PRAGMA_DISABLE_DEPRECATION_WARNINGS
+		const FName AssetName = FName(*InDebugAssetName);
+		return CreateGraph(InGraph, InSubgraphs, InDependencies, FTopLevelAssetPath(AssetName), InPageOrder);
+PRAGMA_ENABLE_DEPRECATION_WARNINGS
 	}
 
-	TUniquePtr<FFrontendGraph> FGraphBuilder::CreateGraph(const FMetasoundFrontendDocument& InDocument, const Frontend::FProxyDataCache& InProxyDataCache, const FString& InDebugAssetName, const FGuid InGraphId, TArrayView<const FGuid> InPageOrder)
+	TUniquePtr<FFrontendGraph> FGraphBuilder::CreateGraph(
+		const FMetasoundFrontendGraphClass& InGraph,
+		const TArray<FMetasoundFrontendGraphClass>& InSubgraphs,
+		const TArray<FMetasoundFrontendClass>& InDependencies,
+		const Frontend::FProxyDataCache& InProxyDataCache,
+		const FString& InDebugAssetName,
+		const FGuid InGraphId,
+		TArrayView<const FGuid> InPageOrder)
 	{
-		return CreateGraph(InDocument.RootGraph, InDocument.Subgraphs, InDocument.Dependencies, InProxyDataCache, InDebugAssetName, InGraphId, InPageOrder);
+PRAGMA_DISABLE_DEPRECATION_WARNINGS
+		const FName AssetName = FName(*InDebugAssetName);
+		return CreateGraph(InGraph, InSubgraphs, InDependencies, InProxyDataCache, FTopLevelAssetPath(AssetName), InGraphId, InPageOrder);
+PRAGMA_ENABLE_DEPRECATION_WARNINGS
 	}
 }
 
diff --git a/Engine/Plugins/Runtime/Metasound/Source/MetasoundFrontend/Private/MetasoundFrontendNodeClassRegistry.cpp b/Engine/Plugins/Runtime/Metasound/Source/MetasoundFrontend/Private/MetasoundFrontendNodeClassRegistry.cpp
--- a/Engine/Plugins/Runtime/Metasound/Source/MetasoundFrontend/Private/MetasoundFrontendNodeClassRegistry.cpp
+++ b/Engine/Plugins/Runtime/Metasound/Source/MetasoundFrontend/Private/MetasoundFrontendNodeClassRegistry.cpp
@@ -346,7 +346,7 @@ PRAGMA_ENABLE_DEPRECATION_WARNINGS
 		using namespace RegistryPrivate;
 
 		METASOUND_TRACE_CPUPROFILER_EVENT_SCOPE(Metasound::FNodeClassRegistry::BuildAndRegisterGraphFromDocument);
-		METASOUND_TRACE_CPUPROFILER_EVENT_SCOPE_TEXT(*FString::Printf(TEXT("Metasound::FNodeClassRegistry::BuildAndRegisterGraphFromDocument asset %s"), *AssetPath.ToString()));
+		METASOUND_TRACE_CPUPROFILER_EVENT_SCOPE_TEXT(*FString::Printf(TEXT("Metasound::FNodeClassRegistry::BuildAndRegisterGraphFromDocument asset %s"), *AssetPath.GetAssetName().ToString()));
 
 		FGuid AssetClassID;
 		if (IMetaSoundAssetManager* AssetManager = IMetaSoundAssetManager::Get())
@@ -360,7 +360,7 @@ PRAGMA_ENABLE_DEPRECATION_WARNINGS
 		}
 
 		// Use the asset class id for the graph id because it should be locally unique per asset.
-		TUniquePtr<FFrontendGraph> FrontendGraph = Frontend::FGraphBuilder::CreateGraph(InDocument, InProxyDataCache, AssetPath.ToString(), /*GraphId=*/AssetClassID, InPageOrder);
+		TUniquePtr<FFrontendGraph> FrontendGraph = Frontend::FGraphBuilder::CreateGraph(InDocument, InProxyDataCache, AssetPath, /*GraphId=*/ AssetClassID, InPageOrder);
 		if (!FrontendGraph.IsValid())
 		{
 			UE_LOG(LogMetaSound, Error, TEXT("Failed to build MetaSound graph in asset '%s'"), *AssetPath.ToString());
diff --git a/Engine/Plugins/Runtime/Metasound/Source/MetasoundFrontend/Public/MetasoundFrontendGraph.h b/Engine/Plugins/Runtime/Metasound/Source/MetasoundFrontend/Public/MetasoundFrontendGraph.h
--- a/Engine/Plugins/Runtime/Metasound/Source/MetasoundFrontend/Public/MetasoundFrontendGraph.h
+++ b/Engine/Plugins/Runtime/Metasound/Source/MetasoundFrontend/Public/MetasoundFrontendGraph.h
@@ -28,12 +28,15 @@ namespace Metasound
 	class FFrontendGraph : public FGraph
 	{
 		public:
+			UE_DEPRECATED(5.7, "Use constructor that provides a top-level asset path instead")
+			UE_API FFrontendGraph(const FString& InInstanceName, const FGuid& InInstanceID);
+
 			/** FFrontendGraph constructor.
 			 *
-			 * @parma InInstanceName - Name of this graph.
+			 * @parma InAssetPath - Path to this graph's asset.
 			 * @parma InInstanceID - ID of this graph.
 			 */
-			UE_API FFrontendGraph(const FString& InInstanceName, const FGuid& InInstanceID);
+			UE_API FFrontendGraph(const FTopLevelAssetPath& InAssetPath, const FGuid& InInstanceID);
 
 			virtual ~FFrontendGraph() = default;
 
diff --git a/Engine/Plugins/Runtime/Metasound/Source/MetasoundFrontend/Public/MetasoundFrontendGraphBuilder.h b/Engine/Plugins/Runtime/Metasound/Source/MetasoundFrontend/Public/MetasoundFrontendGraphBuilder.h
--- a/Engine/Plugins/Runtime/Metasound/Source/MetasoundFrontend/Public/MetasoundFrontendGraphBuilder.h
+++ b/Engine/Plugins/Runtime/Metasound/Source/MetasoundFrontend/Public/MetasoundFrontendGraphBuilder.h
@@ -16,6 +16,7 @@
 #include "Misc/Guid.h"
 #include "Templates/SharedPointer.h"
 #include "Templates/UniquePtr.h"
+#include "UObject/TopLevelAssetPath.h"
 
 #define UE_API METASOUNDFRONTEND_API
 
@@ -39,14 +40,20 @@ namespace Metasound::Frontend
 
 		struct FCreateNodeParams
 		{
-			/** Debug name used for logging errors and warnings */
-			const FString& DebugAssetName;
+#if WITH_EDITORONLY_DATA
+			UE_DEPRECATED(5.7, "Deprecated in favor of AssetPath")
+			const FString DebugAssetName = { };
+#endif // WITH_EDITORONLY_DATA
+
 			/** The Frontend representation of the node being created. */
-			const FMetasoundFrontendNode& FrontendNode; 
+			const FMetasoundFrontendNode& FrontendNode;
+
 			/** The Frontend Node Class representing the node. */
-			const FMetasoundFrontendClass& FrontendNodeClass; 
+			const FMetasoundFrontendClass& FrontendNodeClass;
+
 			/** The Frontend Graph which contains the Frontend Node being created. */
 			const FMetasoundFrontendGraph& OwningFrontendGraph;
+
 			/** The Frontend Graph Class which contains the OwningFrontendGraph. */
 			const FMetasoundFrontendGraphClass& OwningFrontendGraphClass;
 
@@ -64,7 +71,10 @@ namespace Metasound::Frontend
 
 			/** An array of preferred pages where most preferred pages occur before
 			 * less preferred. */
-			TArrayView<const FGuid> PageOrder = {};
+			TArrayView<const FGuid> PageOrder = { };
+
+			/** AssetPath to graph of which the node is being added to. */
+			const FTopLevelAssetPath& OwningAssetPath;
 		};
 
 		/** Create a INode from a document representation of the node. 
@@ -73,21 +83,24 @@ namespace Metasound::Frontend
 		 */
 		static UE_API TUniquePtr<INode> CreateNode(const FCreateNodeParams& InParams);
 
-		/** Check that all dependencies are C++ class dependencies. 
-		 * 
-		 * @param InDocument - Document containing dependencies.
-		 *
-		 * @return True if all dependencies are C++ classes. False otherwise.
-		 */
-
 		/* Create a FFrontendGraph from a FMetasoundFrontendDocument.*/
-		static UE_API TUniquePtr<FFrontendGraph> CreateGraph(const FMetasoundFrontendDocument& InDocument, const FString& InDebugAssetName, TArrayView<const FGuid> InPageOrder = {});
+		static UE_API TUniquePtr<FFrontendGraph> CreateGraph(const FMetasoundFrontendDocument& InDocument, const FTopLevelAssetPath& InAssetPath, TArrayView<const FGuid> InPageOrder = { });
 
 		/* Create a FFrontendGraph from a FMetasoundFrontendDocument retrieving proxies from a FProxyDataCache.*/
-		static UE_API TUniquePtr<FFrontendGraph> CreateGraph(const FMetasoundFrontendDocument& InDocument, const Frontend::FProxyDataCache& InProxies, const FString& InDebugAssetName, const FGuid InGraphId = Frontend::CreateLocallyUniqueId(), TArrayView<const FGuid> InPageOrder = {});
+		static UE_API TUniquePtr<FFrontendGraph> CreateGraph(
+			const FMetasoundFrontendDocument& InDocument,
+			const Frontend::FProxyDataCache& InProxies,
+			const FTopLevelAssetPath& InAssetPath,
+			const FGuid InGraphId = Frontend::CreateLocallyUniqueId(),
+			TArrayView<const FGuid> InPageOrder = { });
 
 		/* Create a FFrontendGraph from a FMetasoundFrontendDocument and subobjects.*/
-		static UE_API TUniquePtr<FFrontendGraph> CreateGraph(const FMetasoundFrontendGraphClass& InGraph, const TArray<FMetasoundFrontendGraphClass>& InSubgraphs, const TArray<FMetasoundFrontendClass>& InDependencies, const FString& InDebugAssetName, TArrayView<const FGuid> InPageOrder = {});
+		static UE_API TUniquePtr<FFrontendGraph> CreateGraph(
+			const FMetasoundFrontendGraphClass& InGraph,
+			const TArray<FMetasoundFrontendGraphClass>& InSubgraphs,
+			const TArray<FMetasoundFrontendClass>& InDependencies,
+			const FTopLevelAssetPath& InAssetPath,
+			TArrayView<const FGuid> InPageOrder = { });
 
 		/* Create a FFrontendGraph from a FMetasoundFrontendDocument and subobjects, retrieving proxies from a FProxyDataCache.*/
 		static UE_API TUniquePtr<FFrontendGraph> CreateGraph(
@@ -95,10 +108,39 @@ namespace Metasound::Frontend
 			const TArray<FMetasoundFrontendGraphClass>& InSubgraphs,
 			const TArray<FMetasoundFrontendClass>& InDependencies,
 			const Frontend::FProxyDataCache& InProxyDataCache,
-			const FString& InDebugAssetName,
+			const FTopLevelAssetPath& InAssetPath,
 			const FGuid InGraphId = Frontend::CreateLocallyUniqueId(), 
-			TArrayView<const FGuid> InPageOrder = {});
+			TArrayView<const FGuid> InPageOrder = { });
+
+		// DEPRECATED API BELOW
+		UE_DEPRECATED(5.7, "Use overload providing TopLevelAssetPath instead of DebugAssetName string")
+		static UE_API TUniquePtr<FFrontendGraph> CreateGraph(const FMetasoundFrontendDocument& InDocument, const FString& InDebugAssetName, TArrayView<const FGuid> InPageOrder = { });
 
+		UE_DEPRECATED(5.7, "Use overload providing TopLevelAssetPath instead of DebugAssetName string")
+		static UE_API TUniquePtr<FFrontendGraph> CreateGraph(
+			const FMetasoundFrontendDocument& InDocument,
+			const Frontend::FProxyDataCache& InProxies,
+			const FString& InDebugAssetName,
+			const FGuid InGraphId = Frontend::CreateLocallyUniqueId(),
+			TArrayView<const FGuid> InPageOrder = { });
+
+		UE_DEPRECATED(5.7, "Use overload providing TopLevelAssetPath instead of DebugAssetName string")
+		static UE_API TUniquePtr<FFrontendGraph> CreateGraph(
+			const FMetasoundFrontendGraphClass& InGraph,
+			const TArray<FMetasoundFrontendGraphClass>& InSubgraphs,
+			const TArray<FMetasoundFrontendClass>& InDependencies,
+			const FString& InDebugAssetName,
+			TArrayView<const FGuid> InPageOrder = { });
+
+		UE_DEPRECATED(5.7, "Use overload providing TopLevelAssetPath instead of DebugAssetName string")
+		static UE_API TUniquePtr<FFrontendGraph> CreateGraph(
+			const FMetasoundFrontendGraphClass& InGraph,
+			const TArray<FMetasoundFrontendGraphClass>& InSubgraphs,
+			const TArray<FMetasoundFrontendClass>& InDependencies,
+			const Frontend::FProxyDataCache& InProxyDataCache,
+			const FString& InDebugAssetName,
+			const FGuid InGraphId = Frontend::CreateLocallyUniqueId(),
+			TArrayView<const FGuid> InPageOrder = { });
 	};
 }
 
diff --git a/Engine/Plugins/Runtime/Metasound/Source/MetasoundGenerator/Private/MetasoundGenerator.cpp b/Engine/Plugins/Runtime/Metasound/Source/MetasoundGenerator/Private/MetasoundGenerator.cpp
--- a/Engine/Plugins/Runtime/Metasound/Source/MetasoundGenerator/Private/MetasoundGenerator.cpp
+++ b/Engine/Plugins/Runtime/Metasound/Source/MetasoundGenerator/Private/MetasoundGenerator.cpp
@@ -49,7 +49,7 @@ namespace Metasound
 	namespace MetasoundGeneratorPrivate
 	{
 #if ENABLE_METASOUNDGENERATOR_INVALID_SAMPLE_VALUE_LOGGING
-		void LogInvalidAudioSampleValues(const FString& InMetaSoundName, const TArray<FAudioBufferReadRef>& InAudioBuffers)
+		void LogInvalidAudioSampleValues(const FTopLevelAssetPath& InAssetPath, const TArray<FAudioBufferReadRef>& InAudioBuffers)
 		{
 			if (ConsoleVariables::bEnableMetaSoundGeneratorNonFiniteLogging || ConsoleVariables::bEnableMetaSoundGeneratorInvalidSampleValueLogging)
 			{
@@ -66,7 +66,7 @@ namespace Metasound
 					{
 						if (!FMath::IsFinite(Data[i]))
 						{
-							UE_LOG(LogMetaSound, Error, TEXT("Found non-finite sample (%f) in channel %d of MetaSound %s"), Data[i], ChannelIndex, *InMetaSoundName);
+							UE_LOG(LogMetaSound, Error, TEXT("Found non-finite sample (%f) in channel %d of MetaSound %s"), Data[i], ChannelIndex, *InAssetPath.ToString());
 							break;
 						}
 					}
@@ -78,8 +78,7 @@ namespace Metasound
 						const float MaxAbsValue = Audio::ArrayMaxAbsValue(Buffer);
 						if (MaxAbsValue > Threshold)
 						{
-							UE_LOG(LogMetaSound, Warning, TEXT("Found sample (absolute value: %f) exceeding threshold (%f) in channel %d of MetaSound %s"), MaxAbsValue, Threshold, ChannelIndex, *InMetaSoundName);
-
+							UE_LOG(LogMetaSound, Warning, TEXT("Found sample (absolute value: %f) exceeding threshold (%f) in channel %d of MetaSound %s"), MaxAbsValue, Threshold, ChannelIndex, *InAssetPath.ToString());
 						}
 					}
 				}
@@ -253,10 +252,10 @@ namespace Metasound
 		}
 	}
 
-	void FMetasoundGeneratorInitParams::Reset(FMetasoundGeneratorInitParams& InParams)
+	void FGeneratorInitParams::Reset(FGeneratorInitParams& InParams)
 	{
 		InParams.Environment = {};
-		InParams.MetaSoundName = {};
+		InParams.AssetPath = {};
 		InParams.AudioOutputNames = {};
 		InParams.DefaultParameters = {};
 		InParams.DataChannel.Reset();
@@ -265,7 +264,7 @@ namespace Metasound
 
 	void FMetasoundDynamicGraphGeneratorInitParams::Reset(FMetasoundDynamicGraphGeneratorInitParams& InParams)
 	{
-		FMetasoundGeneratorInitParams::Reset(InParams);
+		FGeneratorInitParams::Reset(InParams);
 		InParams.TransformQueue.Reset();
 	}
 
@@ -274,22 +273,11 @@ namespace Metasound
 #if ENABLE_METASOUND_GENERATOR_INSTANCE_COUNTING
 		, InstanceCounter(FModuleManager::GetModuleChecked<FMetasoundGeneratorModule>("MetasoundGenerator").GetOperatorInstanceCounterManager())
 #endif // if ENABLE_METASOUND_GENERATOR_INSTANCE_COUNTING
-		, bIsFinishTriggered(false)
-		, bIsFinished(false)
-		, bPendingGraphTrigger(true)
-		, bIsNewGraphPending(false)
-		, bIsWaitingForFirstGraph(true)
-		, NumChannels(0)
-		, NumFramesPerExecute(InOperatorSettings.GetNumFramesPerBlock())
-		, NumSamplesPerExecute(0)
 		, OnFinishedTriggerRef(FTriggerWriteRef::CreateNew(InOperatorSettings))
-		, RenderTime(0.0)
 #if ENABLE_METASOUND_GENERATOR_RENDER_TIMING
-		, bDoRuntimeRenderTiming(true)
-#else
-		, bDoRuntimeRenderTiming(false)
-#endif // if ENABLE_METASOUND_GENERATOR_RENDER_TIMING
-	    , RelativeRenderCost(1.f)
+		, bDoRuntimeRenderTiming(1)
+#endif // ENABLE_METASOUND_GENERATOR_RENDER_TIMING
+		, NumFramesPerExecute(InOperatorSettings.GetNumFramesPerBlock())
 	{
 	}
 
@@ -311,14 +299,14 @@ namespace Metasound
 		return OnSetGraph.Add(Delegate);
 	}
 
-	void FMetasoundGenerator::InitBase(FMetasoundGeneratorInitParams& InInitParams)
+	void FMetasoundGenerator::InitBase(FGeneratorInitParams& InInitParams)
 	{
-		MetasoundName = InInitParams.MetaSoundName;
+		AssetPath = InInitParams.AssetPath;
 		NumChannels = InInitParams.AudioOutputNames.Num();
 		NumSamplesPerExecute = NumChannels * NumFramesPerExecute;
 
 #if ENABLE_METASOUND_GENERATOR_INSTANCE_COUNTING
-		InstanceCounter.Init(MetasoundName);
+		InstanceCounter.Init(AssetPath);
 #endif // if ENABLE_METASOUND_GENERATOR_INSTANCE_COUNTING
 
 		// Data channels
@@ -581,7 +569,7 @@ PRAGMA_ENABLE_DEPRECATION_WARNINGS
 	int32 FMetasoundGenerator::OnGenerateAudio(float* OutAudio, int32 NumSamplesRemaining)
 	{
 		METASOUND_LLM_SCOPE;
-		METASOUND_TRACE_CPUPROFILER_EVENT_SCOPE_TEXT(*FString::Printf(TEXT("MetasoundGenerator::OnGenerateAudio %s"), *MetasoundName));
+		METASOUND_TRACE_CPUPROFILER_EVENT_SCOPE_TEXT(*FString::Printf(TEXT("MetasoundGenerator::OnGenerateAudio %s"), *AssetPath.GetAssetName().ToString()));
 		CSV_SCOPED_TIMING_STAT(Audio_Metasound, OnGenerateAudio);
 
 		// Defer finishing the metasound generator one block
@@ -602,7 +590,7 @@ PRAGMA_ENABLE_DEPRECATION_WARNINGS
 		// Output silent audio if we're still building a graph
 		if (bIsWaitingForFirstGraph)
 		{
-			METASOUND_TRACE_CPUPROFILER_EVENT_SCOPE_TEXT(*FString::Printf(TEXT("MetasoundGenerator::OnGenerateAudio::MissedRenderDeadline %s"), *MetasoundName));
+			METASOUND_TRACE_CPUPROFILER_EVENT_SCOPE_TEXT(*FString::Printf(TEXT("MetasoundGenerator::OnGenerateAudio::MissedRenderDeadline %s"), *AssetPath.GetAssetName().ToString()));
 			CSV_CUSTOM_STAT(Audio_Metasound, WaitingForGraphBuildOnPlaybackCount, 1, ECsvCustomStatOp::Accumulate);
 			FMemory::Memset(OutAudio, 0, sizeof(float)* NumSamplesRemaining);
 			return NumSamplesRemaining;
@@ -636,7 +624,7 @@ PRAGMA_ENABLE_DEPRECATION_WARNINGS
 			OverflowBuffer.RemoveAtSwap(0 /* Index */, NumSamplesWritten /* Count */, EAllowShrinking::No);
 		}
 
-		while (NumSamplesRemaining > 0 && !bIsFinishTriggered)
+		while (NumSamplesRemaining > 0 && bIsFinishTriggered == 0)
 		{
 			if (GraphRenderCost)
 			{
@@ -690,7 +678,7 @@ PRAGMA_ENABLE_DEPRECATION_WARNINGS
 				
 				if (ThisLoopNumSamplesWritten <= 0)
 				{
-					UE_LOG(LogMetasoundGenerator, Warning, TEXT("Terminating OnGenerateAudio() loop for MetaSound %s -- FillWithBuffer() returned %i"), *MetasoundName, ThisLoopNumSamplesWritten);
+					UE_LOG(LogMetasoundGenerator, Warning, TEXT("Terminating OnGenerateAudio() loop for MetaSound %s -- FillWithBuffer() returned %i"), *AssetPath.ToString(), ThisLoopNumSamplesWritten);
 					break;
 				}
 			}
@@ -768,7 +756,7 @@ PRAGMA_ENABLE_DEPRECATION_WARNINGS
 		
 		if (bIsFinishTriggered)
 		{
-		    FMemory::Memcpy(OutAudio, InBuffer.GetData(), MaxNumOutputSamples * sizeof(float));
+			FMemory::Memcpy(OutAudio, InBuffer.GetData(), MaxNumOutputSamples * sizeof(float));
 			return NumSamplesToCopy;
 		}
 		
@@ -781,13 +769,13 @@ PRAGMA_ENABLE_DEPRECATION_WARNINGS
 			}
 			else
 			{
-				bIsFinishTriggered = true;
+				bIsFinishTriggered = 1;
 			}
 		}
 		
-		const int32 NumFramesToWrite = bIsFinishTriggered? FinishSample : NumSamplesToCopy;
+		const int32 NumFramesToWrite = bIsFinishTriggered ? FinishSample : NumSamplesToCopy;
 		const int32 NumFramesToZero = MaxNumOutputSamples - NumFramesToWrite;
-		FinishSample = bIsFinishTriggered? INDEX_NONE : FinishSample;
+		FinishSample = bIsFinishTriggered ? INDEX_NONE : FinishSample;
 		
 		FMemory::Memcpy(OutAudio, InBuffer.GetData(), NumFramesToWrite * sizeof(float));
 		FMemory::Memzero(OutAudio + NumFramesToWrite, NumFramesToZero * sizeof(float));
@@ -806,7 +794,7 @@ PRAGMA_ENABLE_DEPRECATION_WARNINGS
 		}
 
 #if ENABLE_METASOUNDGENERATOR_INVALID_SAMPLE_VALUE_LOGGING
-		MetasoundGeneratorPrivate::LogInvalidAudioSampleValues(MetasoundName, GraphOutputAudio);
+		MetasoundGeneratorPrivate::LogInvalidAudioSampleValues(AssetPath, GraphOutputAudio);
 #endif // if ENABLE_METASOUNDGENERATOR_INVALID_SAMPLE_VALUE_LOGGING
 
 		// Iterate over channels
@@ -891,7 +879,7 @@ PRAGMA_ENABLE_DEPRECATION_WARNINGS
 		}
 	}
 
-	FMetasoundConstGraphGenerator::FMetasoundConstGraphGenerator(FMetasoundGeneratorInitParams&& InInitParams)
+	FMetasoundConstGraphGenerator::FMetasoundConstGraphGenerator(FGeneratorInitParams&& InInitParams)
 	: FMetasoundGenerator(InInitParams.OperatorSettings)
 	{
 		Init(MoveTemp(InInitParams));
@@ -903,7 +891,7 @@ PRAGMA_ENABLE_DEPRECATION_WARNINGS
 	}
 
 	static FThreadSafeCounter NumActiveConstGraphs;
-	void FMetasoundConstGraphGenerator::Init(FMetasoundGeneratorInitParams&& InParams)
+	void FMetasoundConstGraphGenerator::Init(FGeneratorInitParams&& InParams)
 	{
 		NumActiveConstGraphs.Increment();
 		CSV_CUSTOM_STAT(Audio_Metasound, NumActiveConstGraphs, NumActiveConstGraphs.GetValue(), ECsvCustomStatOp::Set);
@@ -949,7 +937,7 @@ PRAGMA_ENABLE_DEPRECATION_WARNINGS
 		}
 	}
 
-	bool FMetasoundConstGraphGenerator::TryUseCachedOperator(FMetasoundGeneratorInitParams& InInitParams, bool bInTriggerGraph)
+	bool FMetasoundConstGraphGenerator::TryUseCachedOperator(FGeneratorInitParams& InInitParams, bool bInTriggerGraph)
 	{
 		using namespace MetasoundGeneratorPrivate;
 		METASOUND_TRACE_CPUPROFILER_EVENT_SCOPE(FMetasoundConstGraphGenerator::TryUseCachedOperator);
@@ -963,7 +951,7 @@ PRAGMA_ENABLE_DEPRECATION_WARNINGS
 			FOperatorAndInputs GraphOperatorAndInputs = OperatorPool->ClaimOperator(*OperatorPoolID, FOperatorContext::FromInitParams(InInitParams));
 			if (GraphOperatorAndInputs.Operator.IsValid())
 			{
-				UE_LOG(LogMetasoundGenerator, VeryVerbose, TEXT("Using cached operator %s for MetaSound %s"), *OperatorPoolID->ToString(), *InInitParams.MetaSoundName);
+				UE_LOG(LogMetasoundGenerator, VeryVerbose, TEXT("Using cached operator %s for MetaSound %s"), *OperatorPoolID->ToString(), *InInitParams.AssetPath.ToString());
 				bUseOperatorPool = true; // raise this flag to make sure we put the operator back in the pool (regardless of CVAR state)
 
 				// Apply and default inputs to the operator.
@@ -1038,7 +1026,7 @@ PRAGMA_ENABLE_DEPRECATION_WARNINGS
 		}
 	}
 
-	void FMetasoundConstGraphGenerator::BuildGraph(FMetasoundGeneratorInitParams&& InInitParams)
+	void FMetasoundConstGraphGenerator::BuildGraph(FGeneratorInitParams&& InInitParams)
 	{
 		const bool bBuildAsync = !InInitParams.bBuildSynchronous && Metasound::ConsoleVariables::bEnableAsyncMetaSoundGeneratorBuilder;
 
diff --git a/Engine/Plugins/Runtime/Metasound/Source/MetasoundGenerator/Private/MetasoundGeneratorBuilder.cpp b/Engine/Plugins/Runtime/Metasound/Source/MetasoundGenerator/Private/MetasoundGeneratorBuilder.cpp
--- a/Engine/Plugins/Runtime/Metasound/Source/MetasoundGenerator/Private/MetasoundGeneratorBuilder.cpp
+++ b/Engine/Plugins/Runtime/Metasound/Source/MetasoundGenerator/Private/MetasoundGeneratorBuilder.cpp
@@ -36,7 +36,7 @@ namespace Metasound
 		using namespace Frontend;
 		using namespace MetasoundGeneratorPrivate;
 
-		const FMetasoundGeneratorInitParams& InitParams = GetGeneratorParams();
+		const FGeneratorInitParams& InitParams = GetGeneratorParams();
 
 		if (InOpAndInputs.Operator.IsValid())
 		{
@@ -54,14 +54,14 @@ namespace Metasound
 		}
 		else 
 		{
-			UE_LOG(LogMetaSound, Error, TEXT("Failed to build Metasound operator from graph in MetasoundSource [%s]"), *InitParams.MetaSoundName);
+			UE_LOG(LogMetaSound, Error, TEXT("Failed to build Metasound operator from graph in MetasoundSource [%s]"), *InitParams.AssetPath.ToString());
 			// Set null generator data to inform that generator failed to build. 
 			// Otherwise, generator will continually wait for a new generator.
 			Generator->SetPendingGraphBuildFailed();
 		}
 	}
 
-	FAsyncMetaSoundBuilder::FAsyncMetaSoundBuilder(FMetasoundConstGraphGenerator* InGenerator, FMetasoundGeneratorInitParams&& InInitParams, bool bInTriggerGenerator)
+	FAsyncMetaSoundBuilder::FAsyncMetaSoundBuilder(FMetasoundConstGraphGenerator* InGenerator, FGeneratorInitParams&& InInitParams, bool bInTriggerGenerator)
 		: FAsyncMetaSoundBuilderBase(InGenerator, bInTriggerGenerator)
 		, InitParams(MoveTemp(InInitParams))
 	{
@@ -73,19 +73,19 @@ namespace Metasound
 		using namespace Frontend;
 		using namespace MetasoundGeneratorPrivate;
 
-		METASOUND_TRACE_CPUPROFILER_EVENT_SCOPE_TEXT(*FString::Printf(TEXT("AsyncMetaSoundBuilder::DoWork %s"), *InitParams.MetaSoundName));
+		METASOUND_TRACE_CPUPROFILER_EVENT_SCOPE_TEXT(*FString::Printf(TEXT("AsyncMetaSoundBuilder::DoWork %s"), *InitParams.AssetPath.GetAssetName().ToString()));
 
 		// Create an instance of the new graph
 		FBuildResults BuildResults;
 		FOperatorAndInputs GraphOperatorAndInputs = GeneratorBuilder::BuildGraphOperator(Generator->OperatorSettings, InitParams, BuildResults);
-		GeneratorBuilder::LogBuildErrors(InitParams.MetaSoundName, BuildResults);
+		GeneratorBuilder::LogBuildErrors(InitParams.AssetPath, BuildResults);
 
 		SetGraphOnGenerator(MoveTemp(GraphOperatorAndInputs), MoveTemp(BuildResults));
 
-		FMetasoundGeneratorInitParams::Reset(InitParams);
+		FGeneratorInitParams::Reset(InitParams);
 	}
 
-	const FMetasoundGeneratorInitParams& FAsyncMetaSoundBuilder::GetGeneratorParams() const
+	const FGeneratorInitParams& FAsyncMetaSoundBuilder::GetGeneratorParams() const
 	{
 		return InitParams;
 	}
@@ -104,7 +104,7 @@ namespace Metasound
 		using namespace MetasoundGeneratorPrivate;
 		using namespace DynamicGraph;
 
-		METASOUND_TRACE_CPUPROFILER_EVENT_SCOPE_TEXT(*FString::Printf(TEXT("AsyncMetaSoundBuilder::DoWork %s"), *InitParams.MetaSoundName));
+		METASOUND_TRACE_CPUPROFILER_EVENT_SCOPE_TEXT(*FString::Printf(TEXT("AsyncMetaSoundBuilder::DoWork %s"), *InitParams.AssetPath.GetAssetName().ToString()));
 
 		FDynamicOperatorUpdateCallbacks UpdateCallbacks;
 
@@ -124,21 +124,21 @@ namespace Metasound
 		// Create an instance of the new graph
 		FBuildResults BuildResults;
 		FOperatorAndInputs GraphOperatorAndInputs = GeneratorBuilder::BuildDynamicGraphOperator(Generator->OperatorSettings, InitParams, UpdateCallbacks, BuildResults);
-		GeneratorBuilder::LogBuildErrors(InitParams.MetaSoundName, BuildResults);
+		GeneratorBuilder::LogBuildErrors(InitParams.AssetPath, BuildResults);
 
 		SetGraphOnGenerator(MoveTemp(GraphOperatorAndInputs), MoveTemp(BuildResults));
 
 		FMetasoundDynamicGraphGeneratorInitParams::Reset(InitParams);
 	}
 
-	const FMetasoundGeneratorInitParams& FAsyncDynamicMetaSoundBuilder::GetGeneratorParams() const
+	const FGeneratorInitParams& FAsyncDynamicMetaSoundBuilder::GetGeneratorParams() const
 	{
 		return InitParams;
 	}
 
 	namespace GeneratorBuilder
 	{
-		TArray<FAudioBufferReadRef> FindOutputAudioBuffers(const TArray<FVertexName>& InAudioVertexNames, const FVertexInterfaceData& InVertexData, const FOperatorSettings& InOperatorSettings, const FString& InMetaSoundName)
+		TArray<FAudioBufferReadRef> FindOutputAudioBuffers(const TArray<FVertexName>& InAudioVertexNames, const FVertexInterfaceData& InVertexData, const FOperatorSettings& InOperatorSettings, const FTopLevelAssetPath& InAssetPath)
 		{
 			TArray<FAudioBufferReadRef> OutputBuffers;
 
@@ -149,7 +149,7 @@ namespace Metasound
 			{
 				if (!OutputVertexData.IsVertexBound(AudioOutputName))
 				{
-					UE_LOG(LogMetaSound, Warning, TEXT("MetasoundSource [%s] does not contain audio output [%s] in output"), *InMetaSoundName, *AudioOutputName.ToString());
+					UE_LOG(LogMetaSound, Warning, TEXT("MetasoundSource [%s] does not contain audio output [%s] in output"), *InAssetPath.ToString(), *AudioOutputName.ToString());
 				}
 				OutputBuffers.Add(OutputVertexData.GetOrConstructDataReadReference<FAudioBuffer>(AudioOutputName, InOperatorSettings));
 			}
@@ -157,14 +157,14 @@ namespace Metasound
 			return OutputBuffers;
 		}
 
-		void LogBuildErrors(const FString& InMetaSoundName, const FBuildResults& InBuildResults)
+		void LogBuildErrors(const FTopLevelAssetPath& InAssetPath, const FBuildResults& InBuildResults)
 		{
 			// Log build errors
 			for (const IOperatorBuilder::FBuildErrorPtr& Error : InBuildResults.Errors)
 			{
 				if (Error.IsValid())
 				{
-					UE_LOG(LogMetaSound, Warning, TEXT("MetasoundSource [%s] build error [%s] \"%s\""), *InMetaSoundName, *(Error->GetErrorType().ToString()), *(Error->GetErrorDescription().ToString()));
+					UE_LOG(LogMetaSound, Warning, TEXT("MetasoundSource [%s] build error [%s] \"%s\""), *InAssetPath.ToString(), *(Error->GetErrorType().ToString()), *(Error->GetErrorDescription().ToString()));
 				}
 			}
 		}
@@ -178,7 +178,7 @@ namespace Metasound
 			return MakeUnique<FGraphAnalyzer>(InOperatorSettings, InstanceID, MoveTemp(InInternalDataReferences));
 		}
 
-		FInputVertexInterfaceData BuildGraphOperatorInputs(const FOperatorSettings& InOperatorSettings, FMetasoundGeneratorInitParams& InInitParams)
+		FInputVertexInterfaceData BuildGraphOperatorInputs(const FOperatorSettings& InOperatorSettings, FGeneratorInitParams& InInitParams)
 		{
 			METASOUND_TRACE_CPUPROFILER_EVENT_SCOPE(GeneratorBuilder::BuildGraphOperatorInputs);
 
@@ -200,7 +200,7 @@ namespace Metasound
 
 			if (!InInitParams.Graph)
 			{
-				UE_LOG(LogMetaSound, Error, TEXT("Unable to build graph operator inputs for null graph in MetaSoundSource [%s]"), *InInitParams.MetaSoundName);
+				UE_LOG(LogMetaSound, Error, TEXT("Unable to build graph operator inputs for null graph in MetaSoundSource [%s]"), *InInitParams.AssetPath.ToString());
 				return FInputVertexInterfaceData();
 			}
 
@@ -224,12 +224,12 @@ namespace Metasound
 					}
 					else if(MetaSoundParameterEnableWarningOnIgnoredParameterCVar)
 					{
-						UE_LOG(LogMetaSound, Warning, TEXT("Failed to create initial input data reference from parameter %s of type %s on graph in MetaSoundSource [%s]"), *ParamName.ToString(), *InputVertex->DataTypeName.ToString(), *InInitParams.MetaSoundName);
+						UE_LOG(LogMetaSound, Warning, TEXT("Failed to create initial input data reference from parameter %s of type %s on graph in MetaSoundSource [%s]"), *ParamName.ToString(), *InputVertex->DataTypeName.ToString(), *InInitParams.AssetPath.ToString());
 					}
 				}
 				else if(MetaSoundParameterEnableWarningOnIgnoredParameterCVar)
 				{
-					UE_LOG(LogMetaSound, Warning, TEXT("Failed to set initial input parameter %s on graph in MetaSoundSource [%s]"), *ParamName.ToString(), *InInitParams.MetaSoundName);
+					UE_LOG(LogMetaSound, Warning, TEXT("Failed to set initial input parameter %s on graph in MetaSoundSource [%s]"), *ParamName.ToString(), *InInitParams.AssetPath.ToString());
 				}
 			}
 
@@ -256,7 +256,7 @@ namespace Metasound
 			return MoveTemp(InputData);
 		}
 
-		FOperatorAndInputs BuildGraphOperator(const FOperatorSettings& InOperatorSettings, FMetasoundGeneratorInitParams& InInitParams, FBuildResults& OutBuildResults)
+		FOperatorAndInputs BuildGraphOperator(const FOperatorSettings& InOperatorSettings, FGeneratorInitParams& InInitParams, FBuildResults& OutBuildResults)
 		{
 			METASOUND_TRACE_CPUPROFILER_EVENT_SCOPE(GeneratorBuilder::BuildGraphOperator);
 
@@ -266,7 +266,7 @@ namespace Metasound
 
 			if (!InInitParams.Graph)
 			{
-				UE_LOG(LogMetaSound, Error, TEXT("Unable to build graph operator for null graph in MetaSoundSource [%s]"), *InInitParams.MetaSoundName);
+				UE_LOG(LogMetaSound, Error, TEXT("Unable to build graph operator for null graph in MetaSoundSource [%s]"), *InInitParams.AssetPath.ToString());
 				return OpAndInputs;
 			}
 			// Create an instance of the new graph operator
@@ -277,11 +277,7 @@ namespace Metasound
 
 			// asset-level trace asset path
 			METASOUND_DEBUG_DECLARE_SCOPE;
-			METASOUND_DEBUG_SET_ACTIVE_ASSET_SCOPE(
-				  FName(InInitParams.MetaSoundName)
-				, InInitParams.ClassName
-				, InInitParams.AssetPath
-			);
+			METASOUND_DEBUG_SET_ACTIVE_ASSET_SCOPE(InInitParams.AssetPath.GetAssetName(), InInitParams.ClassName, InInitParams.AssetPath.GetPackageName());
 			
 			OpAndInputs.Operator = Builder.BuildGraphOperator(BuildParams, OutBuildResults);
 			OpAndInputs.RenderCost = InInitParams.GraphRenderCost;
@@ -315,7 +311,7 @@ namespace Metasound
 			return OpAndInputs;
 		}
 
-		MetasoundGeneratorPrivate::FMetasoundGeneratorData BuildGeneratorData(const FOperatorSettings& InOperatorSettings, const FMetasoundGeneratorInitParams& InInitParams, FOperatorAndInputs&& InGraphOperatorAndInputs, TUniquePtr<Frontend::FGraphAnalyzer> InAnalyzer)
+		MetasoundGeneratorPrivate::FMetasoundGeneratorData BuildGeneratorData(const FOperatorSettings& InOperatorSettings, const FGeneratorInitParams& InInitParams, FOperatorAndInputs&& InGraphOperatorAndInputs, TUniquePtr<Frontend::FGraphAnalyzer> InAnalyzer)
 		{
 			using namespace Audio;
 			using namespace Frontend;
@@ -337,7 +333,7 @@ namespace Metasound
 			FTriggerWriteRef PlayTrigger = VertexData.GetInputs().GetOrCreateDefaultDataWriteReference<FTrigger>(SourceInterface::Inputs::OnPlay, InOperatorSettings);
 
 			// Get outputs
-			TArray<FAudioBufferReadRef> OutputBuffers = FindOutputAudioBuffers(InInitParams.AudioOutputNames, VertexData, InOperatorSettings, InInitParams.MetaSoundName);
+			TArray<FAudioBufferReadRef> OutputBuffers = FindOutputAudioBuffers(InInitParams.AudioOutputNames, VertexData, InOperatorSettings, InInitParams.AssetPath);
 			FTriggerReadRef FinishTrigger = TDataReadReferenceFactory<FTrigger>::CreateExplicitArgs(InOperatorSettings, false);
 
 			if (const FAnyDataReference* FinishTriggerRef = VertexData.GetOutputs().FindDataReference(SourceOneShotInterface::Outputs::OnFinished))
diff --git a/Engine/Plugins/Runtime/Metasound/Source/MetasoundGenerator/Private/MetasoundGeneratorBuilder.h b/Engine/Plugins/Runtime/Metasound/Source/MetasoundGenerator/Private/MetasoundGeneratorBuilder.h
--- a/Engine/Plugins/Runtime/Metasound/Source/MetasoundGenerator/Private/MetasoundGeneratorBuilder.h
+++ b/Engine/Plugins/Runtime/Metasound/Source/MetasoundGenerator/Private/MetasoundGeneratorBuilder.h
@@ -52,7 +52,7 @@ namespace Metasound
 		void SetGraphOnGenerator(FOperatorAndInputs&& InOpAndInputs, FBuildResults&& InBuildResults);
 
 	private:
-		virtual const FMetasoundGeneratorInitParams& GetGeneratorParams() const = 0;
+		virtual const FGeneratorInitParams& GetGeneratorParams() const = 0;
 
 	};
 
@@ -60,15 +60,15 @@ namespace Metasound
 	class FAsyncMetaSoundBuilder : public FAsyncMetaSoundBuilderBase
 	{
 	public:
-		FAsyncMetaSoundBuilder(FMetasoundConstGraphGenerator* InGenerator, FMetasoundGeneratorInitParams&& InInitParams, bool bInTriggerGenerator);
+		FAsyncMetaSoundBuilder(FMetasoundConstGraphGenerator* InGenerator, FGeneratorInitParams&& InInitParams, bool bInTriggerGenerator);
 
 		virtual ~FAsyncMetaSoundBuilder() = default;
 
 		void DoWork();
 
 	private:
-		virtual const FMetasoundGeneratorInitParams& GetGeneratorParams() const override;
-		FMetasoundGeneratorInitParams InitParams;
+		virtual const FGeneratorInitParams& GetGeneratorParams() const override;
+		FGeneratorInitParams InitParams;
 	};
 
 	// Task class for asynchronously building a FMetasoundDynamicGraphGenerator
@@ -82,7 +82,7 @@ namespace Metasound
 		void DoWork();
 
 	private:
-		virtual const FMetasoundGeneratorInitParams& GetGeneratorParams() const override;
+		virtual const FGeneratorInitParams& GetGeneratorParams() const override;
 		FMetasoundDynamicGraphGenerator* DynamicGenerator;
 		FMetasoundDynamicGraphGeneratorInitParams InitParams;
 	};
@@ -90,25 +90,25 @@ namespace Metasound
 	namespace GeneratorBuilder
 	{
 		// Finds an array of audio buffer read references from the VertexData
-		TArray<FAudioBufferReadRef> FindOutputAudioBuffers(const TArray<FVertexName>& InAudioVertexNames, const FVertexInterfaceData& InVertexData, const FOperatorSettings& InOperatorSettings, const FString& InMetaSoundName);
+		TArray<FAudioBufferReadRef> FindOutputAudioBuffers(const TArray<FVertexName>& InAudioVertexNames, const FVertexInterfaceData& InVertexData, const FOperatorSettings& InOperatorSettings, const FTopLevelAssetPath& InAssetPath);
 
 		// Log any errors in the build results
-		void LogBuildErrors(const FString& InMetaSoundName, const FBuildResults& InBuildResults);
+		void LogBuildErrors(const FTopLevelAssetPath& InAssetPath, const FBuildResults& InBuildResults);
 
 		// Build an analyzer for a graph
 		TUniquePtr<Frontend::FGraphAnalyzer> BuildGraphAnalyzer(TMap<FGuid, FDataReferenceCollection>&& InInternalDataReferences, const FMetasoundEnvironment& InEnvironment, const FOperatorSettings& InOperatorSettings);
 
 		// Instantiate the input data references for a generator.
-		FInputVertexInterfaceData BuildGraphOperatorInputs(const FOperatorSettings& InOperatorSettings, FMetasoundGeneratorInitParams& InInitParams);
+		FInputVertexInterfaceData BuildGraphOperatorInputs(const FOperatorSettings& InOperatorSettings, FGeneratorInitParams& InInitParams);
 		
 		// Build a standard graph operatpr
-		FOperatorAndInputs BuildGraphOperator(const FOperatorSettings& InOperatorSettings, FMetasoundGeneratorInitParams& InInitParams, FBuildResults& OutBuildResults);
+		FOperatorAndInputs BuildGraphOperator(const FOperatorSettings& InOperatorSettings, FGeneratorInitParams& InInitParams, FBuildResults& OutBuildResults);
 		
 		// Build a dynamic graph operator
 		FOperatorAndInputs BuildDynamicGraphOperator(const FOperatorSettings& InOperatorSettings, FMetasoundDynamicGraphGeneratorInitParams& InInitParams, const DynamicGraph::FDynamicOperatorUpdateCallbacks& InCallbacks, FBuildResults& OutBuildResults);
 
 		// Create all the data required to run a FMetasoundGenerator
-		MetasoundGeneratorPrivate::FMetasoundGeneratorData BuildGeneratorData(const FOperatorSettings& InOperatorSettings, const FMetasoundGeneratorInitParams& InInitParams, FOperatorAndInputs&& InGraphOperatorAndInputs, TUniquePtr<Frontend::FGraphAnalyzer> InAnalyzer);
+		MetasoundGeneratorPrivate::FMetasoundGeneratorData BuildGeneratorData(const FOperatorSettings& InOperatorSettings, const FGeneratorInitParams& InInitParams, FOperatorAndInputs&& InGraphOperatorAndInputs, TUniquePtr<Frontend::FGraphAnalyzer> InAnalyzer);
 
 		// Resets FInputVertexInterfaceData to it's initial default state using the default literals stored ther-in. If a parameter exists in InParmaeterOverrides, that value is used instead. 
 		void ResetGraphOperatorInputs(const FOperatorSettings& InOperatorSettings, TArray<FAudioParameter> InParameterOverrides, FInputVertexInterfaceData& InOutInterface);
diff --git a/Engine/Plugins/Runtime/Metasound/Source/MetasoundGenerator/Private/MetasoundGeneratorModule.cpp b/Engine/Plugins/Runtime/Metasound/Source/MetasoundGenerator/Private/MetasoundGeneratorModule.cpp
--- a/Engine/Plugins/Runtime/Metasound/Source/MetasoundGenerator/Private/MetasoundGeneratorModule.cpp
+++ b/Engine/Plugins/Runtime/Metasound/Source/MetasoundGenerator/Private/MetasoundGeneratorModule.cpp
@@ -7,6 +7,7 @@
 #include "HAL/IConsoleManager.h"
 #include "Logging/LogMacros.h"
 #include "MetasoundOperatorCache.h"
+#include "MetasoundOperatorCacheStatTracker.h"
 #include "Misc/CString.h"
 #include "Modules/ModuleInterface.h"
 #include "Modules/ModuleManager.h"
@@ -47,22 +48,17 @@ static FAutoConsoleCommand CommandMetaSoundExperimentalOperatorPoolSetMaxNumOper
 		if (OperatorPool.IsValid())
 		{
 			OperatorPool->SetMaxNumOperators(static_cast<uint32>(MaxNumOperators));
-			UE_LOG(LogMetasoundGenerator, Display, TEXT("Metasound operator cache size set to %d operators."), MaxNumOperators);
+			UE_LOG(LogMetasoundGenerator, Display, TEXT("Operator cache size set to %d operators."), MaxNumOperators);
 		}
 	})
 );
 
 namespace Metasound
 {
-	static const FString InstanceCounterCategory(TEXT("Metasound/Active_Generators"));
-
 	void FMetasoundGeneratorModule::StartupModule()
 	{
-		FOperatorPoolSettings OperatorPoolSettings;
-		OperatorPoolSettings.MaxNumOperators = 64;
-
-		OperatorPool = MakeShared<FOperatorPool>(OperatorPoolSettings);
-		OperatorInstanceCounterManager = MakeShared<FConcurrentInstanceCounterManager>(InstanceCounterCategory);
+		OperatorPool = MakeShared<FOperatorPool>();
+		OperatorInstanceCounterManager = MakeShared<FConcurrentInstanceCounterManager>(TEXT("MetaSound/Active_Generators"));
 
 #if CSV_PROFILER && !UE_BUILD_SHIPPING
 		CsvEndFrameDelegateHandle = FCsvProfiler::Get()->OnCSVProfileEndFrame().AddLambda([WeakCounterManager = OperatorInstanceCounterManager.ToWeakPtr()]
@@ -74,9 +70,9 @@ namespace Metasound
 
 			if (TSharedPtr<FConcurrentInstanceCounterManager> CounterManager = WeakCounterManager.Pin())
 			{
-				CounterManager->VisitStats([](const FName& StatName, int64 Value)
+				CounterManager->VisitStats([](const FTopLevelAssetPath& InAssetPath, int64 Value)
 				{
-					FCsvProfiler::RecordCustomStat(StatName, CSV_CATEGORY_INDEX(MetaSound_ActiveOperators), (int32)Value, ECsvCustomStatOp::Set);
+					Engine::RecordOperatorStat(InAssetPath, CSV_CATEGORY_INDEX(MetaSound_ActiveOperators), (int32)Value, ECsvCustomStatOp::Set);
 				});
 			}
 		});
diff --git a/Engine/Plugins/Runtime/Metasound/Source/MetasoundGenerator/Private/MetasoundInstanceCounter.cpp b/Engine/Plugins/Runtime/Metasound/Source/MetasoundGenerator/Private/MetasoundInstanceCounter.cpp
--- a/Engine/Plugins/Runtime/Metasound/Source/MetasoundGenerator/Private/MetasoundInstanceCounter.cpp
+++ b/Engine/Plugins/Runtime/Metasound/Source/MetasoundGenerator/Private/MetasoundInstanceCounter.cpp
@@ -9,22 +9,22 @@ namespace Metasound
 	{
 	}
 
-	void FConcurrentInstanceCounterManager::Increment(const FName& InstanceName)
+	void FConcurrentInstanceCounterManager::Increment(const FTopLevelAssetPath& InstancePath)
 	{
 		FScopeLock Lock(&MapCritSec);
-		GetOrAddStats(InstanceName).Increment();
+		GetOrAddStats(InstancePath).Increment();
 	}
 
-	void FConcurrentInstanceCounterManager::Decrement(const FName& InstanceName)
+	void FConcurrentInstanceCounterManager::Decrement(const FTopLevelAssetPath& InstancePath)
 	{
 		FScopeLock Lock(&MapCritSec);
-		GetOrAddStats(InstanceName).Decrement();
+		GetOrAddStats(InstancePath).Decrement();
 	}
 
-	int64 FConcurrentInstanceCounterManager::GetCountForName(const FName& InName)
+	int64 FConcurrentInstanceCounterManager::GetCountForPath(const FTopLevelAssetPath& InPath)
 	{
 		FScopeLock Lock(&MapCritSec);
-		if (FStats* Stats = StatsMap.Find(InName))
+		if (FStats* Stats = StatsMap.Find(InPath))
 		{
 			return Stats->GetCount();
 		}
@@ -32,10 +32,10 @@ namespace Metasound
 		return 0;
 	}
 
-	int64 FConcurrentInstanceCounterManager::GetPeakCountForName(const FName& InName)
+	int64 FConcurrentInstanceCounterManager::GetPeakCountForPath(const FTopLevelAssetPath& InPath)
 	{
 		FScopeLock Lock(&MapCritSec);
-		if (FStats* Stats = StatsMap.Find(InName))
+		if (FStats* Stats = StatsMap.Find(InPath))
 		{
 			return Stats->GetPeakCount();
 		}
@@ -43,19 +43,22 @@ namespace Metasound
 		return 0;
 	}
 
-	void FConcurrentInstanceCounterManager::VisitStats(TFunctionRef<void(const FName&, int64)> Visitor)
+	void FConcurrentInstanceCounterManager::VisitStats(TFunctionRef<void(const FTopLevelAssetPath&, int64)> Visitor)
 	{
 		FScopeLock Lock(&MapCritSec);
 
-		for (const TPair<FName, FStats>& Pair : StatsMap)
+		for (const TPair<FTopLevelAssetPath, FStats>& Pair : StatsMap)
 		{
 			Visitor(Pair.Key, Pair.Value.GetCount());
 		}
 	}
 
 #if COUNTERSTRACE_ENABLED
-	FConcurrentInstanceCounterManager::FStats::FStats(const FString& InName)
-	: TraceCounter(MakeUnique<FCountersTrace::FCounterInt>(TraceCounterNameType_Dynamic, *InName, TraceCounterDisplayHint_None))
+	FConcurrentInstanceCounterManager::FStats::FStats(const FString& InCategoryName, const FTopLevelAssetPath& InPath)
+	: TraceCounter(MakeUnique<FCountersTrace::FCounterInt>(
+		TraceCounterNameType_Dynamic,
+		*FString::Printf(TEXT("%s - %s"), *InCategoryName, *InPath.ToString()),
+		TraceCounterDisplayHint_None))
 	{
 	}
 #endif
@@ -89,20 +92,20 @@ namespace Metasound
 		return PeakCount;
 	}
 
-	FConcurrentInstanceCounterManager::FStats& FConcurrentInstanceCounterManager::GetOrAddStats(const FName& InstanceName)
+	FConcurrentInstanceCounterManager::FStats& FConcurrentInstanceCounterManager::GetOrAddStats(const FTopLevelAssetPath& InstancePath)
 	{
 		FScopeLock Lock(&MapCritSec);
 
 		// avoid re-constructing the trace counter string unless it's new
-		if (StatsMap.Contains(InstanceName))
+		if (StatsMap.Contains(InstancePath))
 		{
-			return StatsMap[InstanceName];
+			return StatsMap[InstancePath];
 		}
 
 #if COUNTERSTRACE_ENABLED
-		return StatsMap.Emplace(InstanceName, FString::Printf(TEXT("%s - %s"), *CategoryName, *InstanceName.ToString()));
+		return StatsMap.Emplace(InstancePath, FStats(CategoryName, InstancePath));
 #else
-		return StatsMap.Add(InstanceName);
+		return StatsMap.Add(InstancePath);
 #endif
 	}
 
@@ -112,37 +115,24 @@ namespace Metasound
 		check(ManagerPtr);
 	}
 
-	FConcurrentInstanceCounter::FConcurrentInstanceCounter(const FName& InName, TSharedPtr<FConcurrentInstanceCounterManager> InCounterManager)
-	: InstanceName(InName)
+	FConcurrentInstanceCounter::FConcurrentInstanceCounter(const FTopLevelAssetPath& InPath, TSharedPtr<FConcurrentInstanceCounterManager> InCounterManager)
+	: InstancePath(InPath)
 	, ManagerPtr(InCounterManager)
 	{
 		check(ManagerPtr);
-		GetManagerChecked().Increment(InstanceName);
-	}
-
-	FConcurrentInstanceCounter::FConcurrentInstanceCounter(const FString& InName, TSharedPtr<FConcurrentInstanceCounterManager> InCounterManager)
-	: InstanceName(InName)
-	, ManagerPtr(InCounterManager)
-	{
-		GetManagerChecked().Increment(InstanceName);
+		GetManagerChecked().Increment(InstancePath);
 	}
 
 	// dtor
 	FConcurrentInstanceCounter::~FConcurrentInstanceCounter()
 	{
-		GetManagerChecked().Decrement(InstanceName);
-	}
-
-	void FConcurrentInstanceCounter::Init(const FName& InName)
-	{
-		InstanceName = InName;
-		GetManagerChecked().Increment(InstanceName);
+		GetManagerChecked().Decrement(InstancePath);
 	}
 
-	void FConcurrentInstanceCounter::Init(const FString& InName)
+	void FConcurrentInstanceCounter::Init(const FTopLevelAssetPath& InPath)
 	{
-		InstanceName = FName(InName);
-		GetManagerChecked().Increment(InstanceName);
+		InstancePath = InPath;
+		GetManagerChecked().Increment(InstancePath);
 	}
 
 	FConcurrentInstanceCounterManager& FConcurrentInstanceCounter::GetManagerChecked()
diff --git a/Engine/Plugins/Runtime/Metasound/Source/MetasoundGenerator/Private/MetasoundOperatorCache.cpp b/Engine/Plugins/Runtime/Metasound/Source/MetasoundGenerator/Private/MetasoundOperatorCache.cpp
--- a/Engine/Plugins/Runtime/Metasound/Source/MetasoundGenerator/Private/MetasoundOperatorCache.cpp
+++ b/Engine/Plugins/Runtime/Metasound/Source/MetasoundGenerator/Private/MetasoundOperatorCache.cpp
@@ -17,8 +17,8 @@
 #include "Misc/ScopeLock.h"
 #include "Modules/ModuleManager.h"
 #include "ProfilingDebugging/CountersTrace.h"
-#include "Templates/UniquePtr.h"
 #include "ProfilingDebugging/CsvProfiler.h"
+#include "Templates/UniquePtr.h"
 
 namespace Metasound
 {
@@ -170,11 +170,49 @@ namespace Metasound
 			}
 		}
 #endif // #if METASOUND_OPERATORCACHEPROFILER_ENABLED
+
+#if CSV_PROFILER
+		template <typename TValue>
+		void RecordOperatorStat(const FTopLevelAssetPath& InAssetPath, int32 CategoryIndex, const TValue& Value, ECsvCustomStatOp StatOp)
+		{
+			// Transient MetaSound graphs are guaranteed to have unique names, whereas asset
+			// packages may have the same name but are guaranteed to have different package names.
+			if (const UPackage* TransientPackage = GetTransientPackage())
+			{
+				if (InAssetPath.GetPackageName() == TransientPackage->GetFName())
+				{
+					FCsvProfiler::Get()->RecordCustomStat(InAssetPath.GetAssetName(), CategoryIndex, Value, StatOp);
+					return;
+				}
+			}
+
+			FCsvProfiler::Get()->RecordCustomStat(InAssetPath.GetPackageName(), CategoryIndex, Value, StatOp);
+		}
+#endif // CSV_PROFILER
 	} // namespace OperatorPoolPrivate
 
+#if CSV_PROFILER
+	namespace Engine
+	{
+		void RecordOperatorStat(const FTopLevelAssetPath& InAssetPath, int32 CategoryIndex, int32 Value, ECsvCustomStatOp StatOp)
+		{
+			OperatorPoolPrivate::RecordOperatorStat(InAssetPath, CategoryIndex, Value, StatOp);
+		}
+
+		void RecordOperatorStat(const FTopLevelAssetPath& InAssetPath, int32 CategoryIndex, float Value, ECsvCustomStatOp StatOp)
+		{
+			OperatorPoolPrivate::RecordOperatorStat(InAssetPath, CategoryIndex, Value, StatOp);
+		}
+
+		void RecordOperatorStat(const FTopLevelAssetPath& InAssetPath, int32 CategoryIndex, double Value, ECsvCustomStatOp StatOp)
+		{
+			OperatorPoolPrivate::RecordOperatorStat(InAssetPath, CategoryIndex, Value, StatOp);
+		}
+	} // namespace Engine
+#endif // CSV_PROFILER
 
 	FOperatorBuildData::FOperatorBuildData(
-		  FMetasoundGeneratorInitParams&& InInitParams
+		  FGeneratorInitParams&& InInitParams
 		, Frontend::FGraphRegistryKey InRegistryKey
 		, FGuid InAssetClassID
 		, int32 InNumInstances
@@ -188,12 +226,12 @@ namespace Metasound
 	{
 	}
 
-	FOperatorContext FOperatorContext::FromInitParams(const FMetasoundGeneratorInitParams& InParams)
+	FOperatorContext FOperatorContext::FromInitParams(const FGeneratorInitParams& InParams)
 	{
 		return FOperatorContext
 		{
-			.GraphInstanceName = InParams.Graph ? InParams.Graph->GetInstanceName() : NAME_None,
-			.MetaSoundName = InParams.MetaSoundName
+			.GraphAssetPath = InParams.Graph ? InParams.Graph->GetAssetPath() : InParams.AssetPath,
+			.AssetPath = InParams.AssetPath
 		};
 	}
 
@@ -230,8 +268,8 @@ namespace Metasound
 	}
 
 
-	FOperatorPool::FOperatorPool(const FOperatorPoolSettings& InSettings)
-	: Settings(InSettings)
+	FOperatorPool::FOperatorPool(FOperatorPoolSettings InSettings)
+	: Settings(MoveTemp(InSettings))
 	, AsyncBuildPipe(UE_SOURCE_LOCATION)
 	{
 #if METASOUND_OPERATORCACHEPROFILER_ENABLED
@@ -244,13 +282,6 @@ namespace Metasound
 		StopAsyncTasks();
 	}
 
-	FOperatorAndInputs FOperatorPool::ClaimOperator(const FGuid& InOperatorID)
-	{
-		return ClaimOperator(
-			FOperatorPoolEntryID{InOperatorID, FOperatorSettings{OperatorPoolPrivate::DefaultSampleRateForDeprecatedAPI, OperatorPoolPrivate::BlockRate}},
-			FOperatorContext{});
-	}
-
 	FOperatorAndInputs FOperatorPool::ClaimOperator(const FOperatorPoolEntryID& InOperatorID, const FOperatorContext& InContext)
 	{
 		FOperatorAndInputs OpAndInputs;
@@ -285,12 +316,6 @@ namespace Metasound
 
 		return OpAndInputs;
 	}
-
-	void FOperatorPool::AddOperator(const FGuid& InOperatorID, TUniquePtr<IOperator>&& InOperator, FInputVertexInterfaceData&& InputData)
-	{
-		const FOperatorPoolEntryID EntryID{InOperatorID, FOperatorSettings{OperatorPoolPrivate::DefaultSampleRateForDeprecatedAPI, OperatorPoolPrivate::BlockRate}};
-		AddOperator(EntryID, MoveTemp(InOperator), MoveTemp(InputData));
-	}
 	
 	void FOperatorPool::AddOperator(const FOperatorPoolEntryID& InOperatorID, TUniquePtr<IOperator>&& InOperator, FInputVertexInterfaceData&& InInputData, TSharedPtr<FGraphRenderCost>&& InRenderCost)
 	{
@@ -327,12 +352,6 @@ namespace Metasound
 		return true;
 	}
 
-	void FOperatorPool::AddOperator(const FGuid& InOperatorID, FOperatorAndInputs&& OperatorAndInputs)
-	{
-		const FOperatorPoolEntryID EntryID{InOperatorID, FOperatorSettings{OperatorPoolPrivate::DefaultSampleRateForDeprecatedAPI, OperatorPoolPrivate::BlockRate}};
-		AddOperator(EntryID, MoveTemp(OperatorAndInputs));
-	}
-
 	void FOperatorPool::AddOperator(const FOperatorPoolEntryID& InOperatorID, FOperatorAndInputs&& OperatorAndInputs)
 	{
 		using namespace UE::Tasks;
@@ -385,11 +404,6 @@ namespace Metasound
 		return bStopping.load();
 	}
 
-	void FOperatorPool::CancelAllBuildEvents()
-	{
-		StopAsyncTasks();
-	}
-
 	void FOperatorPool::StopAsyncTasks()
 	{
 		using namespace UE::Tasks;
@@ -526,7 +540,7 @@ namespace Metasound
 				}
 
 				FOperatorAndInputs OperatorAndInputs = GeneratorBuilder::BuildGraphOperator(PreCacheData->InitParams.OperatorSettings, PreCacheData->InitParams, BuildResults);
-				GeneratorBuilder::LogBuildErrors(PreCacheData->InitParams.MetaSoundName, BuildResults);
+				GeneratorBuilder::LogBuildErrors(PreCacheData->InitParams.AssetPath, BuildResults);
 
 				OperatorPool->AddOperatorInternal(EntryID, MoveTemp(OperatorAndInputs));
 				OperatorPool->AddAssetIdToGraphIdLookUpInternal(PreCacheData->AssetClassID, EntryID);
@@ -534,12 +548,6 @@ namespace Metasound
 		});
 	}
 
-	void FOperatorPool::TouchOperators(const FGuid& InOperatorID, int32 NumToTouch)
-	{
-		const FOperatorPoolEntryID EntryID{InOperatorID, FOperatorSettings{OperatorPoolPrivate::DefaultSampleRateForDeprecatedAPI, OperatorPoolPrivate::BlockRate}};
-		TouchOperators(EntryID, NumToTouch);
-	}
-
 	void FOperatorPool::TouchOperators(const FOperatorPoolEntryID& InOperatorID, int32 NumToTouch)
 	{
 		if (!IsStopping())
@@ -581,12 +589,6 @@ namespace Metasound
 		}
 	}
 
-	void FOperatorPool::RemoveOperatorsWithID(const FGuid& InOperatorID)
-	{
-		const FOperatorPoolEntryID EntryID{InOperatorID, FOperatorSettings{OperatorPoolPrivate::DefaultSampleRateForDeprecatedAPI, OperatorPoolPrivate::BlockRate}};
-		RemoveOperatorsWithID(EntryID);
-	}
-
 	void FOperatorPool::RemoveOperatorsWithID(const FOperatorPoolEntryID& InOperatorID)
 	{
 		using namespace UE::Tasks;
diff --git a/Engine/Plugins/Runtime/Metasound/Source/MetasoundGenerator/Private/MetasoundOperatorCacheStatTracker.cpp b/Engine/Plugins/Runtime/Metasound/Source/MetasoundGenerator/Private/MetasoundOperatorCacheStatTracker.cpp
--- a/Engine/Plugins/Runtime/Metasound/Source/MetasoundGenerator/Private/MetasoundOperatorCacheStatTracker.cpp
+++ b/Engine/Plugins/Runtime/Metasound/Source/MetasoundGenerator/Private/MetasoundOperatorCacheStatTracker.cpp
@@ -32,12 +32,7 @@ namespace Metasound::Engine
 			true,
 			TEXT("Record which metasounds incur a cache miss when building their graph.")
 		);
-
-		FString GetGraphName(FName GraphInstanceName)
-		{
-			return FPackageName::ObjectPathToPathWithinPackage(GraphInstanceName.ToString());
-		}
-	}
+	} // namespace Private
 
 	FOperatorCacheStatTracker::FOperatorCacheStatTracker()
 	{
@@ -71,39 +66,36 @@ namespace Metasound::Engine
 
 			// Show how much we're increasing the existing cache for this sound by.
 			UE_LOG(LogMetasoundGenerator, Log,
-				TEXT("Pre-cached Metasound: %s [Graph: %s]. Added %d instances, Total: %d."),
-				*BuildData.InitParams.MetaSoundName,
-				*Entry->GraphName.ToString(),
+				TEXT("Pre-cached MetaSound: %s [Graph: %s]. Added %d instances, Total: %d."),
+				*BuildData.InitParams.AssetPath.ToString(),
+				*Entry->GraphAssetPath.ToString(),
 				NumInstancesToBuild,
 				Entry->NumCacheSlots);
 		}
 		else
 		{
-			// Get the asset name from the package path
-			const FString GraphName = Private::GetGraphName(BuildData.InitParams.Graph->GetInstanceName());
-
 			FStatEntry& StatEntry = StatEntries.Add(EntryID, FStatEntry
 			{
-				.GraphName = FName(*GraphName),
+				.GraphAssetPath = BuildData.InitParams.Graph->GetAssetPath(),
 				.NumInstancesBuilt = NumInstancesToBuild,
 				.NumCacheSlots = NumInstancesToBuild
 			});
 
-			if (GraphName == BuildData.InitParams.MetaSoundName)
+			if (StatEntry.GraphAssetPath == BuildData.InitParams.AssetPath)
 			{
 				UE_LOG(LogMetasoundGenerator, Log,
-					TEXT("Pre-cached Metasound: %s. Requested: %d, Built: %d."),
-					*BuildData.InitParams.MetaSoundName,
+					TEXT("Pre-cached MetaSound: %s. Requested: %d, Built: %d."),
+					*BuildData.InitParams.AssetPath.ToString(),
 					BuildData.NumInstances,
 					NumInstancesToBuild);
 			}
 			else
 			{
-				// Include the parent graph so it's clearer which this contributes to.
+				// Include the parent graph if preset so it's clearer which entry the build instances contribute to.
 				UE_LOG(LogMetasoundGenerator, Log,
-					TEXT("Pre-cached Metasound: %s [Graph: %s] Requested: %d, Built: %d."),
-					*BuildData.InitParams.MetaSoundName,
-					*StatEntry.GraphName.ToString(),
+					TEXT("Pre-cached MetaSound: %s [Graph: %s] Requested: %d, Built: %d."),
+					*BuildData.InitParams.AssetPath.ToString(),
+					*StatEntry.GraphAssetPath.ToString(),
 					BuildData.NumInstances,
 					NumInstancesToBuild);
 			}
@@ -138,7 +130,7 @@ namespace Metasound::Engine
 				const int32 NumMissing = ExpectedCacheSlots - StatEntry.NumCacheSlots;
 				UE_LOG(LogMetasoundGenerator, Log,
 					TEXT("FOperatorCacheStatTracker detected a cache slot mismatch for %s. Have %d, expected %d. Updating to expected value."),
-					*StatEntry.GraphName.ToString(),
+					*StatEntry.GraphAssetPath.ToString(),
 					StatEntry.NumCacheSlots,
 					ExpectedCacheSlots);
 
@@ -158,10 +150,9 @@ namespace Metasound::Engine
 #if CSV_PROFILER_STATS 
 			if (Private::bMetasoundOperatorPoolCsvStatsEnabled &&
 				Private::CVarCacheMissCsvStatsEnabled->GetBool() &&
-				Context.GraphInstanceName != NAME_Name)
+				Context.GraphAssetPath.IsValid())
 			{
-				const FString GraphName = Private::GetGraphName(Context.GraphInstanceName);
-				FCsvProfiler::Get()->RecordCustomStat(*GraphName, CSV_CATEGORY_INDEX(Metasound_OperatorCacheMiss), 1, ECsvCustomStatOp::Accumulate);
+				RecordOperatorStat(Context.GraphAssetPath, CSV_CATEGORY_INDEX(Metasound_OperatorCacheMiss), 1, ECsvCustomStatOp::Accumulate);
 			}
 #endif // CSV_PROFILER_STATS 
 			return;
@@ -204,13 +195,13 @@ namespace Metasound::Engine
 
 			if (StatEntry->NumCacheSlots == 0)
 			{
-				UE_LOG(LogMetasoundGenerator, Log, TEXT("Evicted %s from the Operator Pool."), *StatEntry->GraphName.ToString())
+				UE_LOG(LogMetasoundGenerator, Log, TEXT("Evicted %s from the Operator Pool."), *StatEntry->GraphAssetPath.ToString())
 			}
 			else
 			{
 				UE_LOG(LogMetasoundGenerator, Log,
 					TEXT("Trimmed 1 instance of %s from the Operator Pool. %d instances remaining."),
-					*StatEntry->GraphName.ToString(),
+					*StatEntry->GraphAssetPath.ToString(),
 					StatEntry->NumCacheSlots);
 			}
 		}
@@ -235,7 +226,7 @@ namespace Metasound::Engine
 
 	void FOperatorCacheStatTracker::OnCsvProfileEndFrame()
 	{
-#if CSV_PROFILER_STATS 
+#if CSV_PROFILER_STATS
 		if (!Private::bMetasoundOperatorPoolCsvStatsEnabled)
 		{
 			return;
@@ -262,15 +253,13 @@ namespace Metasound::Engine
 			// Record cache utilization stats.
 			if (ensure(Entry.NumCacheSlots > 0))
 			{
-				const int32 NumAvailableIncache = Entry.NumAvailableInCache;
-				const int32 NumUsed = Entry.NumCacheSlots - NumAvailableIncache;
-				const float UtilizationRatio = static_cast<float>(NumUsed) / static_cast<float>(Entry.NumCacheSlots);
+				RecordOperatorStat(Entry.GraphAssetPath, CSV_CATEGORY_INDEX(MetaSound_AvailableCachedOperators), Entry.NumAvailableInCache, ECsvCustomStatOp::Set);
 
-				FCsvProfiler::Get()->RecordCustomStat(Entry.GraphName, CSV_CATEGORY_INDEX(MetaSound_AvailableCachedOperators), NumAvailableIncache, ECsvCustomStatOp::Set);
-				FCsvProfiler::Get()->RecordCustomStat(Entry.GraphName, CSV_CATEGORY_INDEX(MetaSound_OperatorCacheUtilization), UtilizationRatio, ECsvCustomStatOp::Set);
+				const float UtilizationRatio = (Entry.NumCacheSlots - Entry.NumAvailableInCache) / static_cast<float>(Entry.NumCacheSlots);
+				RecordOperatorStat(Entry.GraphAssetPath, CSV_CATEGORY_INDEX(MetaSound_OperatorCacheUtilization), UtilizationRatio, ECsvCustomStatOp::Set);
 			}
 		}
 #endif // CSV_PROFILER_STATS 
 	}
-} // namespace Metasound::Private
+} // namespace Metasound::Engine
 #endif // METASOUND_OPERATORCACHEPROFILER_ENABLED
diff --git a/Engine/Plugins/Runtime/Metasound/Source/MetasoundGenerator/Private/MetasoundOperatorCacheStatTracker.h b/Engine/Plugins/Runtime/Metasound/Source/MetasoundGenerator/Private/MetasoundOperatorCacheStatTracker.h
--- a/Engine/Plugins/Runtime/Metasound/Source/MetasoundGenerator/Private/MetasoundOperatorCacheStatTracker.h
+++ b/Engine/Plugins/Runtime/Metasound/Source/MetasoundGenerator/Private/MetasoundOperatorCacheStatTracker.h
@@ -35,13 +35,16 @@ namespace Metasound::Engine
 
 		struct FStatEntry
 		{
-			// Stored as an fname to avoid conversions when writing out csv stats.
-			FName GraphName;
+			// Stored as a TopLevelPath to avoid conversions writing out csv stats.
+			FTopLevelAssetPath GraphAssetPath;
+
 			// The actual number of instances we built.
 			int32 NumInstancesBuilt;
+
 			// The number of instances we have space to cache.
 			// This can be different from the number currently in the stack as those are removed when in use.
 			int32 NumCacheSlots;
+
 			// The number of operators sitting in the pool waiting to be used.
 			// Set exclusively by OnOperatorAdded/OnOperatorTrimmed.
 			int32 NumAvailableInCache = 0;
diff --git a/Engine/Plugins/Runtime/Metasound/Source/MetasoundGenerator/Public/MetasoundGenerator.h b/Engine/Plugins/Runtime/Metasound/Source/MetasoundGenerator/Public/MetasoundGenerator.h
--- a/Engine/Plugins/Runtime/Metasound/Source/MetasoundGenerator/Public/MetasoundGenerator.h
+++ b/Engine/Plugins/Runtime/Metasound/Source/MetasoundGenerator/Public/MetasoundGenerator.h
@@ -6,6 +6,7 @@
 #include "MetasoundInstanceCounter.h"
 #include "MetasoundOperatorBuilder.h"
 #include "MetasoundOperatorInterface.h"
+#include "MetasoundOperatorSettings.h"
 #include "MetasoundParameterPack.h"
 #include "MetasoundRouter.h"
 #include "MetasoundTrigger.h"
@@ -106,23 +107,30 @@ namespace Metasound
 		FOperatorSettings OperatorSettings;
 	};
 
-	// Struct needed for building the metasound graph
-	struct FMetasoundGeneratorInitParams
+	struct FGeneratorInitParams
 	{
 		FOperatorSettings OperatorSettings;
 		FOperatorBuilderSettings BuilderSettings;
 		TSharedPtr<const IGraph, ESPMode::ThreadSafe> Graph;
 		FMetasoundEnvironment Environment;
-		FString MetaSoundName;
+
 		TArray<FVertexName> AudioOutputNames;
 		TArray<FAudioParameter> DefaultParameters;
 		bool bBuildSynchronous = false;
 		TSharedPtr<TSpscQueue<FMetaSoundParameterTransmitter::FParameter>> DataChannel;
 		TSharedPtr<FGraphRenderCost> GraphRenderCost;
 		FName ClassName;
-		FName AssetPath;
+		FTopLevelAssetPath AssetPath;
 
-		static UE_API void Reset(FMetasoundGeneratorInitParams& InParams);
+		static UE_API void Reset(FGeneratorInitParams& InParams);
+	};
+
+	// Initialization params required for building a MetaSound generator
+	struct UE_DEPRECATED(5.7, "Use FGeneratorInitParams instead (MetaSoundName is deprecated in favor of AssetPath") FMetasoundGeneratorInitParams : public FGeneratorInitParams
+	{
+#if WITH_EDITORONLY_DATA
+		FString MetaSoundName = { };
+#endif // WITH_EDITORONLY_DATA
 	};
 
 	enum class EVertexInterfaceChangeType : uint8
@@ -147,8 +155,6 @@ namespace Metasound
 		using FOperatorUniquePtr = TUniquePtr<Metasound::IOperator>;
 		using FAudioBufferReadRef = Metasound::FAudioBufferReadRef;
 
-		FString MetasoundName;
-
 		const FOperatorSettings OperatorSettings;
 
 		/** Create the generator with a graph operator and an output audio reference.
@@ -291,7 +297,7 @@ namespace Metasound
 
 	protected:
 
-		UE_API void InitBase(FMetasoundGeneratorInitParams& InInitParams);
+		UE_API void InitBase(FGeneratorInitParams& InInitParams);
 
 
 		/** SetGraph directly sets graph. Callers must ensure that no race conditions exist. */
@@ -314,7 +320,6 @@ namespace Metasound
 		std::atomic<bool> bVertexInterfaceHasChanged{ false };
 
 	private:
-
 		friend class FAsyncMetaSoundBuilderBase;
 
 		UE_API void SetPendingGraphBuildFailed();
@@ -339,18 +344,6 @@ namespace Metasound
 
 		UE_API void HandleRenderTimingEnableDisable();
 
-		bool bIsGraphBuilding;
-		bool bIsFinishTriggered;
-		bool bIsFinished;
-		bool bPendingGraphTrigger;
-		bool bIsNewGraphPending;
-		bool bIsWaitingForFirstGraph;
-
-		int32 FinishSample = INDEX_NONE;
-		int32 NumChannels;
-		int32 NumFramesPerExecute;
-		int32 NumSamplesPerExecute;
-
 	protected:
 		FExecuter RootExecuter;
 		FVertexInterfaceData VertexInterfaceData;
@@ -360,37 +353,44 @@ namespace Metasound
 
 		// Triggered when metasound is finished
 		FTriggerReadRef OnFinishedTriggerRef;
+
+		MetasoundGeneratorPrivate::FParameterSetterSortedMap ParameterSetters;
+
+		// This map provides setters for all of the input nodes in the metasound graph. 
+		// It is used when processing named parameters in a parameter pack.
+		TMap<FName, MetasoundGeneratorPrivate::FParameterPackSetter> ParameterPackSetters;
+		TSharedPtr<FGraphRenderCost> GraphRenderCost;
+
 	private:
+		FTopLevelAssetPath AssetPath;
 
-		Audio::FAlignedFloatBuffer InterleavedAudioBuffer;
+		uint8 bIsGraphBuilding : 1 = 0;
+		uint8 bIsFinishTriggered : 1 = 0;
+		uint8 bIsFinished : 1 = 0;
+		uint8 bPendingGraphTrigger : 1 = 1;
+		uint8 bIsNewGraphPending : 1 = 0;
+		uint8 bIsWaitingForFirstGraph : 1 = 1;
+		uint8 bDoRuntimeRenderTiming : 1 = 0;
 
+		int32 FinishSample = INDEX_NONE;
+		int32 NumChannels = 0;
+		int32 NumFramesPerExecute = 0;
+		int32 NumSamplesPerExecute = 0;
+
+		Audio::FAlignedFloatBuffer InterleavedAudioBuffer;
 		Audio::FAlignedFloatBuffer OverflowBuffer;
 
 		FCriticalSection PendingGraphMutex;
 		TUniquePtr<MetasoundGeneratorPrivate::FMetasoundGeneratorData> PendingGraphData;
-
 		TUniquePtr<Frontend::FGraphAnalyzer> GraphAnalyzer;
-
 		TSharedPtr<TSpscQueue<FMetaSoundParameterTransmitter::FParameter>> ParameterQueue;
-	protected:
-
-		MetasoundGeneratorPrivate::FParameterSetterSortedMap ParameterSetters;
-	private:
 
 		// These next items are needed to provide a destination for the FAudioDevice, etc. to
 		// send parameter packs to. Every playing metasound will have a parameter destination
 		// that can accept parameter packs.
 		FSendAddress ParameterPackSendAddress;
 		TReceiverPtr<FMetasoundParameterStorageWrapper> ParameterPackReceiver;
-		
-	protected:
-		// This map provides setters for all of the input nodes in the metasound graph. 
-		// It is used when processing named parameters in a parameter pack.
-		TMap<FName, MetasoundGeneratorPrivate::FParameterPackSetter> ParameterPackSetters;
 
-		TSharedPtr<FGraphRenderCost> GraphRenderCost;
-		
-	private:
 		// While parameter packs may arrive via the IAudioParameterInterface system,
 		// a faster method of sending parameters is via the QueueParameterPack function 
 		// and this queue.
@@ -401,11 +401,10 @@ namespace Metasound
 
 		FOnSetGraph OnSetGraph;
 
-		double RenderTime;
-		bool bDoRuntimeRenderTiming;
+		double RenderTime = 0.0;
 		TUniquePtr<MetasoundGeneratorPrivate::FRenderTimer> RenderTimer;
 		
-		std::atomic<float> RelativeRenderCost;
+		std::atomic<float> RelativeRenderCost { 1.f };
 	};
 
 	/** FMetasoundConstGraphGenerator generates audio from a given metasound IOperator
@@ -415,18 +414,18 @@ namespace Metasound
 	{
 	public:
 
-		UE_API explicit FMetasoundConstGraphGenerator(FMetasoundGeneratorInitParams&& InParams);
+		UE_API explicit FMetasoundConstGraphGenerator(FGeneratorInitParams&& InParams);
 
 		UE_API explicit FMetasoundConstGraphGenerator(const FOperatorSettings& InOperatorSettings);
 
-		UE_API void Init(FMetasoundGeneratorInitParams&& InParams);
+		UE_API void Init(FGeneratorInitParams&& InParams);
 
 		UE_API virtual ~FMetasoundConstGraphGenerator() override;
 
 	private:
 		friend class FAsyncMetaSoundBuilder;
-		UE_API void BuildGraph(FMetasoundGeneratorInitParams&& InInitParams);
-		UE_API bool TryUseCachedOperator(FMetasoundGeneratorInitParams& InParams, bool bTriggerGenerator);
+		UE_API void BuildGraph(FGeneratorInitParams&& InInitParams);
+		UE_API bool TryUseCachedOperator(FGeneratorInitParams& InParams, bool bTriggerGenerator);
 		UE_API void ReleaseOperatorToCache();
 
 		TUniquePtr<FMetasoundEnvironment> EnvironmentPtr;
@@ -435,7 +434,7 @@ namespace Metasound
 		bool bUseOperatorPool = false;
 	};
 
-	struct FMetasoundDynamicGraphGeneratorInitParams : FMetasoundGeneratorInitParams
+	struct FMetasoundDynamicGraphGeneratorInitParams : FGeneratorInitParams
 	{
 		TSharedPtr<TSpscQueue<TUniquePtr<DynamicGraph::IDynamicOperatorTransform>>> TransformQueue;
 		static UE_API void Reset(FMetasoundDynamicGraphGeneratorInitParams& InParams);
diff --git a/Engine/Plugins/Runtime/Metasound/Source/MetasoundGenerator/Public/MetasoundGeneratorModule.h b/Engine/Plugins/Runtime/Metasound/Source/MetasoundGenerator/Public/MetasoundGeneratorModule.h
--- a/Engine/Plugins/Runtime/Metasound/Source/MetasoundGenerator/Public/MetasoundGeneratorModule.h
+++ b/Engine/Plugins/Runtime/Metasound/Source/MetasoundGenerator/Public/MetasoundGeneratorModule.h
@@ -2,8 +2,9 @@
 
 #pragma once
 
-#include "CoreMinimal.h"
+#include "Logging/LogMacros.h"
 #include "Modules/ModuleInterface.h"
+#include "Templates/SharedPointer.h"
 
 DECLARE_LOG_CATEGORY_EXTERN(LogMetasoundGenerator, Log, All);
 
@@ -19,6 +20,4 @@ namespace Metasound
 		virtual TSharedPtr<FOperatorPool> GetOperatorPool() = 0;
 		virtual TSharedPtr<FConcurrentInstanceCounterManager> GetOperatorInstanceCounterManager() = 0;
 	};
-}
-
-
+} // namespace Metasound
diff --git a/Engine/Plugins/Runtime/Metasound/Source/MetasoundGenerator/Public/MetasoundInstanceCounter.h b/Engine/Plugins/Runtime/Metasound/Source/MetasoundGenerator/Public/MetasoundInstanceCounter.h
--- a/Engine/Plugins/Runtime/Metasound/Source/MetasoundGenerator/Public/MetasoundInstanceCounter.h
+++ b/Engine/Plugins/Runtime/Metasound/Source/MetasoundGenerator/Public/MetasoundInstanceCounter.h
@@ -8,6 +8,7 @@
 #include "Stats/Stats.h"
 #include "Templates/UniquePtr.h"
 #include "UObject/NameTypes.h"
+#include "UObject/TopLevelAssetPath.h"
 
 // todo: comments
 
@@ -27,25 +28,24 @@ namespace Metasound
 	public:
 		FConcurrentInstanceCounterManager(const FString InCategoryName);
 
-		void Increment(const FName& InstanceName);
-		void Decrement(const FName& IntsanceName);
+		void Increment(const FTopLevelAssetPath& InstancePath);
+		void Decrement(const FTopLevelAssetPath& InstancePath);
 
-		int64 GetCountForName(const FName& InName);
-		int64 GetPeakCountForName(const FName& InName);
+		int64 GetCountForPath(const FTopLevelAssetPath& InPath);
+		int64 GetPeakCountForPath(const FTopLevelAssetPath& InPath);
 
-		void VisitStats(TFunctionRef<void(const FName&, int64)> Visitor);
+		void VisitStats(TFunctionRef<void(const FTopLevelAssetPath&, int64)> Visitor);
 
 	private:
 		struct FStats
 		{
 		public:
-			// ctor
 #if COUNTERSTRACE_ENABLED
-			FStats(const FString& InName);
+			FStats(const FString& InCategoryName, const FTopLevelAssetPath& InPath);
 #else
 			FStats() = default;
-#endif
-	
+#endif // COUNTERSTRACE_ENABLED
+
 			void Increment();
 			void Decrement();
 	
@@ -56,11 +56,11 @@ namespace Metasound
 			TUniquePtr<CounterType> TraceCounter;
 			int64 PeakCount;
 	
-		}; // struct FStats
+		};
 
-		FStats& GetOrAddStats(const FName& InstanceName);
+		FStats& GetOrAddStats(const FTopLevelAssetPath& AssetPath);
 
-		TMap<FName, FStats> StatsMap;
+		TMap<FTopLevelAssetPath, FStats> StatsMap;
 		FCriticalSection MapCritSec;
 		FString CategoryName;
 
@@ -74,20 +74,18 @@ namespace Metasound
 		// ctor
 		FConcurrentInstanceCounter() = delete; // must init ManagerPtr
 		FConcurrentInstanceCounter(TSharedPtr<FConcurrentInstanceCounterManager> InCounterManager);
-		FConcurrentInstanceCounter(const FName& InName, TSharedPtr<FConcurrentInstanceCounterManager> InCounterManager);
-		FConcurrentInstanceCounter(const FString& InName, TSharedPtr<FConcurrentInstanceCounterManager> InCounterManager);
+		FConcurrentInstanceCounter(const FTopLevelAssetPath& InPath, TSharedPtr<FConcurrentInstanceCounterManager> InCounterManager);
 	
 		// dtor
 		virtual ~FConcurrentInstanceCounter();
 	
 		// for non-RAII clients
-		void Init(const FName& InName);
-		void Init(const FString& InName);	
+		void Init(const FTopLevelAssetPath& InPath);
 	
 	private:
 		FConcurrentInstanceCounterManager& GetManagerChecked();
 
-		FName InstanceName;
+		FTopLevelAssetPath InstancePath;
 		TSharedPtr<FConcurrentInstanceCounterManager> ManagerPtr;
 	}; // class FConcurrentInstanceCounter
 } // namespace Metasound
diff --git a/Engine/Plugins/Runtime/Metasound/Source/MetasoundGenerator/Public/MetasoundOperatorCache.h b/Engine/Plugins/Runtime/Metasound/Source/MetasoundGenerator/Public/MetasoundOperatorCache.h
--- a/Engine/Plugins/Runtime/Metasound/Source/MetasoundGenerator/Public/MetasoundOperatorCache.h
+++ b/Engine/Plugins/Runtime/Metasound/Source/MetasoundGenerator/Public/MetasoundOperatorCache.h
@@ -8,6 +8,7 @@
 #include "MetasoundGenerator.h"
 #include "Misc/Guid.h"
 #include "HAL/CriticalSection.h"
+#include "ProfilingDebugging/CsvProfiler.h"
 #include "Tasks/Pipe.h"
 #include "Tasks/Task.h"
 #include "Templates/UniquePtr.h"
@@ -23,7 +24,13 @@ namespace Metasound
 	namespace Engine
 	{
 		class FOperatorCacheStatTracker;
-	}
+
+#if CSV_PROFILER
+		void RecordOperatorStat(const FTopLevelAssetPath& InAssetPath, int32 CategoryIndex, int32 Value, ECsvCustomStatOp StatOp);
+		void RecordOperatorStat(const FTopLevelAssetPath& InAssetPath, int32 CategoryIndex, float Value, ECsvCustomStatOp StatOp);
+		void RecordOperatorStat(const FTopLevelAssetPath& InAssetPath, int32 CategoryIndex, double Value, ECsvCustomStatOp StatOp);
+#endif // CSV_PROFILER
+	} // namespace Engine
 
 	namespace OperatorPoolPrivate
 	{
@@ -61,10 +68,11 @@ namespace Metasound
 			void ExpireResult(const IntermediateResult& InResultToExpire);
 			void TickResults(const float DeltaTimeSeconds);
 	
-		}; // class FWindowedHitRate
+		};
 	} // namespace OperatorPoolPrivate
 #endif // #if METASOUND_OPERATORCACHEPROFILER_ENABLED
 
+
 	struct FOperatorPoolSettings
 	{
 		uint32 MaxNumOperators = 64;
@@ -74,7 +82,7 @@ namespace Metasound
 	// Data required to build an operator without immediately playing it
 	struct FOperatorBuildData
 	{
-		FMetasoundGeneratorInitParams InitParams;
+		FGeneratorInitParams InitParams;
 		Frontend::FGraphRegistryKey RegistryKey;
 		FGuid AssetClassID;
 		int32 NumInstances;
@@ -84,7 +92,7 @@ namespace Metasound
 
 		FOperatorBuildData() = delete;
 		UE_API FOperatorBuildData(
-			  FMetasoundGeneratorInitParams&& InInitParams
+			  FGeneratorInitParams&& InInitParams
 			, Frontend::FGraphRegistryKey InRegistryKey
 			, FGuid InAssetID
 			, int32 InNumInstances = 1
@@ -96,10 +104,13 @@ namespace Metasound
 	// Provides additional debug context for the operator the pool is interacting with.
 	struct FOperatorContext
 	{
-		FName GraphInstanceName;
-		FStringView MetaSoundName;
+		// Path to underlying graph asset
+		FTopLevelAssetPath GraphAssetPath;
 
-		static UE_API FOperatorContext FromInitParams(const FMetasoundGeneratorInitParams& InParams);
+		// Path to asset with given operator (presets have different paths than graph asset)
+		FTopLevelAssetPath AssetPath;
+
+		static UE_API FOperatorContext FromInitParams(const FGeneratorInitParams& InParams);
 	};
 
 	// Pool of re-useable metasound operators to be used / put back by the metasound generator
@@ -108,52 +119,32 @@ namespace Metasound
 	{
 	public:
 
-		UE_API FOperatorPool(const FOperatorPoolSettings& InSettings);
+		UE_API FOperatorPool(FOperatorPoolSettings InSettings = { });
 		UE_API ~FOperatorPool();
 
-
-		UE_DEPRECATED(5.5, "Use ClaimOperator(const FOperatorPoolEntryID&, ...) instead")
-		UE_API FOperatorAndInputs ClaimOperator(const FGuid& InOperatorID);
 		UE_API FOperatorAndInputs ClaimOperator(const FOperatorPoolEntryID& InOperatorID, const FOperatorContext& InContext);
 
-		UE_DEPRECATED(5.5, "Use AddOperator(const FOperatorPoolEntryID&, ...) instead")
-		UE_API void AddOperator(const FGuid& InOperatorID, TUniquePtr<IOperator>&& InOperator, FInputVertexInterfaceData&& InputData);
 		UE_API void AddOperator(const FOperatorPoolEntryID& InOperatorID, TUniquePtr<IOperator>&& InOperator, FInputVertexInterfaceData&& InputData, TSharedPtr<FGraphRenderCost>&& InRenderCost = {});
-
-		UE_DEPRECATED(5.5, "Use AddOperator(const FOperatorPoolEntryID&, ...) instead")
-		UE_API void AddOperator(const FGuid& InOperatorID, FOperatorAndInputs&& OperatorAndInputs);
 		UE_API void AddOperator(const FOperatorPoolEntryID& InOperatorID, FOperatorAndInputs&& OperatorAndInputs);
 
 		UE_API void BuildAndAddOperator(TUniquePtr<FOperatorBuildData> InBuildData);
 
-		UE_DEPRECATED(5.5, "Use TouchOperators(const FOperatorPoolEntryID&, ...) instead")
-		UE_API void TouchOperators(const FGuid& InOperatorID, int32 NumToTouch = 1);
 		UE_API void TouchOperators(const FOperatorPoolEntryID& InOperatorID, int32 NumToTouch = 1);
 		UE_API void TouchOperatorsViaAssetClassID(const FGuid& InAssetClassID, int32 NumToTouch = 1);
 
 		UE_API bool IsStopping() const;
 
-		UE_DEPRECATED(5.5, "Use RemoveOperatorsWithID(const FOperatorPoolEntryID&) instead")
-		UE_API void RemoveOperatorsWithID(const FGuid& InOperatorID);
 		UE_API void RemoveOperatorsWithID(const FOperatorPoolEntryID& InOperatorID);
 		UE_API void RemoveOperatorsWithAssetClassID(const FGuid& InAssetClassID);
 
-		UE_DEPRECATED(5.5, "Use GetNumCachedOperatorsWithID(const FOperatorPoolEntryID&) instead")
-		UE_API int32 GetNumCachedOperatorsWithID(const FGuid& InOperatorID) const;
 		UE_API int32 GetNumCachedOperatorsWithID(const FOperatorPoolEntryID& InOperatorID) const;
 		UE_API int32 GetNumCachedOperatorsWithAssetClassID(const FGuid& InAssetClassID) const;
 
-		UE_DEPRECATED(5.5, "Adding id to look-up is now private implementation")
-		void AddAssetIdToGraphIdLookUp(const FGuid& InAssetClassID, const FOperatorPoolEntryID& InOperatorID) { }
-
 		UE_API void SetMaxNumOperators(uint32 InMaxNumOperators);
 #if METASOUND_OPERATORCACHEPROFILER_ENABLED
 		UE_API void UpdateHitRateTracker();
 #endif // #if METASOUND_OPERATORCACHEPROFILER_ENABLED
 
-		UE_DEPRECATED(5.5, "Use StopAsyncTasks")
-		UE_API void CancelAllBuildEvents();
-
 		UE_API void StopAsyncTasks();
 
 		using FTaskId = int32;
diff --git a/Engine/Plugins/Runtime/Metasound/Source/MetasoundGraphCore/MetasoundGraphCore.Build.cs b/Engine/Plugins/Runtime/Metasound/Source/MetasoundGraphCore/MetasoundGraphCore.Build.cs
--- a/Engine/Plugins/Runtime/Metasound/Source/MetasoundGraphCore/MetasoundGraphCore.Build.cs
+++ b/Engine/Plugins/Runtime/Metasound/Source/MetasoundGraphCore/MetasoundGraphCore.Build.cs
@@ -7,22 +7,23 @@ namespace UnrealBuildTool.Rules
 {
 	public class MetasoundGraphCore : ModuleRules
 	{
-        public MetasoundGraphCore(ReadOnlyTargetRules Target) : base(Target)
+		public MetasoundGraphCore(ReadOnlyTargetRules Target) : base(Target)
 		{
-            PublicDependencyModuleNames.AddRange(
-				new string[] {
-                    "Core",
+			PublicDependencyModuleNames.AddRange(
+				new string[]
+				{
+					"Core",
+					"CoreUObject",
 					"SignalProcessing",
 					"AudioExtensions"
 				}
-            );
+			);
 
 			PrivateDependencyModuleNames.AddRange(
 				new string[]
 				{
 					"MathCore",
 					"TraceLog"
-					// ... add private dependencies that you statically link with here ...	
 				}
 			);
 		}
diff --git a/Engine/Plugins/Runtime/Metasound/Source/MetasoundGraphCore/Private/MetasoundDynamicOperatorTransactor.cpp b/Engine/Plugins/Runtime/Metasound/Source/MetasoundGraphCore/Private/MetasoundDynamicOperatorTransactor.cpp
--- a/Engine/Plugins/Runtime/Metasound/Source/MetasoundGraphCore/Private/MetasoundDynamicOperatorTransactor.cpp
+++ b/Engine/Plugins/Runtime/Metasound/Source/MetasoundGraphCore/Private/MetasoundDynamicOperatorTransactor.cpp
@@ -583,7 +583,7 @@ namespace Metasound
 
 		FDynamicOperatorTransactor::FDynamicOperatorTransactor()
 		: OperatorBuilder(DynamicOperatorTransactorPrivate::GetOperatorBuilderSettings())
-		, Graph(TEXT(""), FGuid())
+		, Graph(FTopLevelAssetPath(), FGuid())
 		{
 		}
 
diff --git a/Engine/Plugins/Runtime/Metasound/Source/MetasoundGraphCore/Private/MetasoundGraph.cpp b/Engine/Plugins/Runtime/Metasound/Source/MetasoundGraphCore/Private/MetasoundGraph.cpp
--- a/Engine/Plugins/Runtime/Metasound/Source/MetasoundGraphCore/Private/MetasoundGraph.cpp
+++ b/Engine/Plugins/Runtime/Metasound/Source/MetasoundGraphCore/Private/MetasoundGraph.cpp
@@ -21,13 +21,39 @@ namespace Metasound
 	}
 
 	FGraph::FGraph(const FString& InInstanceName, const FGuid& InInstanceID, TSharedPtr<const IOperatorData> InOperatorData)
-	: InstanceName(InInstanceName)
+PRAGMA_DISABLE_DEPRECATION_WARNINGS
+		: AssetPath(FTopLevelAssetPath(FName(*InInstanceName)))
+PRAGMA_ENABLE_DEPRECATION_WARNINGS
+		, InstanceName(AssetPath.GetPackageName().IsNone() ? AssetPath.GetAssetName() : AssetPath.GetPackageName())
+		, InstanceID(InInstanceID)
+		, Metadata(FNodeClassMetadata::GetEmpty())
+		, OperatorData(MoveTemp(InOperatorData))
+	{
+	}
+
+	FGraph::FGraph(const FName& InInstanceName, const FGuid& InInstanceID, TSharedPtr<const IOperatorData> InOperatorData)
+		: AssetPath(FTopLevelAssetPath({ }, InInstanceName))
+		, InstanceName(InInstanceName)
+		, InstanceID(InInstanceID)
+		, Metadata(FNodeClassMetadata::GetEmpty())
+		, OperatorData(MoveTemp(InOperatorData))
+	{
+	}
+
+	FGraph::FGraph(const FTopLevelAssetPath& InAssetPath, const FGuid& InInstanceID, TSharedPtr<const IOperatorData> InOperatorData)
+	: AssetPath(InAssetPath)
+	, InstanceName(InAssetPath.GetPackageName().IsNone() ? InAssetPath.GetAssetName() : InAssetPath.GetPackageName())
 	, InstanceID(InInstanceID)
 	, Metadata(FNodeClassMetadata::GetEmpty())
 	, OperatorData(MoveTemp(InOperatorData))
 	{
 	}
 
+	const FTopLevelAssetPath& FGraph::GetAssetPath() const
+	{
+		return AssetPath;
+	}
+
 	const FVertexName& FGraph::GetInstanceName() const
 	{
 		return InstanceName;
diff --git a/Engine/Plugins/Runtime/Metasound/Source/MetasoundGraphCore/Public/MetasoundGraph.h b/Engine/Plugins/Runtime/Metasound/Source/MetasoundGraphCore/Public/MetasoundGraph.h
--- a/Engine/Plugins/Runtime/Metasound/Source/MetasoundGraphCore/Public/MetasoundGraph.h
+++ b/Engine/Plugins/Runtime/Metasound/Source/MetasoundGraphCore/Public/MetasoundGraph.h
@@ -19,11 +19,22 @@ namespace Metasound
 	class FGraph : public IGraph
 	{
 		public:
+			UE_DEPRECATED(5.7, "Use constructor that provides an FName or AssetPath instead")
+			UE_API FGraph(const FString& InInstanceName, const FGuid& InInstanceID, TSharedPtr<const IOperatorData> InNodeConfig = { });
 
-			UE_API FGraph(const FString& InInstanceName, const FGuid& InInstanceID, TSharedPtr<const IOperatorData> InNodeConfig={});
+			// Constructor to be used for dynamically generated MetaSounds not associated with a given asset (ex. test graphs, graphs
+			// not represented by an asset-contained frontend document, etc.)
+			UE_API FGraph(const FName& InInstanceName, const FGuid& InInstanceID, TSharedPtr<const IOperatorData> InNodeConfig = { });
+
+			// Constructor to be used for MetaSounds associated with a given asset. Both transient and serialized assets should use this constructor,
+			// as transient assets are contained in the transient package)
+			UE_API FGraph(const FTopLevelAssetPath& InAssetPath, const FGuid& InInstanceID, TSharedPtr<const IOperatorData> InNodeConfig = { });
 
 			virtual ~FGraph() = default;
 
+			/** Returns path to graph's authoring content (invalid if generated dynamically). */
+			UE_API virtual const FTopLevelAssetPath& GetAssetPath() const override;
+
 			/** Return the name of this specific instance of the node class. */
 			UE_API virtual const FVertexName& GetInstanceName() const override;
 
@@ -247,6 +258,7 @@ namespace Metasound
 				virtual TUniquePtr<IOperator> CreateOperator(const FBuildOperatorParams& InParams, FBuildResults& OutResults) override;
 			};
 
+			FTopLevelAssetPath AssetPath;
 			FVertexName InstanceName;
 			FGuid InstanceID;
 
diff --git a/Engine/Plugins/Runtime/Metasound/Source/MetasoundGraphCore/Public/MetasoundNodeInterface.h b/Engine/Plugins/Runtime/Metasound/Source/MetasoundGraphCore/Public/MetasoundNodeInterface.h
--- a/Engine/Plugins/Runtime/Metasound/Source/MetasoundGraphCore/Public/MetasoundNodeInterface.h
+++ b/Engine/Plugins/Runtime/Metasound/Source/MetasoundGraphCore/Public/MetasoundNodeInterface.h
@@ -360,6 +360,9 @@ namespace Metasound
 		public:
 			virtual ~IGraph() {}
 
+			/** Returns path to graph's authoring content (invalid if generated dynamically). */
+			virtual const FTopLevelAssetPath& GetAssetPath() const = 0;
+
 			/** Retrieve all the edges associated with a graph. */
 			virtual const TArray<FDataEdge>& GetDataEdges() const = 0;
 
