Fix for unresolved namespace when building certain platforms
#rnx

diff --git a/Engine/Plugins/Runtime/Metasound/Source/MetasoundGenerator/Private/MetasoundGeneratorModule.cpp b/Engine/Plugins/Runtime/Metasound/Source/MetasoundGenerator/Private/MetasoundGeneratorModule.cpp
--- a/Engine/Plugins/Runtime/Metasound/Source/MetasoundGenerator/Private/MetasoundGeneratorModule.cpp
+++ b/Engine/Plugins/Runtime/Metasound/Source/MetasoundGenerator/Private/MetasoundGeneratorModule.cpp
@@ -72,7 +72,7 @@ namespace Metasound
 			{
 				CounterManager->VisitStats([](const FTopLevelAssetPath& InAssetPath, int64 Value)
 				{
-					Engine::RecordOperatorStat(InAssetPath, CSV_CATEGORY_INDEX(MetaSound_ActiveOperators), (int32)Value, ECsvCustomStatOp::Set);
+					Metasound::Engine::RecordOperatorStat(InAssetPath, CSV_CATEGORY_INDEX(MetaSound_ActiveOperators), (int32)Value, ECsvCustomStatOp::Set);
 				});
 			}
 		});
