UE: Quick fix for compile error on consoles

diff --git a/Engine/Plugins/Runtime/Metasound/Source/MetasoundGenerator/Private/MetasoundGeneratorModule.cpp b/Engine/Plugins/Runtime/Metasound/Source/MetasoundGenerator/Private/MetasoundGeneratorModule.cpp
--- a/Engine/Plugins/Runtime/Metasound/Source/MetasoundGenerator/Private/MetasoundGeneratorModule.cpp
+++ b/Engine/Plugins/Runtime/Metasound/Source/MetasoundGenerator/Private/MetasoundGeneratorModule.cpp
@@ -60,7 +60,7 @@ namespace Metasound
 		OperatorPool = MakeShared<FOperatorPool>();
 		OperatorInstanceCounterManager = MakeShared<FConcurrentInstanceCounterManager>(TEXT("MetaSound/Active_Generators"));
 
-#if CSV_PROFILER && !UE_BUILD_SHIPPING
+#if CSV_PROFILER && !UE_BUILD_SHIPPING && METASOUND_OPERATORCACHEPROFILER_ENABLED
 		CsvEndFrameDelegateHandle = FCsvProfiler::Get()->OnCSVProfileEndFrame().AddLambda([WeakCounterManager = OperatorInstanceCounterManager.ToWeakPtr()]
 		{
 			if (!CVarRecordActiveOperatorsToCsv->GetBool())
@@ -76,7 +76,7 @@ namespace Metasound
 				});
 			}
 		});
-#endif // #if CSV_PROFILER && !UE_BUILD_SHIPPING
+#endif // #if CSV_PROFILER && !UE_BUILD_SHIPPING && METASOUND_OPERATORCACHEPROFILER_ENABLED
 	}
 
 	void FMetasoundGeneratorModule::ShutdownModule()
