MetasoundFrontendLiteralBlueprintAccess additions and metadata changes

- added pure functions for creating literals and accessing values
- added function categories for Create Literal and Get Literal As
- changed display names for consistency in the function browser
- added tooltips

#rnx

diff --git a/Engine/Plugins/Runtime/Metasound/Source/MetasoundEngine/Private/MetasoundFrontendLiteralBlueprintAccess.cpp b/Engine/Plugins/Runtime/Metasound/Source/MetasoundEngine/Private/MetasoundFrontendLiteralBlueprintAccess.cpp
--- a/Engine/Plugins/Runtime/Metasound/Source/MetasoundEngine/Private/MetasoundFrontendLiteralBlueprintAccess.cpp
+++ b/Engine/Plugins/Runtime/Metasound/Source/MetasoundEngine/Private/MetasoundFrontendLiteralBlueprintAccess.cpp
@@ -20,6 +20,29 @@ namespace MetasoundFrontendLiteralBlueprintAccessPrivate
 		OutResult = Literal.TryGet(Value) ? EMetaSoundBuilderResult::Succeeded : EMetaSoundBuilderResult::Failed;
 		return Value;
 	}
+
+	template<typename TLiteralType>
+	TLiteralType GetPODValueFromMetaSoundLiteral(const FMetasoundFrontendLiteral& Literal, const EMetasoundFrontendLiteralType ExpectedType, const TLiteralType& DefaultValue)
+	{
+		if (Literal.GetType() != ExpectedType)
+		{
+			UE_LOG(LogMetaSound, Warning, TEXT("Tried to get literal as type %s, but literal was of type %s.")
+				, *UEnum::GetValueAsString(ExpectedType)
+				, *UEnum::GetValueAsString(Literal.GetType()));
+			return DefaultValue;
+		}
+
+		EMetaSoundBuilderResult Result;
+		TLiteralType Value = GetPODValueFromMetaSoundLiteral<TLiteralType>(Literal, Result);
+
+		if (Result != EMetaSoundBuilderResult::Succeeded)
+		{
+			UE_LOG(LogMetaSound, Warning, TEXT("Failed to get %s value from InLiteral."), *UEnum::GetValueAsString(ExpectedType));
+			return DefaultValue;
+		}
+
+		return Value;
+	}
 }
 
 FString UMetasoundFrontendLiteralBlueprintAccess::Conv_MetaSoundLiteralToString(const FMetasoundFrontendLiteral& Literal)
@@ -93,6 +116,61 @@ FMetasoundFrontendLiteral UMetasoundFrontendLiteralBlueprintAccess::CreateMetaSo
 	return FMetasoundFrontendLiteral{ Param };
 }
 
+FMetasoundFrontendLiteral UMetasoundFrontendLiteralBlueprintAccess::CreateMetaSoundLiteralFromBoolean(const bool InBoolean)
+{
+	return MetasoundFrontendLiteralBlueprintAccessPrivate::CreatePODMetaSoundLiteral(InBoolean);
+}
+
+FMetasoundFrontendLiteral UMetasoundFrontendLiteralBlueprintAccess::CreateMetaSoundLiteralFromBooleanArray(const TArray<bool>& InBooleanArray)
+{
+	return MetasoundFrontendLiteralBlueprintAccessPrivate::CreatePODMetaSoundLiteral(InBooleanArray);
+}
+
+FMetasoundFrontendLiteral UMetasoundFrontendLiteralBlueprintAccess::CreateMetaSoundLiteralFromFloat(const float InFloat)
+{
+	return MetasoundFrontendLiteralBlueprintAccessPrivate::CreatePODMetaSoundLiteral(InFloat);
+}
+
+FMetasoundFrontendLiteral UMetasoundFrontendLiteralBlueprintAccess::CreateMetaSoundLiteralFromFloatArray(const TArray<float>& InFloatArray)
+{
+	return MetasoundFrontendLiteralBlueprintAccessPrivate::CreatePODMetaSoundLiteral(InFloatArray);
+}
+
+FMetasoundFrontendLiteral UMetasoundFrontendLiteralBlueprintAccess::CreateMetaSoundLiteralFromInteger(const int32 InInteger)
+{
+	return MetasoundFrontendLiteralBlueprintAccessPrivate::CreatePODMetaSoundLiteral(InInteger);
+}
+
+FMetasoundFrontendLiteral UMetasoundFrontendLiteralBlueprintAccess::CreateMetaSoundLiteralFromIntegerArray(const TArray<int32>& InIntegerArray)
+{
+	return MetasoundFrontendLiteralBlueprintAccessPrivate::CreatePODMetaSoundLiteral(InIntegerArray);
+}
+
+FMetasoundFrontendLiteral UMetasoundFrontendLiteralBlueprintAccess::CreateMetaSoundLiteralFromObject(UObject* InObject)
+{
+	return MetasoundFrontendLiteralBlueprintAccessPrivate::CreatePODMetaSoundLiteral(InObject);
+}
+
+FMetasoundFrontendLiteral UMetasoundFrontendLiteralBlueprintAccess::CreateMetaSoundLiteralFromObjectArray(const TArray<UObject*>& InObjectArray)
+{
+	return MetasoundFrontendLiteralBlueprintAccessPrivate::CreatePODMetaSoundLiteral(InObjectArray);
+}
+
+FMetasoundFrontendLiteral UMetasoundFrontendLiteralBlueprintAccess::CreateMetaSoundLiteralFromString(const FString& InString)
+{
+	return MetasoundFrontendLiteralBlueprintAccessPrivate::CreatePODMetaSoundLiteral(InString);
+}
+
+FMetasoundFrontendLiteral UMetasoundFrontendLiteralBlueprintAccess::CreateMetaSoundLiteralFromStringArray(const TArray<FString>& InStringArray)
+{
+	return MetasoundFrontendLiteralBlueprintAccessPrivate::CreatePODMetaSoundLiteral(InStringArray);
+}
+
+FMetasoundFrontendLiteral UMetasoundFrontendLiteralBlueprintAccess::CreateMetaSoundLiteralFromAudioParameter(const FAudioParameter& InAudioParameter)
+{
+	return FMetasoundFrontendLiteral{ InAudioParameter };
+}
+
 bool UMetasoundFrontendLiteralBlueprintAccess::GetBoolValueFromLiteral(const FMetasoundFrontendLiteral& Literal, EMetaSoundBuilderResult& OutResult)
 {
 	return MetasoundFrontendLiteralBlueprintAccessPrivate::GetPODValueFromMetaSoundLiteral<bool>(Literal, OutResult);
@@ -115,7 +193,7 @@ TArray<float> UMetasoundFrontendLiteralBlueprintAccess::GetFloatArrayValueFromLi
 
 int32 UMetasoundFrontendLiteralBlueprintAccess::GetIntValueFromLiteral(const FMetasoundFrontendLiteral& Literal, EMetaSoundBuilderResult& OutResult)
 {
-	return MetasoundFrontendLiteralBlueprintAccessPrivate::GetPODValueFromMetaSoundLiteral<int>(Literal, OutResult);
+	return MetasoundFrontendLiteralBlueprintAccessPrivate::GetPODValueFromMetaSoundLiteral<int32>(Literal, OutResult);
 }
 
 TArray<int32> UMetasoundFrontendLiteralBlueprintAccess::GetIntArrayValueFromLiteral(const FMetasoundFrontendLiteral& Literal, EMetaSoundBuilderResult& OutResult)
@@ -142,3 +220,53 @@ TArray<FString> UMetasoundFrontendLiteralBlueprintAccess::GetStringArrayValueFro
 {
 	return MetasoundFrontendLiteralBlueprintAccessPrivate::GetPODValueFromMetaSoundLiteral<TArray<FString>>(Literal, OutResult);
 }
+
+bool UMetasoundFrontendLiteralBlueprintAccess::GetMetaSoundLiteralAsBool(const FMetasoundFrontendLiteral& InLiteral)
+{
+	return MetasoundFrontendLiteralBlueprintAccessPrivate::GetPODValueFromMetaSoundLiteral<bool>(InLiteral, EMetasoundFrontendLiteralType::Boolean, false);
+}
+
+TArray<bool> UMetasoundFrontendLiteralBlueprintAccess::GetMetaSoundLiteralAsBoolArray(const FMetasoundFrontendLiteral& InLiteral)
+{
+	return MetasoundFrontendLiteralBlueprintAccessPrivate::GetPODValueFromMetaSoundLiteral<TArray<bool>>(InLiteral, EMetasoundFrontendLiteralType::BooleanArray, TArray<bool>());
+}
+
+float UMetasoundFrontendLiteralBlueprintAccess::GetMetaSoundLiteralAsFloat(const FMetasoundFrontendLiteral& InLiteral)
+{
+	return MetasoundFrontendLiteralBlueprintAccessPrivate::GetPODValueFromMetaSoundLiteral<float>(InLiteral, EMetasoundFrontendLiteralType::Float, 0.f);
+}
+
+TArray<float> UMetasoundFrontendLiteralBlueprintAccess::GetMetaSoundLiteralAsFloatArray(const FMetasoundFrontendLiteral& InLiteral)
+{
+	return MetasoundFrontendLiteralBlueprintAccessPrivate::GetPODValueFromMetaSoundLiteral<TArray<float>>(InLiteral, EMetasoundFrontendLiteralType::FloatArray, TArray<float>());
+}
+
+int32 UMetasoundFrontendLiteralBlueprintAccess::GetMetaSoundLiteralAsInteger(const FMetasoundFrontendLiteral& InLiteral)
+{
+	return MetasoundFrontendLiteralBlueprintAccessPrivate::GetPODValueFromMetaSoundLiteral<int32>(InLiteral, EMetasoundFrontendLiteralType::Integer, 0);
+}
+
+TArray<int32> UMetasoundFrontendLiteralBlueprintAccess::GetMetaSoundLiteralAsIntegerArray(const FMetasoundFrontendLiteral& InLiteral)
+{
+	return MetasoundFrontendLiteralBlueprintAccessPrivate::GetPODValueFromMetaSoundLiteral<TArray<int32>>(InLiteral, EMetasoundFrontendLiteralType::IntegerArray, TArray<int32>());
+}
+
+UObject* UMetasoundFrontendLiteralBlueprintAccess::GetMetaSoundLiteralAsObject(const FMetasoundFrontendLiteral& InLiteral)
+{
+	return MetasoundFrontendLiteralBlueprintAccessPrivate::GetPODValueFromMetaSoundLiteral<UObject*>(InLiteral, EMetasoundFrontendLiteralType::UObject, nullptr);
+}
+
+TArray<UObject*> UMetasoundFrontendLiteralBlueprintAccess::GetMetaSoundLiteralAsObjectArray(const FMetasoundFrontendLiteral& InLiteral)
+{
+	return MetasoundFrontendLiteralBlueprintAccessPrivate::GetPODValueFromMetaSoundLiteral<TArray<UObject*>>(InLiteral, EMetasoundFrontendLiteralType::UObjectArray, TArray<UObject*>());
+}
+
+FString UMetasoundFrontendLiteralBlueprintAccess::GetMetaSoundLiteralAsString(const FMetasoundFrontendLiteral& InLiteral)
+{
+	return MetasoundFrontendLiteralBlueprintAccessPrivate::GetPODValueFromMetaSoundLiteral<FString>(InLiteral, EMetasoundFrontendLiteralType::String, FString());
+}
+
+TArray<FString> UMetasoundFrontendLiteralBlueprintAccess::GetMetaSoundLiteralAsStringArray(const FMetasoundFrontendLiteral& InLiteral)
+{
+	return MetasoundFrontendLiteralBlueprintAccessPrivate::GetPODValueFromMetaSoundLiteral<TArray<FString>>(InLiteral, EMetasoundFrontendLiteralType::StringArray, TArray<FString>());
+}
diff --git a/Engine/Plugins/Runtime/Metasound/Source/MetasoundEngine/Public/MetasoundFrontendLiteralBlueprintAccess.h b/Engine/Plugins/Runtime/Metasound/Source/MetasoundEngine/Public/MetasoundFrontendLiteralBlueprintAccess.h
--- a/Engine/Plugins/Runtime/Metasound/Source/MetasoundEngine/Public/MetasoundFrontendLiteralBlueprintAccess.h
+++ b/Engine/Plugins/Runtime/Metasound/Source/MetasoundEngine/Public/MetasoundFrontendLiteralBlueprintAccess.h
@@ -15,77 +15,186 @@ UCLASS()
 class UMetasoundFrontendLiteralBlueprintAccess : public UBlueprintFunctionLibrary
 {
 	GENERATED_BODY()
-	
+
 	UFUNCTION(BlueprintPure, meta = (DisplayName = "To String (MetaSound Literal)", CompactNodeTitle = "->", BlueprintAutocast), Category = "Utilities|String")
 	static FString Conv_MetaSoundLiteralToString(const FMetasoundFrontendLiteral& Literal);
-	
+
 	UFUNCTION(BlueprintPure, Category = "Audio|MetaSound", meta = (DisplayName = "Get MetaSound Literal Type"))
 	static UPARAM(DisplayName = "Type") EMetasoundFrontendLiteralType GetType(const FMetasoundFrontendLiteral& Literal);
 
 	// Literal creation
 
-	UFUNCTION(BlueprintCallable, Category = "Audio|MetaSound", meta = (DisplayName = "Create MetaSound Bool Literal"))
+	// Creates a MetaSound Literal using the given boolean value.
+	UFUNCTION(BlueprintCallable, Category = "Audio|MetaSound|Create Literal", meta = (DisplayName = "Create MetaSound Literal from Boolean"))
 	static UPARAM(DisplayName = "Bool Literal") FMetasoundFrontendLiteral CreateBoolMetaSoundLiteral(bool Value);
 
-	UFUNCTION(BlueprintCallable, Category = "Audio|MetaSound", meta = (DisplayName = "Create MetaSound Bool Array Literal"))
+	// Creates a MetaSound Literal using the given boolean array.
+	UFUNCTION(BlueprintCallable, Category = "Audio|MetaSound|Create Literal", meta = (DisplayName = "Create MetaSound Literal from Boolean Array"))
 	static UPARAM(DisplayName = "Bool Array Literal") FMetasoundFrontendLiteral CreateBoolArrayMetaSoundLiteral(const TArray<bool>& Value);
 
-	UFUNCTION(BlueprintCallable, Category = "Audio|MetaSound", meta = (DisplayName = "Create MetaSound Float Literal"))
+	// Creates a MetaSound Literal using the given float value.
+	UFUNCTION(BlueprintCallable, Category = "Audio|MetaSound|Create Literal", meta = (DisplayName = "Create MetaSound Literal from Float"))
 	static UPARAM(DisplayName = "Float Literal") FMetasoundFrontendLiteral CreateFloatMetaSoundLiteral(float Value);
 
-	UFUNCTION(BlueprintCallable, Category = "Audio|MetaSound", meta = (DisplayName = "Create MetaSound Float Array Literal"))
+	// Creates a MetaSound Literal using the given float array.
+	UFUNCTION(BlueprintCallable, Category = "Audio|MetaSound|Create Literal", meta = (DisplayName = "Create MetaSound Literal from Float Array"))
 	static UPARAM(DisplayName = "Float Array Literal") FMetasoundFrontendLiteral CreateFloatArrayMetaSoundLiteral(const TArray<float>& Value);
 
-	UFUNCTION(BlueprintCallable, Category = "Audio|MetaSound", meta = (DisplayName = "Create MetaSound Int Literal"))
+	// Creates a MetaSound Literal using the given integer value.
+	UFUNCTION(BlueprintCallable, Category = "Audio|MetaSound|Create Literal", meta = (DisplayName = "Create MetaSound Literal from Integer"))
 	static UPARAM(DisplayName = "Int32 Literal") FMetasoundFrontendLiteral CreateIntMetaSoundLiteral(int32 Value);
 
-	UFUNCTION(BlueprintCallable, Category = "Audio|MetaSound", meta = (DisplayName = "Create MetaSound Int Array Literal"))
+	// Creates a MetaSound Literal using the given integer array.
+	UFUNCTION(BlueprintCallable, Category = "Audio|MetaSound|Create Literal", meta = (DisplayName = "Create MetaSound Literal from Integer Array"))
 	static UPARAM(DisplayName = "Int32 Array Literal") FMetasoundFrontendLiteral CreateIntArrayMetaSoundLiteral(const TArray<int32>& Value);
 
-	UFUNCTION(BlueprintCallable, Category = "Audio|MetaSound", meta = (DisplayName = "Create MetaSound Object Literal"))
+	// Creates a MetaSound Literal using the given object.
+	UFUNCTION(BlueprintCallable, Category = "Audio|MetaSound|Create Literal", meta = (DisplayName = "Create MetaSound Literal from Object"))
 	static UPARAM(DisplayName = "Object Literal") FMetasoundFrontendLiteral CreateObjectMetaSoundLiteral(UObject* Value);
 
-	UFUNCTION(BlueprintCallable, Category = "Audio|MetaSound", meta = (DisplayName = "Create MetaSound Object Array Literal"))
+	// Creates a MetaSound Literal using the given object array.
+	UFUNCTION(BlueprintCallable, Category = "Audio|MetaSound|Create Literal", meta = (DisplayName = "Create MetaSound Literal from Object Array"))
 	static UPARAM(DisplayName = "Object Array Literal") FMetasoundFrontendLiteral CreateObjectArrayMetaSoundLiteral(const TArray<UObject*>& Value);
 
-	UFUNCTION(BlueprintCallable, Category = "Audio|MetaSound", meta = (DisplayName = "Create MetaSound String Literal"))
+	// Creates a MetaSound Literal using the given string.
+	UFUNCTION(BlueprintCallable, Category = "Audio|MetaSound|Create Literal", meta = (DisplayName = "Create MetaSound Literal from String"))
 	static UPARAM(DisplayName = "String Literal") FMetasoundFrontendLiteral CreateStringMetaSoundLiteral(const FString& Value);
 
-	UFUNCTION(BlueprintCallable, Category = "Audio|MetaSound", meta = (DisplayName = "Create MetaSound String Array Literal"))
+	// Creates a MetaSound Literal using the given string array.
+	UFUNCTION(BlueprintCallable, Category = "Audio|MetaSound|Create Literal", meta = (DisplayName = "Create MetaSound Literal from String Array"))
 	static UPARAM(DisplayName = "String Array Literal") FMetasoundFrontendLiteral CreateStringArrayMetaSoundLiteral(const TArray<FString>& Value);
 
-	UFUNCTION(BlueprintCallable, Category = "Audio|MetaSound", meta = (DisplayName = "Create MetaSound Literal From AudioParameter"))
+	// Creates a MetaSound Literal using the given audio parameter.
+	UFUNCTION(BlueprintCallable, Category = "Audio|MetaSound|Create Literal", meta = (DisplayName = "Create MetaSound Literal from Audio Parameter"))
 	static UPARAM(DisplayName = "Param Literal") FMetasoundFrontendLiteral CreateMetaSoundLiteralFromParam(const FAudioParameter& Param);
 
-	// Value accessors 
+	// Literal creation (pure)
+
+	// (Pure) Creates a MetaSound Literal using the given boolean value.
+	UFUNCTION(BlueprintPure, Category = "Audio|MetaSound|Create Literal", DisplayName = "Create MetaSound Literal from Boolean")
+	static FMetasoundFrontendLiteral CreateMetaSoundLiteralFromBoolean(const bool InBoolean);
+
+	// (Pure) Creates a MetaSound Literal using the given boolean array.
+	UFUNCTION(BlueprintPure, Category = "Audio|MetaSound|Create Literal", DisplayName = "Create MetaSound Literal from Boolean Array")
+	static FMetasoundFrontendLiteral CreateMetaSoundLiteralFromBooleanArray(const TArray<bool>& InBooleanArray);
+
+	// (Pure) Creates a MetaSound Literal using the given float value.
+	UFUNCTION(BlueprintPure, Category = "Audio|MetaSound|Create Literal", DisplayName = "Create MetaSound Literal from Float")
+	static FMetasoundFrontendLiteral CreateMetaSoundLiteralFromFloat(const float InFloat);
+
+	// (Pure) Creates a MetaSound Literal using the given float array.
+	UFUNCTION(BlueprintPure, Category = "Audio|MetaSound|Create Literal", DisplayName = "Create MetaSound Literal from Float Array")
+	static FMetasoundFrontendLiteral CreateMetaSoundLiteralFromFloatArray(const TArray<float>& InFloatArray);
+
+	// (Pure) Creates a MetaSound Literal using the given integer value.
+	UFUNCTION(BlueprintPure, Category = "Audio|MetaSound|Create Literal", DisplayName = "Create MetaSound Literal from Integer ")
+	static FMetasoundFrontendLiteral CreateMetaSoundLiteralFromInteger(const int32 InInteger);
+
+	// (Pure) Creates a MetaSound Literal using the given integer array.
+	UFUNCTION(BlueprintPure, Category = "Audio|MetaSound|Create Literal", DisplayName = "Create MetaSound Literal from Integer Array")
+	static FMetasoundFrontendLiteral CreateMetaSoundLiteralFromIntegerArray(const TArray<int32>& InIntegerArray);
+
+	// (Pure) Creates a MetaSound Literal using the given object.
+	UFUNCTION(BlueprintPure, Category = "Audio|MetaSound|Create Literal", DisplayName = "Create MetaSound Literal from Object")
+	static FMetasoundFrontendLiteral CreateMetaSoundLiteralFromObject(UObject* InObject);
+
+	// (Pure) Creates a MetaSound Literal using the given object array.
+	UFUNCTION(BlueprintPure, Category = "Audio|MetaSound|Create Literal", DisplayName = "Create MetaSound Literal from Object Array")
+	static FMetasoundFrontendLiteral CreateMetaSoundLiteralFromObjectArray(const TArray<UObject*>& InObjectArray);
+
+	// (Pure) Creates a MetaSound Literal using the given string.
+	UFUNCTION(BlueprintPure, Category = "Audio|MetaSound|Create Literal", DisplayName = "Create MetaSound Literal from String")
+	static FMetasoundFrontendLiteral CreateMetaSoundLiteralFromString(const FString& InString);
+
+	// (Pure) Creates a MetaSound Literal using the given string array.
+	UFUNCTION(BlueprintPure, Category = "Audio|MetaSound|Create Literal", DisplayName = "Create MetaSound Literal from String Array")
+	static FMetasoundFrontendLiteral CreateMetaSoundLiteralFromStringArray(const TArray<FString>& InStringArray);
+
+	// (Pure) Creates a MetaSound Literal using the given audio parameter.
+	UFUNCTION(BlueprintPure, Category = "Audio|MetaSound|Create Literal", DisplayName = "Create MetaSound Literal from Audio Parameter")
+	static FMetasoundFrontendLiteral CreateMetaSoundLiteralFromAudioParameter(const FAudioParameter& InAudioParameter);
+
+	// Value accessors
 
-	UFUNCTION(BlueprintCallable, Category = "Audio|MetaSound", meta = (ExpandEnumAsExecs = "OutResult"))
+	// Returns the value of the given MetaSound Literal as a bool. Has separate execution outputs for success and failure.
+	UFUNCTION(BlueprintCallable, Category = "Audio|MetaSound|Get Literal As", DisplayName = "Get MetaSound Literal as Bool", meta = (ExpandEnumAsExecs = "OutResult"))
 	static bool GetBoolValueFromLiteral(const FMetasoundFrontendLiteral& Literal, EMetaSoundBuilderResult& OutResult);
 
-	UFUNCTION(BlueprintCallable, Category = "Audio|MetaSound", meta = (ExpandEnumAsExecs = "OutResult"))
+	// Returns the value of the given MetaSound Literal as a bool array. Has separate execution outputs for success and failure.
+	UFUNCTION(BlueprintCallable, Category = "Audio|MetaSound|Get Literal As", DisplayName = "Get MetaSound Literal as Bool Array", meta = (ExpandEnumAsExecs = "OutResult"))
 	static TArray<bool> GetBoolArrayValueFromLiteral(const FMetasoundFrontendLiteral& Literal, EMetaSoundBuilderResult& OutResult);
 
-	UFUNCTION(BlueprintCallable, Category = "Audio|MetaSound", meta = (ExpandEnumAsExecs = "OutResult"))
+	// Returns the value of the given MetaSound Literal as a float. Has separate execution outputs for success and failure.
+	UFUNCTION(BlueprintCallable, Category = "Audio|MetaSound|Get Literal As", DisplayName = "Get MetaSound Literal as Float", meta = (ExpandEnumAsExecs = "OutResult"))
 	static float GetFloatValueFromLiteral(const FMetasoundFrontendLiteral& Literal, EMetaSoundBuilderResult& OutResult);
 
-	UFUNCTION(BlueprintCallable, Category = "Audio|MetaSound", meta = (ExpandEnumAsExecs = "OutResult"))
+	// Returns the value of the given MetaSound Literal as a float array. Has separate execution outputs for success and failure.
+	UFUNCTION(BlueprintCallable, Category = "Audio|MetaSound|Get Literal As", DisplayName = "Get MetaSound Literal as Float Array", meta = (ExpandEnumAsExecs = "OutResult"))
 	static TArray<float> GetFloatArrayValueFromLiteral(const FMetasoundFrontendLiteral& Literal, EMetaSoundBuilderResult& OutResult);
 
-	UFUNCTION(BlueprintCallable, Category = "Audio|MetaSound", meta = (ExpandEnumAsExecs = "OutResult"))
+	// Returns the value of the given MetaSound Literal as an integer. Has separate execution outputs for success and failure.
+	UFUNCTION(BlueprintCallable, Category = "Audio|MetaSound|Get Literal As", DisplayName = "Get MetaSound Literal as Integer", meta = (ExpandEnumAsExecs = "OutResult"))
 	static int32 GetIntValueFromLiteral(const FMetasoundFrontendLiteral& Literal, EMetaSoundBuilderResult& OutResult);
 
-	UFUNCTION(BlueprintCallable, Category = "Audio|MetaSound", meta = (ExpandEnumAsExecs = "OutResult"))
+	// Returns the value of the given MetaSound Literal as an integer array. Has separate execution outputs for success and failure.
+	UFUNCTION(BlueprintCallable, Category = "Audio|MetaSound|Get Literal As", DisplayName = "Get MetaSound Literal as Integer Array", meta = (ExpandEnumAsExecs = "OutResult"))
 	static TArray<int32> GetIntArrayValueFromLiteral(const FMetasoundFrontendLiteral& Literal, EMetaSoundBuilderResult& OutResult);
 
-	UFUNCTION(BlueprintCallable, Category = "Audio|MetaSound", meta = (ExpandEnumAsExecs = "OutResult"))
+	// Returns the value of the given MetaSound Literal as an object. Has separate execution outputs for success and failure.
+	UFUNCTION(BlueprintCallable, Category = "Audio|MetaSound|Get Literal As", DisplayName = "Get MetaSound Literal as Object", meta = (ExpandEnumAsExecs = "OutResult"))
 	static UObject* GetObjectValueFromLiteral(const FMetasoundFrontendLiteral& Literal, EMetaSoundBuilderResult& OutResult);
-	
-	UFUNCTION(BlueprintCallable, Category = "Audio|MetaSound", meta = (ExpandEnumAsExecs = "OutResult"))
+
+	// Returns the value of the given MetaSound Literal as an object array. Has separate execution outputs for success and failure.
+	UFUNCTION(BlueprintCallable, Category = "Audio|MetaSound|Get Literal As", DisplayName = "Get MetaSound Literal as Object Array", meta = (ExpandEnumAsExecs = "OutResult"))
 	static TArray<UObject*> GetObjectArrayValueFromLiteral(const FMetasoundFrontendLiteral& Literal, EMetaSoundBuilderResult& OutResult);
-	
-	UFUNCTION(BlueprintCallable, Category = "Audio|MetaSound", meta = (ExpandEnumAsExecs = "OutResult"))
+
+	// Returns the value of the given MetaSound Literal as a string. Has separate execution outputs for success and failure.
+	UFUNCTION(BlueprintCallable, Category = "Audio|MetaSound|Get Literal As", DisplayName = "Get MetaSound Literal as String", meta = (ExpandEnumAsExecs = "OutResult"))
 	static FString GetStringValueFromLiteral(const FMetasoundFrontendLiteral& Literal, EMetaSoundBuilderResult& OutResult);
-	
-	UFUNCTION(BlueprintCallable, Category = "Audio|MetaSound", meta = (ExpandEnumAsExecs = "OutResult"))
+
+	// Returns the value of the given MetaSound Literal as a string array. Has separate execution outputs for success and failure.
+	UFUNCTION(BlueprintCallable, Category = "Audio|MetaSound|Get Literal As", DisplayName = "Get MetaSound Literal as String Array", meta = (ExpandEnumAsExecs = "OutResult"))
 	static TArray<FString> GetStringArrayValueFromLiteral(const FMetasoundFrontendLiteral& Literal, EMetaSoundBuilderResult& OutResult);
+
+	// Value accessors (pure)
+
+	// (Pure) Returns the value of the given MetaSound Literal as a bool. Logs a warning if the value fails to be retrieved.
+	UFUNCTION(BlueprintPure, Category = "Audio|MetaSound|Get Literal As", DisplayName = "Get MetaSound Literal as Bool")
+	static bool GetMetaSoundLiteralAsBool(const FMetasoundFrontendLiteral& InLiteral);
+
+	// (Pure) Returns the value of the given MetaSound Literal as a bool array. Logs a warning if the value fails to be retrieved.
+	UFUNCTION(BlueprintPure, Category = "Audio|MetaSound|Get Literal As", DisplayName = "Get MetaSound Literal as Bool Array")
+	static TArray<bool> GetMetaSoundLiteralAsBoolArray(const FMetasoundFrontendLiteral& InLiteral);
+
+	// (Pure) Returns the value of the given MetaSound Literal as a float. Logs a warning if the value fails to be retrieved.
+	UFUNCTION(BlueprintPure, Category = "Audio|MetaSound|Get Literal As", DisplayName = "Get MetaSound Literal as Float")
+	static float GetMetaSoundLiteralAsFloat(const FMetasoundFrontendLiteral& InLiteral);
+
+	// (Pure) Returns the value of the given MetaSound Literal as a float array. Logs a warning if the value fails to be retrieved.
+	UFUNCTION(BlueprintPure, Category = "Audio|MetaSound|Get Literal As", DisplayName = "Get MetaSound Literal as Float Array")
+	static TArray<float> GetMetaSoundLiteralAsFloatArray(const FMetasoundFrontendLiteral& InLiteral);
+
+	// (Pure) Returns the value of the given MetaSound Literal as an integer. Logs a warning if the value fails to be retrieved.
+	UFUNCTION(BlueprintPure, Category = "Audio|MetaSound|Get Literal As", DisplayName = "Get MetaSound Literal as Integer")
+	static int32 GetMetaSoundLiteralAsInteger(const FMetasoundFrontendLiteral& InLiteral);
+
+	// (Pure) Returns the value of the given MetaSound Literal as an integer array. Logs a warning if the value fails to be retrieved.
+	UFUNCTION(BlueprintPure, Category = "Audio|MetaSound|Get Literal As", DisplayName = "Get MetaSound Literal as Integer Array")
+	static TArray<int32> GetMetaSoundLiteralAsIntegerArray(const FMetasoundFrontendLiteral& InLiteral);
+
+	// (Pure) Returns the value of the given MetaSound Literal as an object. Logs a warning if the value fails to be retrieved.
+	UFUNCTION(BlueprintPure, Category = "Audio|MetaSound|Get Literal As", DisplayName = "Get MetaSound Literal as Object")
+	static UObject* GetMetaSoundLiteralAsObject(const FMetasoundFrontendLiteral& InLiteral);
+
+	// (Pure) Returns the value of the given MetaSound Literal as an object array. Logs a warning if the value fails to be retrieved.
+	UFUNCTION(BlueprintPure, Category = "Audio|MetaSound|Get Literal As", DisplayName = "Get MetaSound Literal as Object Array")
+	static TArray<UObject*> GetMetaSoundLiteralAsObjectArray(const FMetasoundFrontendLiteral& InLiteral);
+
+	// (Pure) Returns the value of the given MetaSound Literal as a string. Logs a warning if the value fails to be retrieved.
+	UFUNCTION(BlueprintPure, Category = "Audio|MetaSound|Get Literal As", DisplayName = "Get MetaSound Literal as String")
+	static FString GetMetaSoundLiteralAsString(const FMetasoundFrontendLiteral& InLiteral);
+
+	// (Pure) Returns the value of the given MetaSound Literal as a string array. Logs a warning if the value fails to be retrieved.
+	UFUNCTION(BlueprintPure, Category = "Audio|MetaSound|Get Literal As", DisplayName = "Get MetaSound Literal as String Array")
+	static TArray<FString> GetMetaSoundLiteralAsStringArray(const FMetasoundFrontendLiteral& InLiteral);
 };
