MetasoundFrontendLiteralBlueprintAccess display name changes

- added (Pure) to display names of pure functions for better differentiation in the function browser

#rnx

diff --git a/Engine/Plugins/Runtime/Metasound/Source/MetasoundEngine/Private/MetasoundFrontendLiteralBlueprintAccess.cpp b/Engine/Plugins/Runtime/Metasound/Source/MetasoundEngine/Private/MetasoundFrontendLiteralBlueprintAccess.cpp
--- a/Engine/Plugins/Runtime/Metasound/Source/MetasoundEngine/Private/MetasoundFrontendLiteralBlueprintAccess.cpp
+++ b/Engine/Plugins/Runtime/Metasound/Source/MetasoundEngine/Private/MetasoundFrontendLiteralBlueprintAccess.cpp
@@ -1,6 +1,8 @@
 // Copyright Epic Games, Inc. All Rights Reserved.
 
 #include "MetasoundFrontendLiteralBlueprintAccess.h"
+
+#include "Logging/StructuredLog.h"
 #include "MetasoundDataReference.h"
 
 namespace MetasoundFrontendLiteralBlueprintAccessPrivate
@@ -26,7 +28,7 @@ namespace MetasoundFrontendLiteralBlueprintAccessPrivate
 	{
 		if (Literal.GetType() != ExpectedType)
 		{
-			UE_LOG(LogMetaSound, Warning, TEXT("Tried to get literal as type %s, but literal was of type %s.")
+			UE_LOGFMT(LogMetaSound, Warning, "Tried to get literal as type {expected}, but literal was of type {actual}."
 				, *UEnum::GetValueAsString(ExpectedType)
 				, *UEnum::GetValueAsString(Literal.GetType()));
 			return DefaultValue;
@@ -37,7 +39,7 @@ namespace MetasoundFrontendLiteralBlueprintAccessPrivate
 
 		if (Result != EMetaSoundBuilderResult::Succeeded)
 		{
-			UE_LOG(LogMetaSound, Warning, TEXT("Failed to get %s value from InLiteral."), *UEnum::GetValueAsString(ExpectedType));
+			UE_LOGFMT(LogMetaSound, Warning, "Failed to get %s value from InLiteral.", *UEnum::GetValueAsString(ExpectedType));
 			return DefaultValue;
 		}
 
diff --git a/Engine/Plugins/Runtime/Metasound/Source/MetasoundEngine/Public/MetasoundFrontendLiteralBlueprintAccess.h b/Engine/Plugins/Runtime/Metasound/Source/MetasoundEngine/Public/MetasoundFrontendLiteralBlueprintAccess.h
--- a/Engine/Plugins/Runtime/Metasound/Source/MetasoundEngine/Public/MetasoundFrontendLiteralBlueprintAccess.h
+++ b/Engine/Plugins/Runtime/Metasound/Source/MetasoundEngine/Public/MetasoundFrontendLiteralBlueprintAccess.h
@@ -70,48 +70,48 @@ class UMetasoundFrontendLiteralBlueprintAccess : public UBlueprintFunctionLibrar
 
 	// Literal creation (pure)
 
-	// (Pure) Creates a MetaSound Literal using the given boolean value.
-	UFUNCTION(BlueprintPure, Category = "Audio|MetaSound|Create Literal", DisplayName = "Create MetaSound Literal from Boolean")
+	// Creates a MetaSound Literal using the given boolean value.
+	UFUNCTION(BlueprintPure, Category = "Audio|MetaSound|Create Literal", DisplayName = "Create MetaSound Literal from Boolean (Pure)")
 	static FMetasoundFrontendLiteral CreateMetaSoundLiteralFromBoolean(const bool InBoolean);
 
-	// (Pure) Creates a MetaSound Literal using the given boolean array.
-	UFUNCTION(BlueprintPure, Category = "Audio|MetaSound|Create Literal", DisplayName = "Create MetaSound Literal from Boolean Array")
+	// Creates a MetaSound Literal using the given boolean array.
+	UFUNCTION(BlueprintPure, Category = "Audio|MetaSound|Create Literal", DisplayName = "Create MetaSound Literal from Boolean Array (Pure)")
 	static FMetasoundFrontendLiteral CreateMetaSoundLiteralFromBooleanArray(const TArray<bool>& InBooleanArray);
 
-	// (Pure) Creates a MetaSound Literal using the given float value.
-	UFUNCTION(BlueprintPure, Category = "Audio|MetaSound|Create Literal", DisplayName = "Create MetaSound Literal from Float")
+	// Creates a MetaSound Literal using the given float value.
+	UFUNCTION(BlueprintPure, Category = "Audio|MetaSound|Create Literal", DisplayName = "Create MetaSound Literal from Float (Pure)")
 	static FMetasoundFrontendLiteral CreateMetaSoundLiteralFromFloat(const float InFloat);
 
-	// (Pure) Creates a MetaSound Literal using the given float array.
-	UFUNCTION(BlueprintPure, Category = "Audio|MetaSound|Create Literal", DisplayName = "Create MetaSound Literal from Float Array")
+	// Creates a MetaSound Literal using the given float array.
+	UFUNCTION(BlueprintPure, Category = "Audio|MetaSound|Create Literal", DisplayName = "Create MetaSound Literal from Float Array (Pure)")
 	static FMetasoundFrontendLiteral CreateMetaSoundLiteralFromFloatArray(const TArray<float>& InFloatArray);
 
-	// (Pure) Creates a MetaSound Literal using the given integer value.
-	UFUNCTION(BlueprintPure, Category = "Audio|MetaSound|Create Literal", DisplayName = "Create MetaSound Literal from Integer ")
+	// Creates a MetaSound Literal using the given integer value.
+	UFUNCTION(BlueprintPure, Category = "Audio|MetaSound|Create Literal", DisplayName = "Create MetaSound Literal from Integer (Pure)")
 	static FMetasoundFrontendLiteral CreateMetaSoundLiteralFromInteger(const int32 InInteger);
 
-	// (Pure) Creates a MetaSound Literal using the given integer array.
-	UFUNCTION(BlueprintPure, Category = "Audio|MetaSound|Create Literal", DisplayName = "Create MetaSound Literal from Integer Array")
+	// Creates a MetaSound Literal using the given integer array.
+	UFUNCTION(BlueprintPure, Category = "Audio|MetaSound|Create Literal", DisplayName = "Create MetaSound Literal from Integer Array (Pure)")
 	static FMetasoundFrontendLiteral CreateMetaSoundLiteralFromIntegerArray(const TArray<int32>& InIntegerArray);
 
-	// (Pure) Creates a MetaSound Literal using the given object.
-	UFUNCTION(BlueprintPure, Category = "Audio|MetaSound|Create Literal", DisplayName = "Create MetaSound Literal from Object")
+	// Creates a MetaSound Literal using the given object.
+	UFUNCTION(BlueprintPure, Category = "Audio|MetaSound|Create Literal", DisplayName = "Create MetaSound Literal from Object (Pure)")
 	static FMetasoundFrontendLiteral CreateMetaSoundLiteralFromObject(UObject* InObject);
 
-	// (Pure) Creates a MetaSound Literal using the given object array.
-	UFUNCTION(BlueprintPure, Category = "Audio|MetaSound|Create Literal", DisplayName = "Create MetaSound Literal from Object Array")
+	// Creates a MetaSound Literal using the given object array.
+	UFUNCTION(BlueprintPure, Category = "Audio|MetaSound|Create Literal", DisplayName = "Create MetaSound Literal from Object Array (Pure)")
 	static FMetasoundFrontendLiteral CreateMetaSoundLiteralFromObjectArray(const TArray<UObject*>& InObjectArray);
 
-	// (Pure) Creates a MetaSound Literal using the given string.
-	UFUNCTION(BlueprintPure, Category = "Audio|MetaSound|Create Literal", DisplayName = "Create MetaSound Literal from String")
+	// Creates a MetaSound Literal using the given string.
+	UFUNCTION(BlueprintPure, Category = "Audio|MetaSound|Create Literal", DisplayName = "Create MetaSound Literal from String (Pure)")
 	static FMetasoundFrontendLiteral CreateMetaSoundLiteralFromString(const FString& InString);
 
-	// (Pure) Creates a MetaSound Literal using the given string array.
-	UFUNCTION(BlueprintPure, Category = "Audio|MetaSound|Create Literal", DisplayName = "Create MetaSound Literal from String Array")
+	// Creates a MetaSound Literal using the given string array.
+	UFUNCTION(BlueprintPure, Category = "Audio|MetaSound|Create Literal", DisplayName = "Create MetaSound Literal from String Array (Pure)")
 	static FMetasoundFrontendLiteral CreateMetaSoundLiteralFromStringArray(const TArray<FString>& InStringArray);
 
-	// (Pure) Creates a MetaSound Literal using the given audio parameter.
-	UFUNCTION(BlueprintPure, Category = "Audio|MetaSound|Create Literal", DisplayName = "Create MetaSound Literal from Audio Parameter")
+	// Creates a MetaSound Literal using the given audio parameter.
+	UFUNCTION(BlueprintPure, Category = "Audio|MetaSound|Create Literal", DisplayName = "Create MetaSound Literal from Audio Parameter (Pure)")
 	static FMetasoundFrontendLiteral CreateMetaSoundLiteralFromAudioParameter(const FAudioParameter& InAudioParameter);
 
 	// Value accessors
@@ -158,43 +158,43 @@ class UMetasoundFrontendLiteralBlueprintAccess : public UBlueprintFunctionLibrar
 
 	// Value accessors (pure)
 
-	// (Pure) Returns the value of the given MetaSound Literal as a bool. Logs a warning if the value fails to be retrieved.
-	UFUNCTION(BlueprintPure, Category = "Audio|MetaSound|Get Literal As", DisplayName = "Get MetaSound Literal as Bool")
+	// Returns the value of the given MetaSound Literal as a bool. Logs a warning if the value fails to be retrieved.
+	UFUNCTION(BlueprintPure, Category = "Audio|MetaSound|Get Literal As", DisplayName = "Get MetaSound Literal as Bool (Pure)")
 	static bool GetMetaSoundLiteralAsBool(const FMetasoundFrontendLiteral& InLiteral);
 
-	// (Pure) Returns the value of the given MetaSound Literal as a bool array. Logs a warning if the value fails to be retrieved.
-	UFUNCTION(BlueprintPure, Category = "Audio|MetaSound|Get Literal As", DisplayName = "Get MetaSound Literal as Bool Array")
+	// Returns the value of the given MetaSound Literal as a bool array. Logs a warning if the value fails to be retrieved.
+	UFUNCTION(BlueprintPure, Category = "Audio|MetaSound|Get Literal As", DisplayName = "Get MetaSound Literal as Bool Array (Pure)")
 	static TArray<bool> GetMetaSoundLiteralAsBoolArray(const FMetasoundFrontendLiteral& InLiteral);
 
-	// (Pure) Returns the value of the given MetaSound Literal as a float. Logs a warning if the value fails to be retrieved.
-	UFUNCTION(BlueprintPure, Category = "Audio|MetaSound|Get Literal As", DisplayName = "Get MetaSound Literal as Float")
+	// Returns the value of the given MetaSound Literal as a float. Logs a warning if the value fails to be retrieved.
+	UFUNCTION(BlueprintPure, Category = "Audio|MetaSound|Get Literal As", DisplayName = "Get MetaSound Literal as Float (Pure)")
 	static float GetMetaSoundLiteralAsFloat(const FMetasoundFrontendLiteral& InLiteral);
 
-	// (Pure) Returns the value of the given MetaSound Literal as a float array. Logs a warning if the value fails to be retrieved.
-	UFUNCTION(BlueprintPure, Category = "Audio|MetaSound|Get Literal As", DisplayName = "Get MetaSound Literal as Float Array")
+	// Returns the value of the given MetaSound Literal as a float array. Logs a warning if the value fails to be retrieved.
+	UFUNCTION(BlueprintPure, Category = "Audio|MetaSound|Get Literal As", DisplayName = "Get MetaSound Literal as Float Array (Pure)")
 	static TArray<float> GetMetaSoundLiteralAsFloatArray(const FMetasoundFrontendLiteral& InLiteral);
 
-	// (Pure) Returns the value of the given MetaSound Literal as an integer. Logs a warning if the value fails to be retrieved.
-	UFUNCTION(BlueprintPure, Category = "Audio|MetaSound|Get Literal As", DisplayName = "Get MetaSound Literal as Integer")
+	// Returns the value of the given MetaSound Literal as an integer. Logs a warning if the value fails to be retrieved.
+	UFUNCTION(BlueprintPure, Category = "Audio|MetaSound|Get Literal As", DisplayName = "Get MetaSound Literal as Integer (Pure)")
 	static int32 GetMetaSoundLiteralAsInteger(const FMetasoundFrontendLiteral& InLiteral);
 
-	// (Pure) Returns the value of the given MetaSound Literal as an integer array. Logs a warning if the value fails to be retrieved.
-	UFUNCTION(BlueprintPure, Category = "Audio|MetaSound|Get Literal As", DisplayName = "Get MetaSound Literal as Integer Array")
+	// Returns the value of the given MetaSound Literal as an integer array. Logs a warning if the value fails to be retrieved.
+	UFUNCTION(BlueprintPure, Category = "Audio|MetaSound|Get Literal As", DisplayName = "Get MetaSound Literal as Integer Array (Pure)")
 	static TArray<int32> GetMetaSoundLiteralAsIntegerArray(const FMetasoundFrontendLiteral& InLiteral);
 
-	// (Pure) Returns the value of the given MetaSound Literal as an object. Logs a warning if the value fails to be retrieved.
-	UFUNCTION(BlueprintPure, Category = "Audio|MetaSound|Get Literal As", DisplayName = "Get MetaSound Literal as Object")
+	// Returns the value of the given MetaSound Literal as an object. Logs a warning if the value fails to be retrieved.
+	UFUNCTION(BlueprintPure, Category = "Audio|MetaSound|Get Literal As", DisplayName = "Get MetaSound Literal as Object (Pure)")
 	static UObject* GetMetaSoundLiteralAsObject(const FMetasoundFrontendLiteral& InLiteral);
 
-	// (Pure) Returns the value of the given MetaSound Literal as an object array. Logs a warning if the value fails to be retrieved.
-	UFUNCTION(BlueprintPure, Category = "Audio|MetaSound|Get Literal As", DisplayName = "Get MetaSound Literal as Object Array")
+	// Returns the value of the given MetaSound Literal as an object array. Logs a warning if the value fails to be retrieved.
+	UFUNCTION(BlueprintPure, Category = "Audio|MetaSound|Get Literal As", DisplayName = "Get MetaSound Literal as Object Array (Pure)")
 	static TArray<UObject*> GetMetaSoundLiteralAsObjectArray(const FMetasoundFrontendLiteral& InLiteral);
 
-	// (Pure) Returns the value of the given MetaSound Literal as a string. Logs a warning if the value fails to be retrieved.
-	UFUNCTION(BlueprintPure, Category = "Audio|MetaSound|Get Literal As", DisplayName = "Get MetaSound Literal as String")
+	// Returns the value of the given MetaSound Literal as a string. Logs a warning if the value fails to be retrieved.
+	UFUNCTION(BlueprintPure, Category = "Audio|MetaSound|Get Literal As", DisplayName = "Get MetaSound Literal as String (Pure)")
 	static FString GetMetaSoundLiteralAsString(const FMetasoundFrontendLiteral& InLiteral);
 
-	// (Pure) Returns the value of the given MetaSound Literal as a string array. Logs a warning if the value fails to be retrieved.
-	UFUNCTION(BlueprintPure, Category = "Audio|MetaSound|Get Literal As", DisplayName = "Get MetaSound Literal as String Array")
+	// Returns the value of the given MetaSound Literal as a string array. Logs a warning if the value fails to be retrieved.
+	UFUNCTION(BlueprintPure, Category = "Audio|MetaSound|Get Literal As", DisplayName = "Get MetaSound Literal as String Array (Pure)")
 	static TArray<FString> GetMetaSoundLiteralAsStringArray(const FMetasoundFrontendLiteral& InLiteral);
 };
