Metasound data type registration code dedupe

diff --git a/Engine/Plugins/Runtime/Metasound/Source/MetasoundFrontend/Private/MetasoundDataTypeRegistrationMacro.cpp b/Engine/Plugins/Runtime/Metasound/Source/MetasoundFrontend/Private/MetasoundDataTypeRegistrationMacro.cpp
--- a/Engine/Plugins/Runtime/Metasound/Source/MetasoundFrontend/Private/MetasoundDataTypeRegistrationMacro.cpp
+++ b/Engine/Plugins/Runtime/Metasound/Source/MetasoundFrontend/Private/MetasoundDataTypeRegistrationMacro.cpp
@@ -147,5 +147,47 @@ namespace Metasound::Frontend
 			}
 			return TUniquePtr<INode>(nullptr);
 		}
+
+		FORCENOINLINE void FDataTypeRegistryEntryBase::InitNodeClassesBase(
+			const FNodeClassMetadata& InInputClassMetadata,
+			const FNodeClassMetadata& InOutputClassMetadata,
+			const FNodeClassMetadata& InLiteralPrototypeMetadata,
+			const FNodeClassMetadata& InVariableClassMetadata,
+			const FNodeClassMetadata& InVariableMutatorClassMetadata,
+			const FNodeClassMetadata& InVariableAccessorClassMetadata,
+			const FNodeClassMetadata& InVariableDeferredAccessorClassMetadata
+		)
+		{
+			this->InputClassMetadata = MakeShared<FNodeClassMetadata>(InInputClassMetadata);
+			this->InputClass = Metasound::Frontend::GenerateClass(*this->InputClassMetadata, EMetasoundFrontendClassType::Input);
+
+			this->OutputClassMetadata = MakeShared<FNodeClassMetadata>(InOutputClassMetadata);
+			this->OutputClass = Metasound::Frontend::GenerateClass(*this->OutputClassMetadata, EMetasoundFrontendClassType::Output);
+
+			this->LiteralClass = Metasound::Frontend::GenerateClass(InLiteralPrototypeMetadata, EMetasoundFrontendClassType::Literal);
+
+			this->VariableClassMetadata = MakeShared<FNodeClassMetadata>(InVariableClassMetadata);
+			this->VariableClass = Metasound::Frontend::GenerateClass(*this->VariableClassMetadata, EMetasoundFrontendClassType::Variable);
+
+			this->VariableMutatorClassMetadata = MakeShared<FNodeClassMetadata>(InVariableMutatorClassMetadata);
+			this->VariableMutatorClass = Metasound::Frontend::GenerateClass(*this->VariableMutatorClassMetadata, EMetasoundFrontendClassType::VariableMutator);
+
+			this->VariableAccessorClassMetadata = MakeShared<FNodeClassMetadata>(InVariableAccessorClassMetadata);
+			this->VariableAccessorClass = Metasound::Frontend::GenerateClass(*this->VariableAccessorClassMetadata, EMetasoundFrontendClassType::VariableAccessor);
+
+			this->VariableDeferredAccessorClassMetadata = MakeShared<FNodeClassMetadata>(InVariableDeferredAccessorClassMetadata);
+			this->VariableDeferredAccessorClass = Metasound::Frontend::GenerateClass(*this->VariableDeferredAccessorClassMetadata, EMetasoundFrontendClassType::VariableDeferredAccessor);
+		}
+
+		FORCENOINLINE void FDataTypeRegistryEntryBase::InitNodeClassesBaseConstructor(
+			const FNodeClassMetadata& InConstructorInputClassMetadata,
+			const FNodeClassMetadata& InConstructorOutputClassMetadata
+		)
+		{
+			this->ConstructorInputClassMetadata = MakeShared<FNodeClassMetadata>(InConstructorInputClassMetadata);
+			this->ConstructorInputClass = Metasound::Frontend::GenerateClass(*this->ConstructorInputClassMetadata, EMetasoundFrontendClassType::Input);
+			this->ConstructorOutputClassMetadata = MakeShared<FNodeClassMetadata>(InConstructorOutputClassMetadata);
+			this->ConstructorOutputClass = Metasound::Frontend::GenerateClass(*this->ConstructorOutputClassMetadata, EMetasoundFrontendClassType::Output);
+		}
 	}
 }
diff --git a/Engine/Plugins/Runtime/Metasound/Source/MetasoundFrontend/Public/MetasoundDataTypeRegistrationMacro.h b/Engine/Plugins/Runtime/Metasound/Source/MetasoundFrontend/Public/MetasoundDataTypeRegistrationMacro.h
--- a/Engine/Plugins/Runtime/Metasound/Source/MetasoundFrontend/Public/MetasoundDataTypeRegistrationMacro.h
+++ b/Engine/Plugins/Runtime/Metasound/Source/MetasoundFrontend/Public/MetasoundDataTypeRegistrationMacro.h
@@ -355,6 +355,23 @@ namespace Metasound::Frontend
 			TSharedPtr<Frontend::IEnumDataTypeInterface> EnumInterface;
 			Frontend::IParameterAssignmentFunction RawAssignmentFunction;
 			Frontend::FLiteralAssignmentFunction LiteralAssignmentFunction = nullptr;
+
+			typedef void CreateNodeClassMetadataFunc();
+
+			UE_API void InitNodeClassesBase(
+				const FNodeClassMetadata& InputClassMetadata,
+				const FNodeClassMetadata& OutputClassMetadata,
+				const FNodeClassMetadata& LiteralPrototypeMetadata,
+				const FNodeClassMetadata& VariableClassMetadata,
+				const FNodeClassMetadata& VariableMutatorClassMetadata,
+				const FNodeClassMetadata& VariableAccessorClassMetadata,
+				const FNodeClassMetadata& VariableDeferredAccessorClassMetadata
+			);
+
+			UE_API void InitNodeClassesBaseConstructor(
+				const FNodeClassMetadata& ConstructorInputClassMetadata,
+				const FNodeClassMetadata& ConstructorOutputClassMetadata
+			);
 		};
 
 
@@ -456,34 +473,22 @@ namespace Metasound::Frontend
 						const FName DataTypeName = GetMetasoundDataTypeName<TDataType>();
 						const FName UnnamedVertex;
 
-
-						this->InputClassMetadata = MakeShared<FNodeClassMetadata>(TInputNode<TDataType, EVertexAccessType::Reference>::CreateNodeClassMetadata(UnnamedVertex));
-						this->InputClass = Metasound::Frontend::GenerateClass(*this->InputClassMetadata, EMetasoundFrontendClassType::Input);
-
-						this->OutputClassMetadata = MakeShared<FNodeClassMetadata>(FOutputNode::CreateNodeClassMetadata(UnnamedVertex, DataTypeName, EVertexAccessType::Reference));
-						this->OutputClass = Metasound::Frontend::GenerateClass(*this->OutputClassMetadata, EMetasoundFrontendClassType::Output);
-
-						TLiteralNode<TDataType> LiteralPrototype(TEXT(""), FGuid(), FLiteral());
-						this->LiteralClass = Metasound::Frontend::GenerateClass(LiteralPrototype.GetMetadata(), EMetasoundFrontendClassType::Literal);
-
-						this->VariableClassMetadata = MakeShared<FNodeClassMetadata>(TVariableNode<TDataType>::CreateNodeClassMetadata());
-						this->VariableClass = Metasound::Frontend::GenerateClass(*this->VariableClassMetadata, EMetasoundFrontendClassType::Variable);
-
-						this->VariableMutatorClassMetadata = MakeShared<FNodeClassMetadata>(TVariableMutatorNode<TDataType>::CreateNodeClassMetadata());
-						this->VariableMutatorClass = Metasound::Frontend::GenerateClass(*this->VariableMutatorClassMetadata, EMetasoundFrontendClassType::VariableMutator);
-
-						this->VariableAccessorClassMetadata = MakeShared<FNodeClassMetadata>(TVariableAccessorNode<TDataType>::CreateNodeClassMetadata());
-						this->VariableAccessorClass = Metasound::Frontend::GenerateClass(*this->VariableAccessorClassMetadata, EMetasoundFrontendClassType::VariableAccessor);
-
-						this->VariableDeferredAccessorClassMetadata = MakeShared<FNodeClassMetadata>(TVariableDeferredAccessorNode<TDataType>::CreateNodeClassMetadata());
-						this->VariableDeferredAccessorClass = Metasound::Frontend::GenerateClass(*this->VariableDeferredAccessorClassMetadata, EMetasoundFrontendClassType::VariableDeferredAccessor);
+						InitNodeClassesBase(
+							TInputNode<TDataType, EVertexAccessType::Reference>::CreateNodeClassMetadata(UnnamedVertex),
+							FOutputNode::CreateNodeClassMetadata(UnnamedVertex, DataTypeName, EVertexAccessType::Reference),
+							TLiteralNode<TDataType>(TEXT(""), FGuid(), FLiteral()).GetMetadata(),
+							TVariableNode<TDataType>::CreateNodeClassMetadata(),
+							TVariableMutatorNode<TDataType>::CreateNodeClassMetadata(),
+							TVariableAccessorNode<TDataType>::CreateNodeClassMetadata(),
+							TVariableDeferredAccessorNode<TDataType>::CreateNodeClassMetadata()
+						);
 
 						if constexpr (bIsConstructorType)
 						{
-							this->ConstructorInputClassMetadata = MakeShared<FNodeClassMetadata>(TInputNode<TDataType, EVertexAccessType::Value>::CreateNodeClassMetadata(UnnamedVertex));
-							this->ConstructorInputClass = Metasound::Frontend::GenerateClass(*this->ConstructorInputClassMetadata, EMetasoundFrontendClassType::Input);
-							this->ConstructorOutputClassMetadata = MakeShared<FNodeClassMetadata>(FOutputNode::CreateNodeClassMetadata(UnnamedVertex, DataTypeName, EVertexAccessType::Value));
-							this->ConstructorOutputClass = Metasound::Frontend::GenerateClass(*this->ConstructorOutputClassMetadata, EMetasoundFrontendClassType::Output);
+							InitNodeClassesBaseConstructor(
+								TInputNode<TDataType, EVertexAccessType::Value>::CreateNodeClassMetadata(UnnamedVertex),
+								FOutputNode::CreateNodeClassMetadata(UnnamedVertex, DataTypeName, EVertexAccessType::Value)
+							);
 						}
 					}
 				}
