- Remove AutoUpdate DenyList as it is no longer supported in MetaSounds (AutoUpdate automatically only runs at cook now as opposed to runtime as it did in past versions)
- Fix-up MetaSound Register Directories feature to not run when cooking
- Remove old deprecated MetaSound functions

[RN] AutoUpdate DenyList is now deprecated in MetaSounds as updates are always done during cook (support for runtime autoupdating MetaSound graphs is no longer supported due to performance limitations))

diff --git a/Engine/Plugins/Runtime/Metasound/Source/MetasoundEditor/Private/MetasoundEditorGraph.cpp b/Engine/Plugins/Runtime/Metasound/Source/MetasoundEditor/Private/MetasoundEditorGraph.cpp
--- a/Engine/Plugins/Runtime/Metasound/Source/MetasoundEditor/Private/MetasoundEditorGraph.cpp
+++ b/Engine/Plugins/Runtime/Metasound/Source/MetasoundEditor/Private/MetasoundEditorGraph.cpp
@@ -672,13 +672,11 @@ void UMetasoundEditorGraphMemberDefaultLiteral::PostEditChangeChainProperty(FPro
 		{
 			const FName ChildPropertyName = ChildProperty->GetFName();
 
-PRAGMA_DISABLE_DEPRECATION_WARNINGS
-			if (ChildPropertyName.IsEqual(GetDefaultsPropertyName()) || ChildPropertyName.IsEqual(GetDefaultPropertyName()))
+			if (ChildPropertyName.IsEqual(GetDefaultsPropertyName()))
 			{
 				ResolvePageDefaults();
 				SortPageDefaults();
 			}
-PRAGMA_ENABLE_DEPRECATION_WARNINGS
 		}
 	}
 
@@ -749,42 +747,6 @@ Metasound::Editor::ENodeSection UMetasoundEditorGraphInput::GetSectionID() const
 	return Metasound::Editor::ENodeSection::Inputs;
 }
 
-Metasound::Frontend::FNodeHandle UMetasoundEditorGraphInput::AddNodeHandle(const FName& InName, const Metasound::Editor::FCreateNodeVertexParams& InParams)
-{
-	using namespace Metasound::Editor;
-	using namespace Metasound::Frontend;
-
-	UMetasoundEditorGraph* Graph = GetOwningGraph();
-	if (!ensure(Graph))
-	{
-		return Metasound::Frontend::INodeController::GetInvalidHandle();
-	}
-
-	UObject& MetaSound = Graph->GetMetasoundChecked();
-	FMetasoundAssetBase* MetaSoundAsset = Metasound::IMetasoundUObjectRegistry::Get().GetObjectAsAssetBase(&MetaSound);
-	check(MetaSoundAsset);
-
-	FMetasoundFrontendClassInput ClassInput = FGraphBuilder::CreateUniqueClassInput(MetaSound, InParams, { }, &InName);
-	return MetaSoundAsset->GetRootGraphHandle()->AddInputVertex(ClassInput);
-}
-
-UMetasoundEditorGraphNode* UMetasoundEditorGraphInput::AddNode(Metasound::Frontend::FNodeHandle InNodeHandle, bool bInSelectNewNode)
-{
-	using namespace Metasound::Editor;
-	using namespace Metasound::Engine;
-	using namespace Metasound::Frontend;
-
-	UMetasoundEditorGraph* Graph = GetOwningGraph();
-	check(Graph);
-
-	if (const FMetasoundFrontendNode* TemplateNode = FInputNodeTemplate::CreateNode(GetFrontendBuilderChecked(), GetMemberName()))
-	{
-		return FGraphBuilder::AddInputNode(Graph->GetMetasoundChecked(), TemplateNode->GetID(), bInSelectNewNode);
-	}
-
-	return nullptr;
-}
-
 const FText& UMetasoundEditorGraphInput::GetGraphMemberLabel() const
 {
 	static const FText Label = LOCTEXT("GraphMemberLabel_Input", "Input");
@@ -1176,34 +1138,6 @@ EMetasoundFrontendVertexAccessType UMetasoundEditorGraphInput::GetVertexAccessTy
 	return EMetasoundFrontendVertexAccessType::Reference;
 }
 
-Metasound::Frontend::FNodeHandle UMetasoundEditorGraphOutput::AddNodeHandle(const FName& InName, const Metasound::Editor::FCreateNodeVertexParams& InParams)
-{
-	using namespace Metasound::Editor;
-	using namespace Metasound::Frontend;
-
-	UMetasoundEditorGraph* Graph = GetOwningGraph();
-	if (!ensure(Graph))
-	{
-		return Metasound::Frontend::INodeController::GetInvalidHandle();
-	} 
-
-	UObject& MetaSound = Graph->GetMetasoundChecked();
-
-	FMetasoundAssetBase* MetaSoundAsset = Metasound::IMetasoundUObjectRegistry::Get().GetObjectAsAssetBase(&MetaSound);
-	check(MetaSoundAsset);
-
-	FMetasoundFrontendClassOutput ClassOutput = FGraphBuilder::CreateUniqueClassOutput(MetaSound, InParams, &InName);
-	return MetaSoundAsset->GetRootGraphHandle()->AddOutputVertex(ClassOutput);
-}
-
-UMetasoundEditorGraphNode* UMetasoundEditorGraphOutput::AddNode(Metasound::Frontend::FNodeHandle InNodeHandle, bool bInSelectNewNode)
-{
-	using namespace Metasound::Editor;
-	UMetasoundEditorGraph* Graph = GetOwningGraph();
-	check(Graph);
-	return FGraphBuilder::AddOutputNode(Graph->GetMetasoundChecked(), InNodeHandle->GetID(), bInSelectNewNode);
-}
-
 FText UMetasoundEditorGraphOutput::GetDescription() const
 {
 	if (!Breadcrumb.Description.IsEmpty())
@@ -1603,14 +1537,6 @@ Metasound::Frontend::FConstVariableHandle UMetasoundEditorGraphVariable::GetCons
 	return MetaSound.GetRootGraphHandle()->FindVariable(VariableID);
 }
 
-UMetasoundEditorGraphNode* UMetasoundEditorGraphVariable::AddNode(Metasound::Frontend::FNodeHandle InNodeHandle, bool bInSelectNewNode)
-{
-	using namespace Metasound::Editor;
-
-	FMetasoundAssetBase& MetaSound = FGraphBuilder::GetOutermostMetaSoundChecked(*this);
-	return FGraphBuilder::AddVariableNode(*MetaSound.GetOwningAsset(), InNodeHandle->GetID(), bInSelectNewNode);
-}
-
 void UMetasoundEditorGraphVariable::SetMemberName(const FName& InNewName, bool bPostTransaction)
 {
 	using namespace Metasound::Editor;
diff --git a/Engine/Plugins/Runtime/Metasound/Source/MetasoundEditor/Private/MetasoundEditorGraphMemberDefaults.cpp b/Engine/Plugins/Runtime/Metasound/Source/MetasoundEditor/Private/MetasoundEditorGraphMemberDefaults.cpp
--- a/Engine/Plugins/Runtime/Metasound/Source/MetasoundEditor/Private/MetasoundEditorGraphMemberDefaults.cpp
+++ b/Engine/Plugins/Runtime/Metasound/Source/MetasoundEditor/Private/MetasoundEditorGraphMemberDefaults.cpp
@@ -603,11 +603,6 @@ void UMetasoundEditorGraphMemberDefaultFloat::ClampDefaults()
 	}
 }
 
-float UMetasoundEditorGraphMemberDefaultFloat::GetDefault()
-{
-	return GetDefaultAs<float>();
-}
-
 EMetasoundFrontendLiteralType UMetasoundEditorGraphMemberDefaultFloat::GetLiteralType() const
 {
 	return EMetasoundFrontendLiteralType::Float;
diff --git a/Engine/Plugins/Runtime/Metasound/Source/MetasoundEditor/Public/MetasoundEditorGraph.h b/Engine/Plugins/Runtime/Metasound/Source/MetasoundEditor/Public/MetasoundEditorGraph.h
--- a/Engine/Plugins/Runtime/Metasound/Source/MetasoundEditor/Public/MetasoundEditorGraph.h
+++ b/Engine/Plugins/Runtime/Metasound/Source/MetasoundEditor/Public/MetasoundEditorGraph.h
@@ -137,30 +137,11 @@ public:
 
 	UE_API UMetasoundEditorGraphMember* FindMember() const;
 
-	UE_DEPRECATED(5.5, "Due to serialization optimization, literals now inherit from LiteralMetadata and are no longer parented under members. "
-		"Use 'GetDefaultsPropertyName' instead to get transient paged defaults property name")
-	static FName GetDefaultPropertyName()
-	{
-		return "Default";
-	}
-
 	static FName GetDefaultsPropertyName()
 	{
 		return "Defaults";
 	}
 
-	UE_DEPRECATED(5.5, "Due to serialization optimization, literals now inherit from LiteralMetadata and are no longer parented under members. Use FindMember instead")
-	const UMetasoundEditorGraphMember* GetParentMember() const
-	{
-		return Cast<UMetasoundEditorGraphMember>(GetOuter());
-	}
-
-	UE_DEPRECATED(5.5, "Due to serialization optimization, literals now inherit from LiteralMetadata and are no longer parented under members. Use FindMember instead")
-	UMetasoundEditorGraphMember* GetParentMember()
-	{
-		return Cast<UMetasoundEditorGraphMember>(GetOuter());
-	}
-
 protected:
 	virtual void ResolvePageDefaults() { }
 	virtual void SortPageDefaults() { }
@@ -264,9 +245,6 @@ public:
 	UE_API void InitializeLiteral();
 
 protected:
-	UE_DEPRECATED(5.5, "No longer supported in favor of internal Editor::FGraphBuilder API which is actively being transitioned to using Document Builder API")
-	virtual UMetasoundEditorGraphNode* AddNode(Metasound::Frontend::FNodeHandle InNodeHandle, bool bInSelectNewNode) PURE_VIRTUAL(UMetasoundEditorGraphMember::AddNode, return nullptr;);
-
 	/** Default literal value of member */
 	UPROPERTY()
 	TObjectPtr<UMetasoundEditorGraphMemberDefaultLiteral> Literal;
@@ -286,9 +264,6 @@ class UMetasoundEditorGraphVertex : public UMetasoundEditorGraphMember
 	GENERATED_BODY()
 
 protected:
-	UE_DEPRECATED(5.5, "EditorGraph vertices no longer generates node handles (use Builder API instead)")
-	virtual Metasound::Frontend::FNodeHandle AddNodeHandle(const FName& InNodeName, const Metasound::Editor::FCreateNodeVertexParams& InParams) PURE_VIRTUAL(UMetasoundEditorGraphVertex::AddNodeHandle, return Metasound::Frontend::INodeController::GetInvalidHandle(); )
-
 	UE_API const FMetasoundFrontendNode* GetFrontendNode() const;
 
 	/** If the vertex's name contains an interface namespace, return true and set OutInterface to that interface if provided. */
@@ -399,12 +374,6 @@ public:
 	UE_API virtual EMetasoundFrontendVertexAccessType GetVertexAccessType() const override;
 
 protected:
-	UE_DEPRECATED(5.5, "No longer supported in favor of internal Editor::FGraphBuilder API which is actively being transitioned to using Document Builder API")
-	UE_API virtual UMetasoundEditorGraphNode* AddNode(Metasound::Frontend::FNodeHandle InNodeHandle, bool bInSelectNewNode) override;
-
-	UE_DEPRECATED(5.5, "EditorGraph vertices no longer generates node handles (use Builder API instead)")
-	UE_API virtual Metasound::Frontend::FNodeHandle AddNodeHandle(const FName& InNodeName, const Metasound::Editor::FCreateNodeVertexParams& InParams) override;
-
 	UE_API virtual bool RenameFrontendMemberInternal(FMetaSoundFrontendDocumentBuilder& Builder, FName OldName, FName InNewName) const override;
 
 	virtual EMetasoundFrontendClassType GetClassType() const override { return EMetasoundFrontendClassType::Input; }
@@ -434,12 +403,6 @@ public:
 	UE_API virtual EMetasoundFrontendVertexAccessType GetVertexAccessType() const override;
 
 protected:
-	UE_DEPRECATED(5.5, "No longer supported in favor of internal Editor::FGraphBuilder API which is actively being transitioned to using Document Builder API")
-	UE_API virtual UMetasoundEditorGraphNode* AddNode(Metasound::Frontend::FNodeHandle InNodeHandle, bool bInSelectNewNode) override;
-
-	UE_DEPRECATED(5.5, "No longer supported in favor of internal Editor::FGraphBuilder API which is actively being transitioned to using Document Builder API")
-	UE_API virtual Metasound::Frontend::FNodeHandle AddNodeHandle(const FName& InNodeName, const Metasound::Editor::FCreateNodeVertexParams& InParams) override;
-
 	UE_API virtual bool RenameFrontendMemberInternal(FMetaSoundFrontendDocumentBuilder& Builder, FName OldName, FName InNewName) const override;
 
 	virtual EMetasoundFrontendClassType GetClassType() const override { return EMetasoundFrontendClassType::Output; }
@@ -488,8 +451,9 @@ public:
 	/**Returns Breadcrumb data of this Variable*/
 	const FMetasoundEditorGraphVariableBreadcrumb& GetBreadcrumb() const { return Breadcrumb; };
 
-	UE_DEPRECATED(5.6, "")
+	UE_DEPRECATED(5.6, "Use MetasoundFrontendDocumentBuilder to edit variables")
 	UE_API Metasound::Frontend::FVariableHandle GetVariableHandle();
+
 	UE_API Metasound::Frontend::FConstVariableHandle GetConstVariableHandle() const;
 
 	UE_API virtual bool CanRename() const override;
@@ -499,10 +463,6 @@ protected:
 	UPROPERTY()
 	FMetasoundEditorGraphVariableBreadcrumb Breadcrumb;
 
-protected:
-	UE_DEPRECATED(5.5, "No longer supported in favor of internal Editor::FGraphBuilder API which is actively being transitioned to using Document Builder API")
-	UE_API virtual UMetasoundEditorGraphNode* AddNode(Metasound::Frontend::FNodeHandle InNodeHandle, bool bInSelectNewNode) override;
-
 private:
 	struct FVariableEditorNodes
 	{
diff --git a/Engine/Plugins/Runtime/Metasound/Source/MetasoundEditor/Public/MetasoundEditorGraphMemberDefaults.h b/Engine/Plugins/Runtime/Metasound/Source/MetasoundEditor/Public/MetasoundEditorGraphMemberDefaults.h
--- a/Engine/Plugins/Runtime/Metasound/Source/MetasoundEditor/Public/MetasoundEditorGraphMemberDefaults.h
+++ b/Engine/Plugins/Runtime/Metasound/Source/MetasoundEditor/Public/MetasoundEditorGraphMemberDefaults.h
@@ -374,15 +374,9 @@ public:
 	virtual void UpdatePreviewInstance(const Metasound::FVertexName& InParameterName, TScriptInterface<IAudioParameterControllerInterface>& InParameterInterface) const override;
 	virtual void PostEditChangeChainProperty(FPropertyChangedChainEvent& InPropertyChangedEvent) override;
 
-	UE_DEPRECATED(5.5, "Use SetFromLiteral instead")
-	void SetDefault(const float InDefault) { }
-
 	// Set Range to reasonable limit given current Default value
 	void SetInitialRange();
 
-	UE_DEPRECATED(5.5, "Use TryFindDefault or GetDefaultAs template call provided from parent class 'UMetaSoundFrontendMemberMetadata' providing an applicable page ID instead")
-	float GetDefault();
-
 	FVector2D GetRange() const;
 	void SetRange(const FVector2D InRange);
 
diff --git a/Engine/Plugins/Runtime/Metasound/Source/MetasoundEngine/Private/MetasoundAssetSubsystem.cpp b/Engine/Plugins/Runtime/Metasound/Source/MetasoundEngine/Private/MetasoundAssetSubsystem.cpp
--- a/Engine/Plugins/Runtime/Metasound/Source/MetasoundEngine/Private/MetasoundAssetSubsystem.cpp
+++ b/Engine/Plugins/Runtime/Metasound/Source/MetasoundEngine/Private/MetasoundAssetSubsystem.cpp
@@ -345,17 +345,6 @@ namespace Metasound::Engine
 		}
 	}
 
-	bool FMetaSoundAssetManager::CanAutoUpdate(const FMetasoundFrontendClassName& InClassName) const
-	{
-		const UMetaSoundSettings* Settings = GetDefault<UMetaSoundSettings>();
-		if (!Settings->bAutoUpdateEnabled)
-		{
-			return false;
-		}
-
-		return !AutoUpdateDenyListCache.Contains(InClassName.GetFullName());
-	}
-
 	bool FMetaSoundAssetManager::ContainsKey(const FMetaSoundAssetKey& InKey) const
 	{
 		FScopeLock Lock(&TagDataMapCriticalSection);
@@ -909,37 +898,6 @@ namespace Metasound::Engine
 
 	void FMetaSoundAssetManager::RebuildDenyListCache(const UAssetManager& InAssetManager)
 	{
-		using namespace Frontend;
-
-		const UMetaSoundSettings* Settings = GetDefault<UMetaSoundSettings>();
-		if (Settings->DenyListCacheChangeID == AutoUpdateDenyListChangeID)
-		{
-			return;
-		}
-
-		AutoUpdateDenyListCache.Reset();
-
-		for (const FMetasoundFrontendClassName& ClassName : Settings->AutoUpdateDenylist)
-		{
-			AutoUpdateDenyListCache.Add(ClassName.GetFullName());
-		}
-
-		check(UAssetManager::IsInitialized());
-		UAssetManager& AssetManager = UAssetManager::Get();
-		for (const FDefaultMetaSoundAssetAutoUpdateSettings& UpdateSettings : Settings->AutoUpdateAssetDenylist)
-		{
-			FAssetData AssetData;
-			if (AssetManager.GetAssetDataForPath(UpdateSettings.MetaSound, AssetData))
-			{
-				const FMetaSoundAssetClassInfo ClassInfo(AssetData);
-				if (ClassInfo.bIsValid)
-				{
-					AutoUpdateDenyListCache.Add(ClassInfo.ClassName.GetFullName());
-				}
-			}
-		}
-
-		AutoUpdateDenyListChangeID = Settings->DenyListCacheChangeID;
 	}
 
 	void FMetaSoundAssetManager::RegisterAssetClassesInDirectories(const TArray<FMetaSoundAssetDirectory>& InDirectories)
@@ -1444,7 +1402,6 @@ void UMetaSoundAssetSubsystem::PostEngineInitInternal()
 	check(UAssetManager::IsInitialized());
 	UAssetManager& AssetManager = UAssetManager::Get();
 	AssetManager.CallOrRegister_OnCompletedInitialScan(FSimpleMulticastDelegate::FDelegate::CreateUObject(this, &UMetaSoundAssetSubsystem::PostInitAssetScanInternal));
-	FMetaSoundAssetManager::GetChecked().RebuildDenyListCache(AssetManager);
 }
 
 void UMetaSoundAssetSubsystem::PostInitAssetScanInternal()
@@ -1454,57 +1411,28 @@ void UMetaSoundAssetSubsystem::PostInitAssetScanInternal()
 
 	METASOUND_TRACE_CPUPROFILER_EVENT_SCOPE(UMetaSoundAssetSubsystem::PostInitAssetScanInternal);
 
-	const UMetaSoundSettings* Settings = GetDefault<UMetaSoundSettings>();
-	if (ensureAlways(Settings))
+	if (!IsRunningCookCommandlet())
 	{
-		FMetaSoundAssetManager& Manager = FMetaSoundAssetManager::GetChecked();
-		FMetaSoundAssetRegistrationOptions RegOptions = Manager.GetRegistrationOptions();
-		Manager.SearchAndIterateDirectoryAssets(Settings->DirectoriesToRegister, [&Manager, &RegOptions](const FAssetData& AssetData)
+		const UMetaSoundSettings* Settings = GetDefault<UMetaSoundSettings>();
+		if (ensureAlways(Settings))
 		{
-			using namespace Metasound;
-
-			Manager.AddOrLoadAndUpdateFromObjectAsync(AssetData, [RegOptions](FMetaSoundAssetKey, UObject& AssetObject)
+			FMetaSoundAssetManager& Manager = FMetaSoundAssetManager::GetChecked();
+			FMetaSoundAssetRegistrationOptions RegOptions = Manager.GetRegistrationOptions();
+			Manager.SearchAndIterateDirectoryAssets(Settings->DirectoriesToRegister, [&Manager, &RegOptions](const FAssetData& AssetData)
 			{
-				FMetasoundAssetBase* MetaSoundAsset = IMetasoundUObjectRegistry::Get().GetObjectAsAssetBase(&AssetObject);
-				check(MetaSoundAsset);
-				MetaSoundAsset->UpdateAndRegisterForExecution(RegOptions);
+				using namespace Metasound;
+
+				Manager.AddOrLoadAndUpdateFromObjectAsync(AssetData, [RegOptions](FMetaSoundAssetKey, UObject& AssetObject)
+				{
+					FMetasoundAssetBase* MetaSoundAsset = IMetasoundUObjectRegistry::Get().GetObjectAsAssetBase(&AssetObject);
+					check(MetaSoundAsset);
+					MetaSoundAsset->UpdateAndRegisterForExecution(RegOptions);
+				});
 			});
-		});
-		AssetSubsystemPrivate::InitialAssetScanComplete.store(true);
+		}
 	}
-}
-
-#if WITH_EDITORONLY_DATA
-void UMetaSoundAssetSubsystem::AddAssetReferences(FMetasoundAssetBase& InAssetBase)
-{
-	using namespace Metasound::Frontend;
-	IMetaSoundAssetManager::GetChecked().AddAssetReferences(InAssetBase);
-}
-#endif // WITH_EDITORONLY_DATA
-
-Metasound::Frontend::FNodeRegistryKey UMetaSoundAssetSubsystem::AddOrUpdateAsset(const UObject& InObject)
-{
-	using namespace Metasound::Frontend;
-	return IMetaSoundAssetManager::GetChecked().AddOrUpdateFromObject(InObject);
-}
-
-Metasound::Frontend::FNodeRegistryKey UMetaSoundAssetSubsystem::AddOrUpdateAsset(const FAssetData& InAssetData)
-{
-	using namespace Metasound::Frontend;
-	IMetaSoundAssetManager::GetChecked().AddOrUpdateFromAssetData(InAssetData);
-	return { };
-}
-
-bool UMetaSoundAssetSubsystem::CanAutoUpdate(const FMetasoundFrontendClassName& InClassName) const
-{
-	using namespace Metasound::Frontend;
-	return IMetaSoundAssetManager::GetChecked().CanAutoUpdate(InClassName);
-}
 
-bool UMetaSoundAssetSubsystem::ContainsKey(const Metasound::Frontend::FNodeRegistryKey& InRegistryKey) const
-{
-	using namespace Metasound::Frontend;
-	return IMetaSoundAssetManager::GetChecked().ContainsKey(FMetaSoundAssetKey(InRegistryKey));
+	AssetSubsystemPrivate::InitialAssetScanComplete.store(true);
 }
 
 #if WITH_EDITOR
@@ -1619,85 +1547,6 @@ bool UMetaSoundAssetSubsystem::FindReferencingAssetClassInfo(TScriptInterface<IM
 }
 #endif // WITH_EDITOR
 
-FMetasoundAssetBase* UMetaSoundAssetSubsystem::GetAsAsset(UObject& InObject) const
-{
-	using namespace Metasound::Frontend;
-	return IMetaSoundAssetManager::GetChecked().GetAsAsset(InObject);
-}
-
-const FMetasoundAssetBase* UMetaSoundAssetSubsystem::GetAsAsset(const UObject& InObject) const
-{
-	using namespace Metasound::Frontend;
-	return IMetaSoundAssetManager::GetChecked().GetAsAsset(InObject);
-}
-
-#if WITH_EDITOR
-PRAGMA_DISABLE_DEPRECATION_WARNINGS
-TSet<UMetaSoundAssetSubsystem::FAssetInfo> UMetaSoundAssetSubsystem::GetReferencedAssetClasses(const FMetasoundAssetBase& InAssetBase) const
-{
-	using namespace Metasound::Frontend;
-	return { };
-}
-PRAGMA_ENABLE_DEPRECATION_WARNINGS
-#endif // WITH_EDITOR
-
-FMetasoundAssetBase* UMetaSoundAssetSubsystem::TryLoadAssetFromKey(const Metasound::Frontend::FNodeRegistryKey& RegistryKey) const
-{
-	using namespace Metasound::Frontend;
-	return IMetaSoundAssetManager::GetChecked().TryLoadAssetFromKey(RegistryKey);
-}
-
-bool UMetaSoundAssetSubsystem::TryLoadReferencedAssets(const FMetasoundAssetBase& InAssetBase, TArray<FMetasoundAssetBase*>& OutReferencedAssets) const
-{
-	using namespace Metasound::Frontend;
-	return IMetaSoundAssetManager::GetChecked().TryLoadReferencedAssets(InAssetBase, OutReferencedAssets);
-}
-
-const FSoftObjectPath* UMetaSoundAssetSubsystem::FindObjectPathFromKey(const Metasound::Frontend::FNodeRegistryKey& InRegistryKey) const
-{
-	using namespace Metasound::Frontend;
-	static FSoftObjectPath TempPath;
-	TempPath.Reset();
-	const FTopLevelAssetPath Path = IMetaSoundAssetManager::GetChecked().FindAssetPath(InRegistryKey);
-	if (Path.IsValid())
-	{
-		TempPath = FSoftObjectPath(Path);
-	}
-	return &TempPath;
-}
-
-FMetasoundAssetBase* UMetaSoundAssetSubsystem::TryLoadAsset(const FSoftObjectPath& InObjectPath) const
-{
-	using namespace Metasound::Frontend;
-	return IMetaSoundAssetManager::GetChecked().TryLoadAsset(InObjectPath);
-}
-
-void UMetaSoundAssetSubsystem::RemoveAsset(const UObject& InObject)
-{
-	using namespace Metasound::Frontend;
-	IMetaSoundAssetManager::GetChecked().RemoveAsset(InObject);
-}
-
-void UMetaSoundAssetSubsystem::RemoveAsset(const FAssetData& InAssetData)
-{
-	using namespace Metasound::Frontend;
-	IMetaSoundAssetManager::GetChecked().RemoveAsset(InAssetData);
-}
-
-void UMetaSoundAssetSubsystem::RenameAsset(const FAssetData& InAssetData, bool /* bInReregisterWithFrontend */)
-{
-	using namespace Metasound::Frontend;
-	IMetaSoundAssetManager::GetChecked().RenameAsset(InAssetData, { });
-}
-
-UMetaSoundAssetSubsystem& UMetaSoundAssetSubsystem::GetChecked()
-{
-	check(GEngine);
-	UMetaSoundAssetSubsystem* Subsystem = GEngine->GetEngineSubsystem<UMetaSoundAssetSubsystem>();
-	check(Subsystem);
-	return *Subsystem;
-}
-
 #if WITH_EDITOR
 bool UMetaSoundAssetSubsystem::ReassignClassName(TScriptInterface<IMetaSoundDocumentInterface> DocInterface)
 {
diff --git a/Engine/Plugins/Runtime/Metasound/Source/MetasoundEngine/Private/MetasoundSettings.cpp b/Engine/Plugins/Runtime/Metasound/Source/MetasoundEngine/Private/MetasoundSettings.cpp
--- a/Engine/Plugins/Runtime/Metasound/Source/MetasoundEngine/Private/MetasoundSettings.cpp
+++ b/Engine/Plugins/Runtime/Metasound/Source/MetasoundEngine/Private/MetasoundSettings.cpp
@@ -521,8 +521,6 @@ void UMetaSoundSettings::PostEditChangeProperty(FPropertyChangedEvent& PropertyC
 		OnPageSettingsUpdated.Broadcast();
 	}
 	UpdateCachedPageOrder();
-
-	DenyListCacheChangeID++;
 }
 
 const FMetaSoundPageSettings* UMetaSoundSettings::GetTargetPageSettingsFromOverride(const FName InTargetPage) const
diff --git a/Engine/Plugins/Runtime/Metasound/Source/MetasoundEngine/Public/MetasoundAssetSubsystem.h b/Engine/Plugins/Runtime/Metasound/Source/MetasoundEngine/Public/MetasoundAssetSubsystem.h
--- a/Engine/Plugins/Runtime/Metasound/Source/MetasoundEngine/Public/MetasoundAssetSubsystem.h
+++ b/Engine/Plugins/Runtime/Metasound/Source/MetasoundEngine/Public/MetasoundAssetSubsystem.h
@@ -84,7 +84,9 @@ PRAGMA_ENABLE_DEPRECATION_WARNINGS
 
 		UE_API bool IsInitialAssetScanComplete() const;
 
+		UE_DEPRECATED(5.7, "DenyList Cache no longer supported (auto-update is now required when cooking and not supported at runtime")
 		UE_API void RebuildDenyListCache(const UAssetManager& InAssetManager);
+
 		UE_API void RegisterAssetClassesInDirectories(const TArray<FMetaSoundAssetDirectory>& Directories);
 
 #if WITH_EDITOR
@@ -111,7 +113,7 @@ PRAGMA_ENABLE_DEPRECATION_WARNINGS
 		UE_API virtual FMetaSoundAssetKey AddOrUpdateFromObject(const UObject& InObject) override;
 		UE_API virtual void AddOrLoadAndUpdateFromObjectAsync(const FAssetData& InAssetData, Frontend::IMetaSoundAssetManager::FOnUpdatedAssetLoaded&& OnUpdatedAssetLoaded) override;
 		UE_API virtual void AddOrUpdateFromAssetData(const FAssetData& InAssetData) override;
-		UE_API virtual bool CanAutoUpdate(const FMetasoundFrontendClassName& InClassName) const override;
+
 		UE_API virtual bool ContainsKey(const FMetaSoundAssetKey& InKey) const override;
 		UE_API virtual bool ContainsKey(const Frontend::FNodeRegistryKey& InRegistryKey) const override;
 		UE_API virtual FMetasoundAssetBase* FindAsset(const FMetaSoundAssetKey& InKey) const override;
@@ -156,6 +158,7 @@ PRAGMA_ENABLE_DEPRECATION_WARNINGS
 		virtual FString GetReferencerName() const override { return TEXT("FMetaSoundAssetManager"); }
 
 		Frontend::FMetaSoundAssetRegistrationOptions GetRegistrationOptions() const;
+
 	private:
 		FMetaSoundAssetKey AddOrUpdateFromObjectInternal(const UObject& InObject);
 
@@ -189,9 +192,6 @@ PRAGMA_ENABLE_DEPRECATION_WARNINGS
 		FStreamableManager StreamableManager;
 		int32 AsyncLoadIDCounter = 0;
 
-		int32 AutoUpdateDenyListChangeID = INDEX_NONE;
-		TSet<FName> AutoUpdateDenyListCache;
-
 		FAssetClassInfoMap ClassInfoMap;
 
 		// Critical section primarily for allowing safe access of class info map during async loading of MetaSound assets.
@@ -249,68 +249,6 @@ UE_API PRAGMA_ENABLE_DEPRECATION_WARNINGS
 
 	virtual void Initialize(FSubsystemCollectionBase& InCollection) override;
 
-	UE_DEPRECATED(5.5, "Moved to internal implementation, use IMetaSoundAssetManager::GetChecked() and analogous call")
-	UE_API virtual void RemoveAsset(const UObject& InObject);
-
-	UE_DEPRECATED(5.5, "Moved to internal implementation, use IMetaSoundAssetManager::GetChecked() and analogous call")
-	UE_API virtual void RemoveAsset(const FAssetData& InAssetData);
-
-	UE_DEPRECATED(5.5, "Moved to internal implementation, use IMetaSoundAssetManager::GetChecked() and analogous call")
-	UE_API virtual void RenameAsset(const FAssetData& InAssetData, bool bInReregisterWithFrontend = true);
-
-#if WITH_EDITORONLY_DATA
-	UE_DEPRECATED(5.5, "Moved to internal implementation, use IMetaSoundAssetManager::GetChecked() and analogous call")
-	UE_API virtual void AddAssetReferences(FMetasoundAssetBase& InAssetBase);
-#endif
-
-	UE_DEPRECATED(5.5, "Moved to internal implementation, use IMetaSoundAssetManager::GetChecked() 'AddOrLoadAndUpdateFromObjectAsync' call")
-	UE_API Metasound::Frontend::FNodeRegistryKey AddOrUpdateAsset(const FAssetData& InAssetData);
-
-	UE_DEPRECATED(5.5, "Moved to internal implementation, use IMetaSoundAssetManager::GetChecked() 'AddOrUpdateFromObject' call")
-	UE_API virtual Metasound::Frontend::FNodeRegistryKey AddOrUpdateAsset(const UObject& InObject);
-
-	UE_DEPRECATED(5.5, "Moved to internal implementation, use IMetaSoundAssetManager::GetChecked() and analogous call")
-	UE_API virtual bool CanAutoUpdate(const FMetasoundFrontendClassName& InClassName) const;
-
-	UE_DEPRECATED(5.5, "Moved to internal implementation, use IMetaSoundAssetManager::GetChecked() and analogous call")
-	UE_API virtual bool ContainsKey(const Metasound::Frontend::FNodeRegistryKey& InRegistryKey) const;
-
-	UE_DEPRECATED(5.5, "Moved to internal implementation, use IMetaSoundAssetManager::GetChecked() and analogous call")
-	UE_API virtual const FSoftObjectPath* FindObjectPathFromKey(const Metasound::Frontend::FNodeRegistryKey& RegistryKey) const;
-
-	UE_DEPRECATED(5.5, "Moved to internal implementation, use IMetaSoundAssetManager::GetChecked() and analogous call")
-	UE_API virtual FMetasoundAssetBase* GetAsAsset(UObject& InObject) const;
-
-	UE_DEPRECATED(5.5, "Moved to internal implementation, use IMetaSoundAssetManager::GetChecked() and analogous call")
-	UE_API virtual const FMetasoundAssetBase* GetAsAsset(const UObject& InObject) const;
-
-	UE_DEPRECATED(5.5, "Implementation of MetaSound asset management has been moved to raw c++ implementation for more reliable, monolithic "
-		"lifetime management. This subsystem continues to exist only for Blueprint-related asset functionality. "
-		"Use IMetaSoundAssetManager::GetChecked() instead")
-	static UE_API UMetaSoundAssetSubsystem& GetChecked();
-
-#if WITH_EDITOR
-PRAGMA_DISABLE_DEPRECATION_WARNINGS
-	UE_DEPRECATED(5.5, "Moved to internal implementation, use IMetaSoundAssetManager::GetChecked() and analogous call")
-	UE_API virtual TSet<FAssetInfo> GetReferencedAssetClasses(const FMetasoundAssetBase& InAssetBase) const;
-PRAGMA_ENABLE_DEPRECATION_WARNINGS
-#endif
-
-	UE_DEPRECATED(5.5, "Moved to internal implementation, use IMetaSoundAssetManager::GetChecked() and analogous call")
-	UE_API virtual FMetasoundAssetBase* TryLoadAsset(const FSoftObjectPath& InObjectPath) const;
-
-	UE_DEPRECATED(5.5, "Moved to internal implementation, use IMetaSoundAssetManager::GetChecked() and analogous call")
-	UE_API virtual FMetasoundAssetBase* TryLoadAssetFromKey(const Metasound::Frontend::FNodeRegistryKey& RegistryKey) const;
-
-	UE_DEPRECATED(5.5, "Moved to internal implementation, use IMetaSoundAssetManager::GetChecked() and analogous call")
-	UE_API virtual bool TryLoadReferencedAssets(const FMetasoundAssetBase& InAssetBase, TArray<FMetasoundAssetBase*>& OutReferencedAssets) const;
-
-	UE_DEPRECATED(5.5, "Moved to internal implementation, use IMetaSoundAssetManager::GetChecked() and analogous call")
-	virtual void RequestAsyncLoadReferencedAssets(FMetasoundAssetBase& InAssetBase) { }
-
-	UE_DEPRECATED(5.5, "Moved to internal implementation, use IMetaSoundAssetManager::GetChecked() and analogous call")
-	virtual void WaitUntilAsyncLoadReferencedAssetsComplete(FMetasoundAssetBase& InAssetBase) { }
-
 #if WITH_EDITOR
 	// Returns asset class info for the given MetaSound asset. Will attempt to get
 	// info without loading the asset if its tag data is up-to-date, or if set to force load, will load
@@ -357,25 +295,9 @@ PRAGMA_ENABLE_DEPRECATION_WARNINGS
 	UFUNCTION(BlueprintCallable, Category = "Audio|MetaSound|Registration")
 	UE_API void UnregisterAssetClassesInDirectories(const TArray<FMetaSoundAssetDirectory>& Directories);
 
-protected:
-	UE_DEPRECATED(5.5, "Moved to private implementation")
-	void PostEngineInit() { }
-
-	UE_DEPRECATED(5.5, "Moved to private implementation")
-	void PostInitAssetScan() { }
-
-	UE_DEPRECATED(5.5, "Moved to internal implementation, use IMetaSoundAssetManager::GetChecked() and analogous call")
-	void RebuildDenyListCache(const UAssetManager& InAssetManager) { }
-
-	UE_DEPRECATED(5.5, "Use FMetaSoundDocumentBuilder::SetDisplayName instead (call now only available with editor compiled)")
-	void ResetAssetClassDisplayName(const FAssetData& InAssetData) { }
-
-	UE_DEPRECATED(5.5, "Moved to internal implementation, use IMetaSoundAssetManager::GetChecked() and analogous call")
-	void SearchAndIterateDirectoryAssets(const TArray<FDirectoryPath>& InDirectories, TFunctionRef<void(const FAssetData&)> InFunction) { }
-
 private:
-	UE_API void PostEngineInitInternal();
-	UE_API void PostInitAssetScanInternal();
+	void PostEngineInitInternal();
+	void PostInitAssetScanInternal();
 };
 
 #undef UE_API
diff --git a/Engine/Plugins/Runtime/Metasound/Source/MetasoundEngine/Public/MetasoundSettings.h b/Engine/Plugins/Runtime/Metasound/Source/MetasoundEngine/Public/MetasoundSettings.h
--- a/Engine/Plugins/Runtime/Metasound/Source/MetasoundEngine/Public/MetasoundSettings.h
+++ b/Engine/Plugins/Runtime/Metasound/Source/MetasoundEngine/Public/MetasoundSettings.h
@@ -183,9 +183,12 @@ public:
 	  */
 	UPROPERTY(EditAnywhere, config, Category = Registration, meta = (RelativePath, LongPackageName))
 	TArray<FDirectoryPath> DirectoriesToRegister;
-		
+
+#if WITH_EDITORONLY_DATA
+	UE_DEPRECATED(5.7, "DenyList no longer supported (auto-update automatically executed during cook and not supported at runtime")
 	UPROPERTY(Transient)
 	int32 DenyListCacheChangeID = 0;
+#endif // WITH_EDITORONLY_DATA
 
 private:
 #if WITH_EDITORONLY_DATA
diff --git a/Engine/Plugins/Runtime/Metasound/Source/MetasoundFrontend/Private/MetasoundDocumentInterface.cpp b/Engine/Plugins/Runtime/Metasound/Source/MetasoundFrontend/Private/MetasoundDocumentInterface.cpp
--- a/Engine/Plugins/Runtime/Metasound/Source/MetasoundFrontend/Private/MetasoundDocumentInterface.cpp
+++ b/Engine/Plugins/Runtime/Metasound/Source/MetasoundFrontend/Private/MetasoundDocumentInterface.cpp
@@ -36,9 +36,4 @@ namespace Metasound::Frontend
 		check(!DocumentInterfacePrivate::Instance.IsValid());
 		DocumentInterfacePrivate::Instance = MoveTemp(InInstance);
 	}
-
-	IMetaSoundDocumentBuilderRegistry& IMetaSoundDocumentBuilderRegistry::GetChecked()
-	{
-		return static_cast<IMetaSoundDocumentBuilderRegistry&>(IDocumentBuilderRegistry::GetChecked());
-	}
 } // namespace Metasound::Frontend
\ No newline at end of file
diff --git a/Engine/Plugins/Runtime/Metasound/Source/MetasoundFrontend/Private/MetasoundFrontendNodeController.cpp b/Engine/Plugins/Runtime/Metasound/Source/MetasoundFrontend/Private/MetasoundFrontendNodeController.cpp
--- a/Engine/Plugins/Runtime/Metasound/Source/MetasoundFrontend/Private/MetasoundFrontendNodeController.cpp
+++ b/Engine/Plugins/Runtime/Metasound/Source/MetasoundFrontend/Private/MetasoundFrontendNodeController.cpp
@@ -1232,7 +1232,7 @@ namespace Metasound
 
 			const FMetasoundFrontendClassMetadata& NodeClassMetadata = GetClassMetadata();
 			const IMetaSoundAssetManager& AssetManager = IMetaSoundAssetManager::GetChecked();
-			if (AssetManager.IsAssetClass(NodeClassMetadata) && (!AssetManager.CanAutoUpdate(NodeClassMetadata.GetClassName())))
+			if (AssetManager.IsAssetClass(NodeClassMetadata))
 			{
 				return false;
 			}
diff --git a/Engine/Plugins/Runtime/Metasound/Source/MetasoundFrontend/Public/MetasoundAssetManager.h b/Engine/Plugins/Runtime/Metasound/Source/MetasoundFrontend/Public/MetasoundAssetManager.h
--- a/Engine/Plugins/Runtime/Metasound/Source/MetasoundFrontend/Public/MetasoundAssetManager.h
+++ b/Engine/Plugins/Runtime/Metasound/Source/MetasoundFrontend/Public/MetasoundAssetManager.h
@@ -272,8 +272,8 @@ namespace Metasound::Frontend
 		// and adds asynchronously if asset or associated tag schema is out-of-date.
 		virtual void AddOrUpdateFromAssetData(const FAssetData& InAssetData) = 0;
 
-		// Whether or not the class is eligible for auto-update
-		virtual bool CanAutoUpdate(const FMetasoundFrontendClassName& InClassName) const = 0;
+		UE_DEPRECATED(5.7, "AutoUpdate is required now prior to MetaSounds executing runtime builds")
+		virtual bool CanAutoUpdate(const FMetasoundFrontendClassName& InClassName) const { return true; }
 
 		// Whether or not the asset manager has loaded the given asset
 		virtual bool ContainsKey(const FMetaSoundAssetKey& InAssetKey) const = 0;
diff --git a/Engine/Plugins/Runtime/Metasound/Source/MetasoundFrontend/Public/MetasoundDocumentInterface.h b/Engine/Plugins/Runtime/Metasound/Source/MetasoundFrontend/Public/MetasoundDocumentInterface.h
--- a/Engine/Plugins/Runtime/Metasound/Source/MetasoundFrontend/Public/MetasoundDocumentInterface.h
+++ b/Engine/Plugins/Runtime/Metasound/Source/MetasoundFrontend/Public/MetasoundDocumentInterface.h
@@ -144,24 +144,6 @@ namespace Metasound::Frontend
 		static UE_API IDocumentBuilderRegistry& GetChecked();
 		static UE_API void Deinitialize();
 		static UE_API void Initialize(TUniquePtr<IDocumentBuilderRegistry>&& InBuilderRegistry);
-
-	protected:
-		UE_DEPRECATED(5.5, "Use 'Initialize' instead")
-		static void Set(TUniqueFunction<IDocumentBuilderRegistry&()>&& InGetInstance) { checkNoEntry(); }
-	};
-
-	class IMetaSoundDocumentBuilderRegistry : public IDocumentBuilderRegistry
-	{
-	public:
-		UE_DEPRECATED(5.5, "Class Renamed: use IDocumentBuilderRegistry instead")
-		IMetaSoundDocumentBuilderRegistry() = default;
-		virtual ~IMetaSoundDocumentBuilderRegistry() = default;
-
-		UE_DEPRECATED(5.4, "Public exposition of modify delegates no longer available to discourage unsafe manipulation of builder document cache")
-		virtual const FDocumentModifyDelegates* FindModifyDelegates(const FMetasoundFrontendClassName& InClassName) const { return nullptr; }
-
-		UE_DEPRECATED(5.4, "Use 'IDocumentBuilderRegistry' instead")
-		static UE_API IMetaSoundDocumentBuilderRegistry& GetChecked();
 	};
 } // namespace Metasound::Frontend
 
