Deprecate MetaSound AutoUpdate project settings no longer used
#rnx

diff --git a/Engine/Plugins/Runtime/Metasound/Source/MetasoundEngine/Private/MetasoundSettings.cpp b/Engine/Plugins/Runtime/Metasound/Source/MetasoundEngine/Private/MetasoundSettings.cpp
--- a/Engine/Plugins/Runtime/Metasound/Source/MetasoundEngine/Private/MetasoundSettings.cpp
+++ b/Engine/Plugins/Runtime/Metasound/Source/MetasoundEngine/Private/MetasoundSettings.cpp
@@ -702,6 +702,11 @@ void UMetaSoundSettings::UpdateCachedPageOrder()
 }
 #endif // WITH_EDITOR
 
+FName UMetaSoundSettings::GetCategoryName() const
+{
+	return FName("Plugins");
+}
+
 TArray<FGuid> UMetaSoundSettings::GeneratePageOrderInternal(const FName& InPlatformName, const FGuid& InTargetPageID) const
 {
 	// Reverse iterate over page settings to build the priority stack. Every page
diff --git a/Engine/Plugins/Runtime/Metasound/Source/MetasoundEngine/Public/MetasoundSettings.h b/Engine/Plugins/Runtime/Metasound/Source/MetasoundEngine/Public/MetasoundSettings.h
--- a/Engine/Plugins/Runtime/Metasound/Source/MetasoundEngine/Public/MetasoundSettings.h
+++ b/Engine/Plugins/Runtime/Metasound/Source/MetasoundEngine/Public/MetasoundSettings.h
@@ -157,20 +157,16 @@ class UMetaSoundSettings : public UDeveloperSettings
 	GENERATED_BODY()
 
 public:
-	virtual FName GetCategoryName() const override { return FName(TEXT("Plugins")); }
-
-	/** If true, AutoUpdate is enabled, increasing load times.  If false, skips AutoUpdate on load, but can result in MetaSounds failing to load, 
-	  * register, and execute if interface differences are present. */
-	UPROPERTY(EditAnywhere, config, Category = AutoUpdate)
+#if WITH_EDITORONLY_DATA
+	UPROPERTY(meta = (Deprecated, DeprecationMessage = "5.7 - For performance reasons, AutoUpdate is now always only called during cook."))
 	bool bAutoUpdateEnabled = true;
 
-	/** List of native MetaSound classes whose node references should not be AutoUpdated. */
-	UPROPERTY(EditAnywhere, config, Category = AutoUpdate, meta = (DisplayName = "DenyList", EditCondition = "bAutoUpdateEnabled"))
+	UPROPERTY(meta = (Deprecated, DeprecationMessage = "5.7 - For performance reasons, AutoUpdate is now always only called during cook."))
 	TArray<FMetasoundFrontendClassName> AutoUpdateDenylist;
 
-	/** List of MetaSound assets whose node references should not be AutoUpdated. */
-	UPROPERTY(EditAnywhere, config, Category = AutoUpdate, meta = (DisplayName = "Asset DenyList", EditCondition = "bAutoUpdateEnabled"))
+	UPROPERTY(meta = (Deprecated, DeprecationMessage = "For performance reasons, AutoUpdate is now always only called during cook."))
 	TArray<FDefaultMetaSoundAssetAutoUpdateSettings> AutoUpdateAssetDenylist;
+#endif // WITH_EDITORONLY_DATA
 
 	/** If true, warnings will be logged if updating a node results in existing connections being discarded. */
 	UPROPERTY(EditAnywhere, config, Category = AutoUpdate, meta = (DisplayName = "Log Warning on Dropped Connection", EditCondition = "bAutoUpdateEnabled"))
@@ -178,7 +174,7 @@ public:
 
 	/** Directories to scan & automatically register MetaSound post initial asset scan on engine start-up.
 	  * May speed up subsequent calls to playback MetaSounds post asset scan but increases application load time.
-	  * See 'MetaSoundAssetSubsystem::RegisterAssetClassesInDirectories' to dynamically register or 
+	  * See 'MetaSoundAssetSubsystem::RegisterAssetClassesInDirectories' to dynamically register or
 	  * 'MetaSoundAssetSubsystem::UnregisterAssetClassesInDirectories' to unregister asset classes.
 	  */
 	UPROPERTY(EditAnywhere, config, Category = Registration, meta = (RelativePath, LongPackageName))
@@ -186,7 +182,7 @@ public:
 
 #if WITH_EDITORONLY_DATA
 	UE_DEPRECATED(5.7, "DenyList no longer supported (auto-update automatically executed during cook and not supported at runtime")
-	UPROPERTY(Transient)
+		UPROPERTY(Transient)
 	int32 DenyListCacheChangeID = 0;
 #endif // WITH_EDITORONLY_DATA
 
@@ -228,6 +224,8 @@ private:
 	mutable FName CookPlatformTargetPage;
 #endif // WITH_EDITORONLY_DATA
 
+	virtual FName GetCategoryName() const override;
+
 	// Generates the page order for a given platform and desired target page ID
 	TArray<FGuid> GeneratePageOrderInternal(const FName& InPlatformName, const FGuid& InTargetPageID) const;
 
