Only enable metasound scopes when named events are enabled

diff --git a/Engine/Plugins/Runtime/Metasound/Source/MetasoundGraphCore/Public/MetasoundTrace.h b/Engine/Plugins/Runtime/Metasound/Source/MetasoundGraphCore/Public/MetasoundTrace.h
--- a/Engine/Plugins/Runtime/Metasound/Source/MetasoundGraphCore/Public/MetasoundTrace.h
+++ b/Engine/Plugins/Runtime/Metasound/Source/MetasoundGraphCore/Public/MetasoundTrace.h
@@ -6,7 +6,7 @@
 #include "ProfilingDebugging/CpuProfilerTrace.h"
 
 #if !defined(METASOUND_CPUPROFILERTRACE_ENABLED)
-#if CPUPROFILERTRACE_ENABLED
+#if CPUPROFILERTRACE_ENABLED && ENABLE_NAMED_EVENTS
 #define METASOUND_CPUPROFILERTRACE_ENABLED 1
 #else
 #define METASOUND_CPUPROFILERTRACE_ENABLED 0
