Fix issues with MetaSound variables copy paste and multiple locations
- Add node validation for multiple locations and remove extra locations in setter

diff --git a/Engine/Plugins/Runtime/Metasound/Source/MetasoundEditor/Private/MetasoundEditorDocumentClipboardUtils.cpp b/Engine/Plugins/Runtime/Metasound/Source/MetasoundEditor/Private/MetasoundEditorDocumentClipboardUtils.cpp
--- a/Engine/Plugins/Runtime/Metasound/Source/MetasoundEditor/Private/MetasoundEditorDocumentClipboardUtils.cpp
+++ b/Engine/Plugins/Runtime/Metasound/Source/MetasoundEditor/Private/MetasoundEditorDocumentClipboardUtils.cpp
@@ -18,13 +18,11 @@
 
 namespace Metasound::Editor
 {
-	void FDocumentClipboardUtils::ProcessPastedInputNodes(FMetasoundAssetBase& OutAsset, TArray<UMetasoundEditorGraphNode*>& OutPastedNodes)
+	void FDocumentClipboardUtils::ProcessPastedInputNodes(FMetaSoundFrontendDocumentBuilder& Builder, FMetasoundAssetBase& OutAsset, TArray<UMetasoundEditorGraphNode*>& OutPastedNodes)
 	{
 		using namespace Engine;
 		using namespace Frontend;
 
-		FMetaSoundFrontendDocumentBuilder& Builder = FDocumentBuilderRegistry::GetChecked().FindOrBeginBuilding(OutAsset.GetOwningAsset());
-
 		TMap<FName, TObjectPtr<UMetasoundEditorGraphInput>> MappedGeneratedInputNames;
 		UMetasoundEditorGraph& Graph = *CastChecked<UMetasoundEditorGraph>(&OutAsset.GetGraphChecked());
 		for (int32 Index = OutPastedNodes.Num() - 1; Index >= 0; --Index)
@@ -134,11 +132,10 @@ namespace Metasound::Editor
 		}
 	}
 
-	void FDocumentClipboardUtils::ProcessPastedOutputNodes(FMetasoundAssetBase& OutAsset, TArray<UMetasoundEditorGraphNode*>& OutPastedNodes, FDocumentPasteNotifications& OutNotifications)
+	void FDocumentClipboardUtils::ProcessPastedOutputNodes(FMetaSoundFrontendDocumentBuilder& Builder, FMetasoundAssetBase& OutAsset, TArray<UMetasoundEditorGraphNode*>& OutPastedNodes, FDocumentPasteNotifications& OutNotifications)
 	{
 		using namespace Frontend;
 		using namespace Engine;
-		FMetaSoundFrontendDocumentBuilder& Builder = FDocumentBuilderRegistry::GetChecked().FindOrBeginBuilding(OutAsset.GetOwningAsset());
 		UMetasoundEditorGraph& Graph = *CastChecked<UMetasoundEditorGraph>(&OutAsset.GetGraphChecked());
 		for (int32 Index = OutPastedNodes.Num() - 1; Index >= 0; --Index)
 		{
@@ -332,7 +329,7 @@ namespace Metasound::Editor
 		}
 	}
 
-	void FDocumentClipboardUtils::ProcessPastedExternalNodes(FMetasoundAssetBase& OutAsset, TArray<UMetasoundEditorGraphNode*>& OutPastedNodes, FDocumentPasteNotifications& OutNotifications)
+	void FDocumentClipboardUtils::ProcessPastedExternalNodes(FMetaSoundFrontendDocumentBuilder& Builder, FMetasoundAssetBase& OutAsset, TArray<UMetasoundEditorGraphNode*>& OutPastedNodes, FDocumentPasteNotifications& OutNotifications)
 	{
 		using namespace Engine;
 		using namespace Frontend;
@@ -374,7 +371,6 @@ namespace Metasound::Editor
 				}
 				else
 				{
-					FMetaSoundFrontendDocumentBuilder& Builder = IDocumentBuilderRegistry::GetChecked().FindOrBeginBuilding(OutAsset.GetOwningAsset());
 					const FMetasoundFrontendNode* NewNode = Builder.AddNodeByClassName(Breadcrumb.ClassName);
 					if (NewNode)
 					{
@@ -393,7 +389,6 @@ namespace Metasound::Editor
 			{
 				if (const INodeTemplate* Template = INodeTemplateRegistry::Get().FindTemplate(Breadcrumb.ClassName))
 				{
-					FMetaSoundFrontendDocumentBuilder& Builder = IDocumentBuilderRegistry::GetChecked().FindOrBeginBuilding(OutAsset.GetOwningAsset());
 					const FNodeTemplateGenerateInterfaceParams TemplateParams = Breadcrumb.TemplateParams.IsSet() ? *Breadcrumb.TemplateParams : FNodeTemplateGenerateInterfaceParams();
 					const FMetasoundFrontendNode* TemplateNode = Builder.AddNodeByTemplate(*Template, TemplateParams);
 					if (TemplateNode)
@@ -411,7 +406,6 @@ namespace Metasound::Editor
 					FMetasoundFrontendClass ExternalClass;
 					if (ISearchEngine::Get().FindClassWithHighestVersion(Breadcrumb.ClassName, ExternalClass))
 					{
-						FMetaSoundFrontendDocumentBuilder& Builder = IDocumentBuilderRegistry::GetChecked().FindOrBeginBuilding(OutAsset.GetOwningAsset());
 						const FMetasoundFrontendNode* NewNode = Builder.AddNodeByClassName(Breadcrumb.ClassName);
 						if (NewNode)
 						{
@@ -436,13 +430,12 @@ namespace Metasound::Editor
 		}
 	}
 
-	void FDocumentClipboardUtils::ProcessPastedCommentNodes(FMetasoundAssetBase& OutAsset, const TArrayView<UMetasoundEditorGraphCommentNode*> CommentNodes)
+	void FDocumentClipboardUtils::ProcessPastedCommentNodes(FMetaSoundFrontendDocumentBuilder& Builder, FMetasoundAssetBase& OutAsset, const TArrayView<UMetasoundEditorGraphCommentNode*> CommentNodes)
 	{
 		using namespace Engine;
 		using namespace Frontend;
 
 		UMetasoundEditorGraph& Graph = *CastChecked<UMetasoundEditorGraph>(&OutAsset.GetGraphChecked());
-		UMetaSoundBuilderBase& Builder = FDocumentBuilderRegistry::GetChecked().FindOrBeginBuilding(*OutAsset.GetOwningAsset());
 
 		for (UMetasoundEditorGraphCommentNode* CommentNode : CommentNodes)
 		{
@@ -456,7 +449,7 @@ namespace Metasound::Editor
 		}
 	}
 
-	void FDocumentClipboardUtils::ProcessPastedNodePositions(FMetasoundAssetBase& OutAsset, const FVector2D& InLocation, TArray<UMetasoundEditorGraphNode*>& OutPastedNodes, const TArrayView<UMetasoundEditorGraphCommentNode*> CommentNodes)
+	void FDocumentClipboardUtils::ProcessPastedNodePositions(FMetaSoundFrontendDocumentBuilder& OutBuilder, const FVector2D& InLocation, TArray<UMetasoundEditorGraphNode*>& OutPastedNodes, const TArrayView<UMetasoundEditorGraphCommentNode*> CommentNodes)
 	{
 		using namespace Frontend;
 
@@ -489,16 +482,8 @@ namespace Metasound::Editor
 			GraphNode->SnapToGrid(SNodePanel::GetSnapGridSize());
 			if (UMetasoundEditorGraphNode* MetasoundGraphNode = Cast<UMetasoundEditorGraphNode>(GraphNode))
 			{
-PRAGMA_DISABLE_DEPRECATION_WARNINGS
-				FNodeHandle NodeHandle = MetasoundGraphNode->GetNodeHandle();
-PRAGMA_ENABLE_DEPRECATION_WARNINGS
-				if (ensure(NodeHandle->IsValid()))
-				{
-					const FVector2D NewNodeLocation = FVector2D(GraphNode->NodePosX, GraphNode->NodePosY);
-					FMetasoundFrontendNodeStyle NodeStyle = NodeHandle->GetNodeStyle();
-					NodeStyle.Display.Locations.FindOrAdd(MetasoundGraphNode->NodeGuid) = NewNodeLocation;
-					NodeHandle->SetNodeStyle(NodeStyle);
-				}
+				const FVector2D NewNodeLocation = FVector2D(GraphNode->NodePosX, GraphNode->NodePosY);
+				OutBuilder.SetNodeLocation(MetasoundGraphNode->GetNodeID(), NewNodeLocation, &MetasoundGraphNode->NodeGuid);
 			}
 		}
 
@@ -591,6 +576,8 @@ PRAGMA_ENABLE_DEPRECATION_WARNINGS
 
 	TArray<UEdGraphNode*> FDocumentClipboardUtils::PasteClipboardString(const FText& InTransactionText, const FString& InClipboardString, const FVector2D& InLocation, UObject& OutMetaSound, FDocumentPasteNotifications& OutNotifications)
 	{
+		using namespace Metasound::Engine;
+
 		FMetasoundAssetBase* Asset = IMetasoundUObjectRegistry::Get().GetObjectAsAssetBase(&OutMetaSound);
 		check(Asset);
 
@@ -618,12 +605,14 @@ PRAGMA_ENABLE_DEPRECATION_WARNINGS
 			return PastedNodes;
 		}
 
-		ProcessPastedCommentNodes(*Asset, PastedCommentNodes);
-		ProcessPastedInputNodes(*Asset, PastedGraphNodes);
-		ProcessPastedOutputNodes(*Asset, PastedGraphNodes, OutNotifications);
+		FMetaSoundFrontendDocumentBuilder& Builder = FDocumentBuilderRegistry::GetChecked().FindOrBeginBuilding(Asset->GetOwningAsset());
+
+		ProcessPastedCommentNodes(Builder, *Asset, PastedCommentNodes);
+		ProcessPastedInputNodes(Builder, *Asset, PastedGraphNodes);
+		ProcessPastedOutputNodes(Builder, *Asset, PastedGraphNodes, OutNotifications);
 		ProcessPastedVariableNodes(*Asset, PastedGraphNodes, OutNotifications);
-		ProcessPastedExternalNodes(*Asset, PastedGraphNodes, OutNotifications);
-		ProcessPastedNodePositions(*Asset, InLocation, PastedGraphNodes, PastedCommentNodes);
+		ProcessPastedExternalNodes(Builder, *Asset, PastedGraphNodes, OutNotifications);
+		ProcessPastedNodePositions(Builder, InLocation, PastedGraphNodes, PastedCommentNodes);
 		ProcessPastedNodeConnections(*Asset, PastedGraphNodes);
 
 		PastedNodes.Append(MoveTemp(PastedGraphNodes));
diff --git a/Engine/Plugins/Runtime/Metasound/Source/MetasoundEditor/Private/MetasoundEditorDocumentClipboardUtils.h b/Engine/Plugins/Runtime/Metasound/Source/MetasoundEditor/Private/MetasoundEditorDocumentClipboardUtils.h
--- a/Engine/Plugins/Runtime/Metasound/Source/MetasoundEditor/Private/MetasoundEditorDocumentClipboardUtils.h
+++ b/Engine/Plugins/Runtime/Metasound/Source/MetasoundEditor/Private/MetasoundEditorDocumentClipboardUtils.h
@@ -36,17 +36,17 @@ namespace Metasound::Editor
 		static const bool CanImportMemberFromText(const FString& TextToImport);
 
 	private:
-		static void ProcessPastedCommentNodes(FMetasoundAssetBase& OutAsset, const TArrayView<UMetasoundEditorGraphCommentNode*> CommentNodes);
+		static void ProcessPastedCommentNodes(FMetaSoundFrontendDocumentBuilder& Builder, FMetasoundAssetBase& OutAsset, const TArrayView<UMetasoundEditorGraphCommentNode*> CommentNodes);
 
-		static void ProcessPastedExternalNodes(FMetasoundAssetBase& OutAsset, TArray<UMetasoundEditorGraphNode*>& OutPastedNodes, FDocumentPasteNotifications& OutNotifications);
+		static void ProcessPastedExternalNodes(FMetaSoundFrontendDocumentBuilder& Builder, FMetasoundAssetBase& OutAsset, TArray<UMetasoundEditorGraphNode*>& OutPastedNodes, FDocumentPasteNotifications& OutNotifications);
 
-		static void ProcessPastedInputNodes(FMetasoundAssetBase& OutAsset, TArray<UMetasoundEditorGraphNode*>& OutPastedNodes);
+		static void ProcessPastedInputNodes(FMetaSoundFrontendDocumentBuilder& Builder, FMetasoundAssetBase& OutAsset, TArray<UMetasoundEditorGraphNode*>& OutPastedNodes);
 
-		static void ProcessPastedNodePositions(FMetasoundAssetBase& OutAsset, const FVector2D& InLocation, TArray<UMetasoundEditorGraphNode*>& OutPastedNodes, const TArrayView<UMetasoundEditorGraphCommentNode*> CommentNodes);
+		static void ProcessPastedNodePositions(FMetaSoundFrontendDocumentBuilder& OutAsset, const FVector2D& InLocation, TArray<UMetasoundEditorGraphNode*>& OutPastedNodes, const TArrayView<UMetasoundEditorGraphCommentNode*> CommentNodes);
 
 		static void ProcessPastedNodeConnections(FMetasoundAssetBase& OutAsset, TArray<UMetasoundEditorGraphNode*>& OutPastedNodes);
 
-		static void ProcessPastedOutputNodes(FMetasoundAssetBase& OutAsset, TArray<UMetasoundEditorGraphNode*>& OutPastedNodes, FDocumentPasteNotifications& OutNotifications);
+		static void ProcessPastedOutputNodes(FMetaSoundFrontendDocumentBuilder& Builder, FMetasoundAssetBase& OutAsset, TArray<UMetasoundEditorGraphNode*>& OutPastedNodes, FDocumentPasteNotifications& OutNotifications);
 
 		static void ProcessPastedVariableNodes(FMetasoundAssetBase& OutAsset, TArray<UMetasoundEditorGraphNode*>& OutPastedNodes, FDocumentPasteNotifications& OutNotifications);
 	};
diff --git a/Engine/Plugins/Runtime/Metasound/Source/MetasoundEditor/Private/MetasoundEditorGraphBuilder.cpp b/Engine/Plugins/Runtime/Metasound/Source/MetasoundEditor/Private/MetasoundEditorGraphBuilder.cpp
--- a/Engine/Plugins/Runtime/Metasound/Source/MetasoundEditor/Private/MetasoundEditorGraphBuilder.cpp
+++ b/Engine/Plugins/Runtime/Metasound/Source/MetasoundEditor/Private/MetasoundEditorGraphBuilder.cpp
@@ -2076,21 +2076,31 @@ PRAGMA_ENABLE_DEPRECATION_WARNINGS
 							const bool bLocationFound = EditorNode->SyncLocationFromFrontendNode(bUpdateEditorNodeID);
 							if (bLocationFound)
 							{
-								bFoundEditorNode = true;
 								FAssociatedNodes& AssociatedNodeData = AssociatedNodes.FindOrAdd(Node.GetID());
-								if (AssociatedNodeData.Node)
+								
+								// If another editor node already exists for this frontend node, skip it 
+								// (leaving it in EditorNodes will cause it to be removed later)
+								// This is an editor only "fix" for frontend nodes with multiple locations, which should not be allowed past 
+								// transient ed graph migration (doc version 1.12)
+								if (AssociatedNodeData.EditorNodes.IsEmpty())
 								{
-									ensure(AssociatedNodeData.Node == &Node);
-								}
-								else
-								{
-									AssociatedNodeData.Node = &Node;
-									AssociatedNodeData.Class = &Class;
-								}
+									bFoundEditorNode = true;
+
+									if (AssociatedNodeData.Node)
+									{
+										ensure(AssociatedNodeData.Node == &Node);
+									}
+									else
+									{
+										AssociatedNodeData.Node = &Node;
+										AssociatedNodeData.Class = &Class;
+									}
+
+									EditorNode->SyncCommentFromFrontendNode();
 
-								EditorNode->SyncCommentFromFrontendNode();
-								AssociatedNodeData.EditorNodes.Add(EditorNode);
-								EditorNodes.RemoveAtSwap(j, EAllowShrinking::No);
+									AssociatedNodeData.EditorNodes.Add(EditorNode);
+									EditorNodes.RemoveAtSwap(j, EAllowShrinking::No);
+								}
 							}
 						}
 					}
@@ -2180,21 +2190,34 @@ PRAGMA_ENABLE_DEPRECATION_WARNINGS
 
 					if (ensure(NewGraphNode))
 					{
-						ensureMsgf(NewGraphNode->NodeGuid.IsValid(), TEXT("New editor NodeGuid must be valid."));
-						ensureMsgf(NewGraphNode->NodeGuid == Location.Key, TEXT("New editor NodeGuid must match location key"));
 
 						FAssociatedNodes& AssociatedNodeData = AssociatedNodes.FindOrAdd(NodeID);
-						if (AssociatedNodeData.Node)
+						if (AssociatedNodeData.EditorNodes.IsEmpty())
 						{
-							ensure(AssociatedNodeData.Node == Pair.Key);
+							ensureMsgf(NewGraphNode->NodeGuid.IsValid(), TEXT("New editor NodeGuid must be valid."));
+							ensureMsgf(NewGraphNode->NodeGuid == Location.Key, TEXT("New editor NodeGuid must match location key"));
+
+							if (AssociatedNodeData.Node)
+							{
+								ensure(AssociatedNodeData.Node == Pair.Key);
+							}
+							else
+							{
+								AssociatedNodeData.Node = Pair.Key;
+								AssociatedNodeData.Class = Pair.Value;
+							}
+
+							AssociatedNodeData.EditorNodes.Add(NewGraphNode);
 						}
+						// If another editor node already exists for this frontend node, remove it 
+						// This is an editor only "fix" for frontend nodes with multiple locations, which should not be allowed past 
+						// transient ed graph migration (doc version 1.12)
 						else
 						{
-							AssociatedNodeData.Node = Pair.Key;
-							AssociatedNodeData.Class = Pair.Value;
+							constexpr bool bMarkDirty = false;
+							constexpr bool bBreakAllLinks = true;
+							OutGraph.RemoveNode(NewGraphNode, bBreakAllLinks, bMarkDirty);
 						}
-
-						AssociatedNodeData.EditorNodes.Add(NewGraphNode);
 					}
 				}
 			}
diff --git a/Engine/Plugins/Runtime/Metasound/Source/MetasoundEditor/Private/MetasoundEditorGraphNode.cpp b/Engine/Plugins/Runtime/Metasound/Source/MetasoundEditor/Private/MetasoundEditorGraphNode.cpp
--- a/Engine/Plugins/Runtime/Metasound/Source/MetasoundEditor/Private/MetasoundEditorGraphNode.cpp
+++ b/Engine/Plugins/Runtime/Metasound/Source/MetasoundEditor/Private/MetasoundEditorGraphNode.cpp
@@ -121,17 +121,18 @@ bool UMetasoundEditorGraphNode::SyncLocationFromFrontendNode(bool bUpdateEditorN
 		{
 			if (bUpdateEditorNodeID)
 			{
-				if (ensureMsgf(Node->Style.Display.Locations.Num() == 1, TEXT("Style location map should only contain one location (legacy support for multiple location values should be migrated by this point)")))
+				if (Node->Style.Display.Locations.Num() != 1)
 				{
-					for (const TPair<FGuid, FVector2D>& IDLocationPair : Node->Style.Display.Locations)
+					UE_LOG(LogMetaSound, Warning, TEXT("Style location map for node %s should only contain one location (legacy support for multiple location values should be migrated by this point)"), *Node->GetID().ToString());
+				}
+				for (const TPair<FGuid, FVector2D>& IDLocationPair : Node->Style.Display.Locations)
+				{
+					if (ensureMsgf(IDLocationPair.Key.IsValid(), TEXT("Location cannot be updated for ed node with invalid guid")))
 					{
-						if (ensureMsgf(IDLocationPair.Key.IsValid(), TEXT("Location cannot be updated for ed node with invalid guid")))
-						{
-							NodeGuid = IDLocationPair.Key;
-							NodePosX = UE::LWC::FloatToIntCastChecked<int32>(IDLocationPair.Value.X);
-							NodePosY = UE::LWC::FloatToIntCastChecked<int32>(IDLocationPair.Value.Y);
-							return true;
-						}
+						NodeGuid = IDLocationPair.Key;
+						NodePosX = UE::LWC::FloatToIntCastChecked<int32>(IDLocationPair.Value.X);
+						NodePosY = UE::LWC::FloatToIntCastChecked<int32>(IDLocationPair.Value.Y);
+						return true;
 					}
 				}
 			}
@@ -394,6 +395,14 @@ void UMetasoundEditorGraphNode::Validate(const FMetaSoundFrontendDocumentBuilder
 			}
 		}
 	}
+
+	// Validate node locations 
+	// Frontend node should only have one location post ed graph migration in doc version 1.12
+	const FMetasoundFrontendNode* Node = InBuilder.FindNode(GetNodeID());
+	if (Node->Style.Display.Locations.Num() > 1)
+	{
+		OutResult.SetMessage(EMessageSeverity::Info, FString::Format(TEXT("Frontend node with {0} locations found when there should only be 1. Moving this node and resaving the asset should fixup node locations."), { Node->Style.Display.Locations.Num() }));
+	}
 #endif // WITH_EDITOR
 }
 
@@ -1600,6 +1609,16 @@ void UMetasoundEditorGraphVariableNode::CacheBreadcrumb()
 			VertexMetadata.SetDescription(FrontendVariable->Description);
 
 			Breadcrumb.VertexMetadata = VertexMetadata;
+
+			FMetaSoundFrontendDocumentBuilder& Builder = Variable->GetFrontendBuilderChecked();
+			if (const FMetasoundFrontendNode* Node = Builder.FindNode(GetNodeID()))
+			{
+				if (const FMetasoundFrontendClass* Class = Builder.FindDependency(Node->ClassID))
+				{
+					const FMetasoundFrontendClassMetadata& Metadata = Class->Metadata;
+					Breadcrumb.ClassName = Metadata.GetClassName();
+				}
+			}
 		}
 	}
 }
diff --git a/Engine/Plugins/Runtime/Metasound/Source/MetasoundFrontend/Private/MetasoundFrontendDocumentBuilder.cpp b/Engine/Plugins/Runtime/Metasound/Source/MetasoundFrontend/Private/MetasoundFrontendDocumentBuilder.cpp
--- a/Engine/Plugins/Runtime/Metasound/Source/MetasoundFrontend/Private/MetasoundFrontendDocumentBuilder.cpp
+++ b/Engine/Plugins/Runtime/Metasound/Source/MetasoundFrontend/Private/MetasoundFrontendDocumentBuilder.cpp
@@ -5991,11 +5991,21 @@ bool FMetaSoundFrontendDocumentBuilder::SetNodeLocation(const FGuid& InNodeID, c
 		{
 			if (InLocationGuid->IsValid())
 			{
+				auto ExistingLocationWithOtherId = [&](const TPair<FGuid, FVector2D> Location) {
+					return Location.Key != *InLocationGuid;
+				};
+				// If there are multiple locations or a single location with a nonmatching location guid (which would result in multiple locations after this location is added)
+				if (Style.Display.Locations.Num() > 1 || Algo::AnyOf(Style.Display.Locations, ExistingLocationWithOtherId))
+				{
+					UE_LOG(LogMetaSound, Warning, TEXT("Multiple locations per node no longer supported, removing other display locations for node ID '%s'"), *InNodeID.ToString());
+					Style.Display.Locations.Empty();
+				}
+
 				Style.Display.Locations.FindOrAdd(*InLocationGuid) = InLocation;
 				return true;
 			}
 
-			UE_LOG(LogMetaSound, Display, TEXT("Invalid Location Guid no longer supported, reseting display location for node with ID '%s'"), *InNodeID.ToString());
+			UE_LOG(LogMetaSound, Display, TEXT("Invalid Location Guid no longer supported, resetting display location for node with ID '%s'"), *InNodeID.ToString());
 		}
 
 		if (Style.Display.Locations.IsEmpty())
@@ -6010,6 +6020,8 @@ bool FMetaSoundFrontendDocumentBuilder::SetNodeLocation(const FGuid& InNodeID, c
 			});
 		}
 
+		ensureMsgf(Style.Display.Locations.Num() <= 1, TEXT("Nodes should not have more than 1 location post ed graph migration (doc version 1.12)."));
+
 		return true;
 	}
 
