Fix for Upgrade note messages not disappearing
#rnx

diff --git a/Engine/Plugins/Runtime/Metasound/Source/MetasoundEditor/Private/MetasoundEditorGraph.cpp b/Engine/Plugins/Runtime/Metasound/Source/MetasoundEditor/Private/MetasoundEditorGraph.cpp
--- a/Engine/Plugins/Runtime/Metasound/Source/MetasoundEditor/Private/MetasoundEditorGraph.cpp
+++ b/Engine/Plugins/Runtime/Metasound/Source/MetasoundEditor/Private/MetasoundEditorGraph.cpp
@@ -2116,13 +2116,10 @@ void UMetasoundEditorGraph::PreSave(FObjectPreSaveContext InSaveContext)
 	GetNodesOfClass<UMetasoundEditorGraphNode>(MetaSoundNodes);
 	for (UMetasoundEditorGraphNode* Node : MetaSoundNodes)
 	{
-PRAGMA_DISABLE_DEPRECATION_WARNINGS
-		FNodeHandle NodeHandle = Node->GetNodeHandle();
-PRAGMA_ENABLE_DEPRECATION_WARNINGS
-
-		FMetasoundFrontendNodeStyle Style = NodeHandle->GetNodeStyle();
-		Style.bMessageNodeUpdated = false;
-		NodeHandle->SetNodeStyle(Style);
+		if (const FMetasoundFrontendNode* FrontendNode = Node->GetFrontendNode())
+		{
+			FrontendNode->Style.bMessageNodeUpdated = false;
+		}
 	}
 
 	Super::PreSave(InSaveContext);
diff --git a/Engine/Plugins/Runtime/Metasound/Source/MetasoundEditor/Private/MetasoundEditorGraphNode.cpp b/Engine/Plugins/Runtime/Metasound/Source/MetasoundEditor/Private/MetasoundEditorGraphNode.cpp
--- a/Engine/Plugins/Runtime/Metasound/Source/MetasoundEditor/Private/MetasoundEditorGraphNode.cpp
+++ b/Engine/Plugins/Runtime/Metasound/Source/MetasoundEditor/Private/MetasoundEditorGraphNode.cpp
@@ -1373,6 +1373,10 @@ void UMetasoundEditorGraphExternalNode::Validate(const FMetaSoundFrontendDocumen
 			FText::FromString(Metadata.GetVersion().ToString())
 		));
 	}
+	else
+	{
+		OutResult.SetUpgradeMessage({ });
+	}
 
 	// 6. Reset pin state (if pin was orphaned or clear if no longer orphaned)
 	for (UEdGraphPin* Pin : Pins)
diff --git a/Engine/Plugins/Runtime/Metasound/Source/MetasoundFrontend/Public/MetasoundFrontendDocument.h b/Engine/Plugins/Runtime/Metasound/Source/MetasoundFrontend/Public/MetasoundFrontendDocument.h
--- a/Engine/Plugins/Runtime/Metasound/Source/MetasoundFrontend/Public/MetasoundFrontendDocument.h
+++ b/Engine/Plugins/Runtime/Metasound/Source/MetasoundFrontend/Public/MetasoundFrontendDocument.h
@@ -657,8 +657,8 @@ struct FMetasoundFrontendNodeStyle
 
 	// Whether or not to display if
 	// the node's version has been updated
-	UPROPERTY()
-	bool bMessageNodeUpdated = false;
+	UPROPERTY(Transient)
+	mutable bool bMessageNodeUpdated = false;
 
 	UPROPERTY()
 	bool bIsPrivate = false;
