Clean-up of old deprecated calls in MetaSound plugin

#rnx

diff --git a/Engine/Plugins/Runtime/Metasound/Source/MetasoundEngine/Public/IMetasoundEngineModule.h b/Engine/Plugins/Runtime/Metasound/Source/MetasoundEngine/Public/IMetasoundEngineModule.h
--- a/Engine/Plugins/Runtime/Metasound/Source/MetasoundEngine/Public/IMetasoundEngineModule.h
+++ b/Engine/Plugins/Runtime/Metasound/Source/MetasoundEngine/Public/IMetasoundEngineModule.h
@@ -46,21 +46,9 @@ namespace Metasound::Engine
 		virtual void ShutdownModule() override = 0;
 
 #if WITH_EDITOR
-		UE_DEPRECATED(5.5, "Async prime no longer required (nor recommended as AssetTag updates now enable querying MetaSound assets for interface information without loading full asset.")
-		virtual void PrimeAssetRegistryAsync() { }
-
 		// Request enumeration of scanned assets and kick off add of MetaSound asset tag data to MetaSoundAssetManager
 		virtual void PrimeAssetManager() = 0;
 
-		UE_DEPRECATED(5.5, "Async prime no longer required (nor recommended as AssetTag updates now enable querying MetaSound assets for interface information without loading full asset.")
-		virtual ENodeClassRegistryPrimeStatus GetNodeClassRegistryPrimeStatus() { return ENodeClassRegistryPrimeStatus::NotRequested; }
-
-		UE_DEPRECATED(5.5, "Use FMetaSoundAssetManager::IsInitialAssetScanComplete instead")
-		virtual EAssetScanStatus GetAssetRegistryScanStatus() const { return EAssetScanStatus::NotRequested; }
-
-		UE_DEPRECATED(5.5, "Async prime no longer required (nor recommended as AssetTag updates now enable querying MetaSound assets for interface information without loading full asset.")
-		virtual bool IsAssetManagerPrimed() const { return false; }
-
 		// Asset registry delegates for calling MetaSound editor module register/unregister with frontend 
 		virtual FOnMetasoundGraphRegister& GetOnGraphRegisteredDelegate() = 0;
 		virtual FOnMetasoundGraphUnregister& GetOnGraphUnregisteredDelegate() = 0;
diff --git a/Engine/Plugins/Runtime/Metasound/Source/MetasoundFrontend/Public/MetasoundAssetManager.h b/Engine/Plugins/Runtime/Metasound/Source/MetasoundFrontend/Public/MetasoundAssetManager.h
--- a/Engine/Plugins/Runtime/Metasound/Source/MetasoundFrontend/Public/MetasoundAssetManager.h
+++ b/Engine/Plugins/Runtime/Metasound/Source/MetasoundFrontend/Public/MetasoundAssetManager.h
@@ -61,29 +61,29 @@ namespace Metasound::Frontend
 	namespace AssetTags
 	{
 		UE_DEPRECATED(5.6, "AssetTags no longer public. Construct FMetaSoundAssetClassInfo from AssetData to parse tag fields.")
-		extern const FString METASOUNDFRONTEND_API ArrayDelim;
+		extern const FString UE_API ArrayDelim;
 
 #if WITH_EDITORONLY_DATA
 		UE_DEPRECATED(5.6, "AssetTags no longer public. Construct FMetaSoundAssetClassInfo from AssetData to parse tag fields.")
-		extern const FName METASOUNDFRONTEND_API IsPreset;
+		extern const FName UE_API IsPreset;
 #endif // WITH_EDITORONLY_DATA
 
 
 		UE_DEPRECATED(5.6, "AssetTags no longer public. Construct FMetaSoundAssetClassInfo from AssetData to parse tag fields.")
-		extern const FName METASOUNDFRONTEND_API AssetClassID;
+		extern const FName UE_API AssetClassID;
 
 		UE_DEPRECATED(5.6, "AssetTags no longer public. Construct FMetaSoundAssetClassInfo from AssetData to parse tag fields.")
-		extern const FName METASOUNDFRONTEND_API RegistryVersionMajor;
+		extern const FName UE_API RegistryVersionMajor;
 
 		UE_DEPRECATED(5.6, "AssetTags no longer public. Construct FMetaSoundAssetClassInfo from AssetData to parse tag fields.")
-		extern const FName METASOUNDFRONTEND_API RegistryVersionMinor;
+		extern const FName UE_API RegistryVersionMinor;
 
 #if WITH_EDITORONLY_DATA
 		UE_DEPRECATED(5.6, "AssetTags no longer public. Construct FMetaSoundAssetClassInfo from AssetData to parse tag fields.")
-		extern const FName METASOUNDFRONTEND_API RegistryInputTypes;
+		extern const FName UE_API RegistryInputTypes;
 
 		UE_DEPRECATED(5.6, "AssetTags no longer public. Construct FMetaSoundAssetClassInfo from AssetData to parse tag fields.")
-		extern const FName METASOUNDFRONTEND_API RegistryOutputTypes;
+		extern const FName UE_API RegistryOutputTypes;
 #endif // WITH_EDITORONLY_DATA
 	} // namespace AssetTags
 
@@ -118,6 +118,9 @@ namespace Metasound::Frontend
 
 		// No longer used. Memory management of document (i.e. copying or using object's version) inferred internally
 		bool bRegisterCopyIfAsync = false;
+
+		// If currently previewing the given MetaSound, registration request is ignored.
+		bool bIgnoreIfLiveAuditioning = false;
 #endif // WITH_EDITOR
 
 		// Contains a ranking of the preferred page to use with the most preferred page ID being first. 
@@ -322,9 +325,6 @@ PRAGMA_ENABLE_DEPRECATION_WARNINGS
 		UE_DEPRECATED(5.6, "Moved to Iterate ClassInfo")
 		virtual void IterateAssets(TFunctionRef<void(const FMetaSoundAssetKey, const TArray<FTopLevelAssetPath>&)> Iter) const { }
 
-		UE_DEPRECATED(5.5, "Rescan no longer supported nor required by Frontend")
-		virtual void RescanAutoUpdateDenyList() { }
-
 		// Set flag for logging active assets on shutdown. In certain cases (ex. validation), it is expected that assets are active at shutdown
 		virtual void SetLogActiveAssetsOnShutdown(bool bLogActiveAssetsOnShutdown) = 0;
 
diff --git a/Engine/Plugins/Runtime/Metasound/Source/MetasoundFrontend/Public/MetasoundDocumentInterface.h b/Engine/Plugins/Runtime/Metasound/Source/MetasoundFrontend/Public/MetasoundDocumentInterface.h
--- a/Engine/Plugins/Runtime/Metasound/Source/MetasoundFrontend/Public/MetasoundDocumentInterface.h
+++ b/Engine/Plugins/Runtime/Metasound/Source/MetasoundFrontend/Public/MetasoundDocumentInterface.h
@@ -107,9 +107,6 @@ namespace Metasound::Frontend
 		virtual bool FinishBuilding(const FMetasoundFrontendClassName& InClassName, bool bForceUnregisterNodeClass = false) const = 0;
 		virtual bool FinishBuilding(const FMetasoundFrontendClassName& InClassName, const FTopLevelAssetPath& InAssetPath, bool bForceUnregisterNodeClass = false) const = 0;
 
-		UE_DEPRECATED(5.5, "Document cache can now be invalidated by retrieving an asset builder and calling 'ReloadBuilder'")
-		virtual void InvalidateDocumentCache(const FMetasoundFrontendClassName& InClassName) const { }
-
 		// Reloads the given builder, maintaining all modify delegate subscriptions. Returns true if builder was found and reloaded,
 		// false if not found.
 		virtual bool ReloadBuilder(const FMetasoundFrontendClassName& InClassName) const = 0;
