Fix for conditional no longer valid after removing MetaSound deprecated auto-update settings that caused unintended early out when auto-updating MetaSounds
#rnx

diff --git a/Engine/Plugins/Runtime/Metasound/Source/MetasoundFrontend/Private/MetasoundFrontendNodeController.cpp b/Engine/Plugins/Runtime/Metasound/Source/MetasoundFrontend/Private/MetasoundFrontendNodeController.cpp
--- a/Engine/Plugins/Runtime/Metasound/Source/MetasoundFrontend/Private/MetasoundFrontendNodeController.cpp
+++ b/Engine/Plugins/Runtime/Metasound/Source/MetasoundFrontend/Private/MetasoundFrontendNodeController.cpp
@@ -1232,10 +1232,6 @@ namespace Metasound
 
 			const FMetasoundFrontendClassMetadata& NodeClassMetadata = GetClassMetadata();
 			const IMetaSoundAssetManager& AssetManager = IMetaSoundAssetManager::GetChecked();
-			if (AssetManager.IsAssetClass(NodeClassMetadata))
-			{
-				return false;
-			}
 
 			FMetasoundFrontendClass RegistryClass;
 			if (!ISearchEngine::Get().FindClassWithHighestMinorVersion( NodeClassMetadata.GetClassName().ToNodeClassName(), NodeClassMetadata.GetVersion().Major, RegistryClass))
