Add nullptr checks in MetaSound editor graph validation

diff --git a/Engine/Plugins/Runtime/Metasound/Source/MetasoundEditor/Private/MetasoundEditorGraphNode.cpp b/Engine/Plugins/Runtime/Metasound/Source/MetasoundEditor/Private/MetasoundEditorGraphNode.cpp
--- a/Engine/Plugins/Runtime/Metasound/Source/MetasoundEditor/Private/MetasoundEditorGraphNode.cpp
+++ b/Engine/Plugins/Runtime/Metasound/Source/MetasoundEditor/Private/MetasoundEditorGraphNode.cpp
@@ -371,22 +371,25 @@ void UMetasoundEditorGraphNode::Validate(const FMetaSoundFrontendDocumentBuilder
 	{
 		for (UEdGraphPin* Pin : Pins)
 		{
-			OutResult.SetPinOrphaned(*Pin, false);
-			if (Pin->Direction == EGPD_Input)
+			if (Pin)
 			{
-				if (!Pin->LinkedTo.IsEmpty())
+				OutResult.SetPinOrphaned(*Pin, false);
+				if (Pin->Direction == EGPD_Input)
 				{
-					UEdGraphPin* ReroutedPin = FGraphBuilder::FindReroutedOutputPin(Pin->LinkedTo.Last());
-					if (ReroutedPin)
+					if (!Pin->LinkedTo.IsEmpty())
 					{
-						if (UMetasoundEditorGraphExternalNode* ExternalNode = Cast<UMetasoundEditorGraphExternalNode>(ReroutedPin->GetOwningNode()))
+						UEdGraphPin* ReroutedPin = FGraphBuilder::FindReroutedOutputPin(Pin->LinkedTo.Last());
+						if (ReroutedPin)
 						{
-							if (ExternalNode->GetBreadcrumb().ClassName == FRerouteNodeTemplate::ClassName)
+							if (UMetasoundEditorGraphExternalNode* ExternalNode = Cast<UMetasoundEditorGraphExternalNode>(ReroutedPin->GetOwningNode()))
 							{
-								if (const FMetasoundFrontendVertex* Vertex = FGraphBuilder::GetPinVertex(InBuilder, Pin))
+								if (ExternalNode->GetBreadcrumb().ClassName == FRerouteNodeTemplate::ClassName)
 								{
-									const FString Msg = FString::Format(*GraphNodePrivate::MissingConcreteOutputConnectionFormat, { Vertex->Name.ToString() });
-									OutResult.SetMessage(EMessageSeverity::Warning, Msg);
+									if (const FMetasoundFrontendVertex* Vertex = FGraphBuilder::GetPinVertex(InBuilder, Pin))
+									{
+										const FString Msg = FString::Format(*GraphNodePrivate::MissingConcreteOutputConnectionFormat, { Vertex->Name.ToString() });
+										OutResult.SetMessage(EMessageSeverity::Warning, Msg);
+									}
 								}
 							}
 						}
@@ -399,9 +402,16 @@ void UMetasoundEditorGraphNode::Validate(const FMetaSoundFrontendDocumentBuilder
 	// Validate node locations 
 	// Frontend node should only have one location post ed graph migration in doc version 1.12
 	const FMetasoundFrontendNode* Node = InBuilder.FindNode(GetNodeID());
-	if (Node->Style.Display.Locations.Num() > 1)
+	if (Node)
+	{
+		if (Node->Style.Display.Locations.Num() > 1)
+		{
+			OutResult.SetMessage(EMessageSeverity::Info, FString::Format(TEXT("Frontend node with {0} locations found when there should only be 1. Moving this node and resaving the asset should fixup node locations."), { Node->Style.Display.Locations.Num() }));
+		}
+	}
+	else
 	{
-		OutResult.SetMessage(EMessageSeverity::Info, FString::Format(TEXT("Frontend node with {0} locations found when there should only be 1. Moving this node and resaving the asset should fixup node locations."), { Node->Style.Display.Locations.Num() }));
+		OutResult.SetMessage(EMessageSeverity::Error, TEXT("No frontend node associated with this editor graph node. Please delete and readd it."));
 	}
 #endif // WITH_EDITOR
 }
