[Metasound] -Static Analysis Fix for undefined behavior

#rnx

diff --git a/Engine/Plugins/Runtime/Metasound/Source/MetasoundStandardNodes/Private/MetasoundSubharmonizerNode.cpp b/Engine/Plugins/Runtime/Metasound/Source/MetasoundStandardNodes/Private/MetasoundSubharmonizerNode.cpp
--- a/Engine/Plugins/Runtime/Metasound/Source/MetasoundStandardNodes/Private/MetasoundSubharmonizerNode.cpp
+++ b/Engine/Plugins/Runtime/Metasound/Source/MetasoundStandardNodes/Private/MetasoundSubharmonizerNode.cpp
@@ -193,6 +193,9 @@ class FSubharmonizerOperator : public TExecutableOperator<FSubharmonizerOperator
 		// Process loop in Octaver Mode
 		if (bOctaverMode)
 		{
+			// Create mask for modulo op
+			const int32 SubharmonicBaseMask = SubharmonicBase - 1;
+
 			// Process creates a 1.f or -1.f output value based on either flipping bits or using the results of a modulus operator
 			for (int32 i = 0; i < NumFrames; ++i)
 			{
@@ -203,7 +206,9 @@ class FSubharmonizerOperator : public TExecutableOperator<FSubharmonizerOperator
 				if (bLastSign != bNewSign)
 				{
 					// Update sign state based on counter
-					CurrentSignState = ++CurrentSignState & (SubharmonicBase - 1);
+					++CurrentSignState;
+
+					CurrentSignState = CurrentSignState & SubharmonicBaseMask;
 				}
 
 				// Select correct output value based on sign and bitmask state
