MetasoundBuilderSubsystem - added FindParentBuilderOfPreset function

#rnx

diff --git a/Engine/Plugins/Runtime/Metasound/Source/MetasoundEngine/Private/MetasoundBuilderSubsystem.cpp b/Engine/Plugins/Runtime/Metasound/Source/MetasoundEngine/Private/MetasoundBuilderSubsystem.cpp
--- a/Engine/Plugins/Runtime/Metasound/Source/MetasoundEngine/Private/MetasoundBuilderSubsystem.cpp
+++ b/Engine/Plugins/Runtime/Metasound/Source/MetasoundEngine/Private/MetasoundBuilderSubsystem.cpp
@@ -1076,6 +1076,63 @@ UMetaSoundBuilderBase* UMetaSoundBuilderSubsystem::FindBuilderOfDocument(TScript
 	return FDocumentBuilderRegistry::GetChecked().FindBuilderObject(InMetaSound);
 }
 
+UMetaSoundBuilderBase* UMetaSoundBuilderSubsystem::FindParentBuilderOfPreset(const TScriptInterface<const IMetaSoundDocumentInterface> InMetaSoundPreset, const bool bFollowPresetChain)
+{
+	using namespace Metasound::Engine;
+	using namespace Metasound::Frontend;
+
+	// Find the builder for the given asset to check if it is a preset.
+	const FDocumentBuilderRegistry& Registry = FDocumentBuilderRegistry::GetChecked();
+	const UMetaSoundBuilderBase* Builder = Registry.FindBuilderObject(InMetaSoundPreset);
+	if (!Builder)
+	{
+		UE_LOG(LogMetaSound, Error, TEXT("Could not find builder for MetaSound '%s'."), *GetNameSafe(InMetaSoundPreset.GetObject()));
+		return nullptr;
+	}
+
+	if (!Builder->IsPreset())
+	{
+		UE_LOG(LogMetaSound, Error, TEXT("MetaSound '%s' is not a preset."), *GetNameSafe(InMetaSoundPreset.GetObject()));
+		return nullptr;
+	}
+
+	TScriptInterface<IMetaSoundDocumentInterface> ParentAsset = Builder->GetReferencedPresetAsset();
+	if (!ParentAsset)
+	{
+		UE_LOG(LogMetaSound, Error, TEXT("Could not find referenced preset asset for MetaSound '%s'."), *GetNameSafe(InMetaSoundPreset.GetObject()));
+		return nullptr;
+	}
+
+	UMetaSoundBuilderBase* ParentBuilder = Registry.FindBuilderObject(ParentAsset);
+	if (!bFollowPresetChain)
+	{
+		return ParentBuilder;
+	}
+
+	// Follow preset chain to the highest non-preset ancestor
+	TSet<const UObject*> CheckedAssets;
+	while (ParentBuilder && ParentBuilder->IsPreset())
+	{
+		// Check for circular references
+		if (CheckedAssets.Contains(ParentAsset.GetObject()))
+		{
+			return nullptr;
+		}
+
+		CheckedAssets.Add(ParentAsset.GetObject());
+
+		ParentAsset = ParentBuilder->GetReferencedPresetAsset();
+		if (!ParentAsset)
+		{
+			return nullptr;
+		}
+
+		ParentBuilder = Registry.FindBuilderObject(ParentAsset.GetObject());
+	}
+
+	return ParentBuilder;
+}
+
 UMetaSoundPatchBuilder* UMetaSoundBuilderSubsystem::FindPatchBuilder(FName BuilderName)
 {
 	if (UMetaSoundBuilderBase* Builder = FindBuilder(BuilderName))
diff --git a/Engine/Plugins/Runtime/Metasound/Source/MetasoundEngine/Public/MetasoundBuilderSubsystem.h b/Engine/Plugins/Runtime/Metasound/Source/MetasoundEngine/Public/MetasoundBuilderSubsystem.h
--- a/Engine/Plugins/Runtime/Metasound/Source/MetasoundEngine/Public/MetasoundBuilderSubsystem.h
+++ b/Engine/Plugins/Runtime/Metasound/Source/MetasoundEngine/Public/MetasoundBuilderSubsystem.h
@@ -244,6 +244,13 @@ public:
 	UFUNCTION(BlueprintCallable, Category = "Audio|MetaSound|Builder", meta = (DisplayName = "Find Builder By MetaSound"))
 	UE_API UPARAM(DisplayName = "Builder") UMetaSoundBuilderBase* FindBuilderOfDocument(TScriptInterface<const IMetaSoundDocumentInterface> InMetaSound) const;
 
+	/**
+	 * Returns the builder associated with the given MetaSound preset's parent. If bFollowPresetChain is true, continues traversing up the preset chain to return the highest non-preset ancestor.
+	 * Returns nullptr if the given MetaSound is not a preset, or if the parent builder is not yet registered.
+	 */
+	UFUNCTION(BlueprintCallable, Category = "Audio|MetaSound|Builder", meta = (DisplayName = "Find Parent Builder of MetaSound Preset"))
+	UE_API UPARAM(DisplayName = "Builder") UMetaSoundBuilderBase* FindParentBuilderOfPreset(TScriptInterface<const IMetaSoundDocumentInterface> InMetaSoundPreset, bool bFollowPresetChain = true);
+
 	// Returns the patch builder manually registered with the MetaSound Builder Subsystem with the provided custom name (if previously registered)
 	UFUNCTION(BlueprintCallable, Category = "Audio|MetaSound|Builder")
 	UE_API UPARAM(DisplayName = "Patch Builder") UMetaSoundPatchBuilder* FindPatchBuilder(FName BuilderName);
