[MetaSound] - Fix for Subharmonizer mismatched Reset function signature

#rnx

diff --git a/Engine/Plugins/Runtime/Metasound/Source/MetasoundStandardNodes/Private/MetasoundSubharmonizerNode.cpp b/Engine/Plugins/Runtime/Metasound/Source/MetasoundStandardNodes/Private/MetasoundSubharmonizerNode.cpp
--- a/Engine/Plugins/Runtime/Metasound/Source/MetasoundStandardNodes/Private/MetasoundSubharmonizerNode.cpp
+++ b/Engine/Plugins/Runtime/Metasound/Source/MetasoundStandardNodes/Private/MetasoundSubharmonizerNode.cpp
@@ -94,11 +94,14 @@ class FSubharmonizerOperator : public TExecutableOperator<FSubharmonizerOperator
 
 		void UpdateParams();
 
-		void Reset(const FOperatorSettings& InSettings);
+		void Reset(const IOperator::FResetParams& InParams);
 		
 		void Execute();
 
 	private:
+		// Internal reset function
+		void InternalReset();
+
 		// Input Pin Refs
 		FAudioBufferReadRef	AudioIn;
 		FFloatReadRef FrequencyIn;
@@ -351,7 +354,10 @@ class FSubharmonizerOperator : public TExecutableOperator<FSubharmonizerOperator
 			, OctaverModeIn(InOctaverModeIn)
 			, AudioOut(TDataWriteReferenceFactory<FAudioBuffer>::CreateExplicitArgs(InSettings))
 	{
-		Reset(InSettings);
+		NumFramesPerBlock = InSettings.GetNumFramesPerBlock();
+		SampleRate = InSettings.GetSampleRate();
+
+		InternalReset();
 	}
 
 	void FSubharmonizerOperator::BindInputs(FInputVertexInterfaceData& InOutVertexData)
@@ -419,12 +425,17 @@ class FSubharmonizerOperator : public TExecutableOperator<FSubharmonizerOperator
 		}
 	}
 
-	void FSubharmonizerOperator::Reset(const FOperatorSettings& InSettings)
+	void FSubharmonizerOperator::Reset(const IOperator::FResetParams& InParams)
 	{
 		// Cache environment
-		NumFramesPerBlock = InSettings.GetNumFramesPerBlock();
-		SampleRate = InSettings.GetSampleRate();
+		NumFramesPerBlock = InParams.OperatorSettings.GetNumFramesPerBlock();
+		SampleRate = InParams.OperatorSettings.GetSampleRate();
+
+		InternalReset();
+	}
 
+	void FSubharmonizerOperator::InternalReset()
+	{
 		// Max filter frequency/nyquist frequency
 		MaxFilterFrequency = SampleRate * 0.5;
 
