Checkpoint: Add audition enablement cvar for experimenting with live audition in MetaSound Editor. Adding to iterate on testing node config swap easier, but has longer term ramifications for when we can go wide with auditioning live in the editor more broadly.

#rnx

diff --git a/Engine/Plugins/Runtime/Metasound/Source/MetasoundEditor/Private/MetasoundDetailCustomization.cpp b/Engine/Plugins/Runtime/Metasound/Source/MetasoundEditor/Private/MetasoundDetailCustomization.cpp
--- a/Engine/Plugins/Runtime/Metasound/Source/MetasoundEditor/Private/MetasoundDetailCustomization.cpp
+++ b/Engine/Plugins/Runtime/Metasound/Source/MetasoundEditor/Private/MetasoundDetailCustomization.cpp
@@ -244,8 +244,12 @@ namespace Metasound::Editor
 							{
 								DocBuilder.GetConstDocumentChecked().RootGraph.Style.UpdateChangeID();
 							}
-							constexpr bool bForceViewSynchronization = true;
-							FGraphBuilder::RegisterGraphWithFrontend(DocBuilder.CastDocumentObjectChecked<UObject>(), bForceViewSynchronization);
+
+							{
+								FMetaSoundAssetRegistrationOptions RegOptions = FGraphBuilder::GetDefaultRegistrationOptions();
+								RegOptions.bForceViewSynchronization = true;
+								FGraphBuilder::RegisterGraphWithFrontend(DocBuilder.CastDocumentObjectChecked<UObject>(), MoveTemp(RegOptions));
+							}
 						}
 					});
 					AuthorHandle->SetOnChildPropertyValueChanged(RegisterOnChange);
diff --git a/Engine/Plugins/Runtime/Metasound/Source/MetasoundEditor/Private/MetasoundEditor.cpp b/Engine/Plugins/Runtime/Metasound/Source/MetasoundEditor/Private/MetasoundEditor.cpp
--- a/Engine/Plugins/Runtime/Metasound/Source/MetasoundEditor/Private/MetasoundEditor.cpp
+++ b/Engine/Plugins/Runtime/Metasound/Source/MetasoundEditor/Private/MetasoundEditor.cpp
@@ -111,24 +111,16 @@ namespace Metasound
 		namespace AssetEditorPrivate
 		{
 			static const FText LoadingAssetsFormat = LOCTEXT("MetaSoundScanInProgressNotificationText", "Loading {0} MetaSound Assets...");
-		} // namespace AssetEditorPrivate
+			static bool bEnableLiveAudition = false;
+			FAutoConsoleVariableRef CVarMetaSoundEditorEnableLiveAudition(
+				TEXT("au.MetaSound.Editor.Experimental.EnableLiveAudition"),
+				bEnableLiveAudition,
+				TEXT("If true, enable live audition in MetaSound Editor. This allows immediate audible feedback from all builder actions supported"
+					"by the MetaSound Dynamic Transactor (ex. connecting, disconnecting, creating and removing nodes. Experimental).\n")
+				TEXT("Default: false"),
+				ECVF_Default);
 
-		bool IsPreviewingMetaSound(const UObject& InMetaSound)
-		{
-			using namespace Engine;
-			if (const UAudioComponent* PreviewComponent = GEditor->GetPreviewAudioComponent())
-			{
-				if (PreviewComponent->IsPlaying())
-				{
-					if (const USoundBase* Sound = PreviewComponent->Sound)
-					{
-						return Sound->GetUniqueID() == InMetaSound.GetUniqueID();
-					}
-				}
-			}
-
-			return false;
-		}
+		} // namespace AssetEditorPrivate
 
 		bool IsPreviewingPageInputDefault(const FMetaSoundFrontendDocumentBuilder& Builder, const FMetasoundFrontendClassInput& InClassInput, const FGuid& InPageID)
 		{
@@ -136,7 +128,7 @@ namespace Metasound
 			using namespace Frontend;
 
 			const UObject& MetaSound = Builder.CastDocumentObjectChecked<const UObject>();
-			if (IsPreviewingMetaSound(MetaSound))
+			if (FGraphBuilder::IsPreviewingMetaSound(MetaSound))
 			{
 				if (const FMetasoundFrontendClassInputDefault* InputDefault = FindPreferredPage(InClassInput.GetDefaults(), UMetaSoundSettings::GetPageOrder()))
 				{
@@ -153,7 +145,7 @@ namespace Metasound
 			using namespace Frontend;
 
 			const UObject& MetaSound = Builder.CastDocumentObjectChecked<const UObject>();
-			if (IsPreviewingMetaSound(MetaSound))
+			if (FGraphBuilder::IsPreviewingMetaSound(MetaSound))
 			{
 				const FMetasoundFrontendGraphClass& GraphClass = Builder.GetConstDocumentChecked().RootGraph;
 				if (const FMetasoundFrontendGraph* GraphPage = FindPreferredPage(GraphClass.GetConstGraphPages(), UMetaSoundSettings::GetPageOrder()))
@@ -1083,6 +1075,11 @@ namespace Metasound
 			InTabManager->UnregisterTabSpawner(TabNamesPrivate::UserPresetWidget);
 		}
 
+		bool FEditor::IsLiveAuditionEnabled()
+		{
+			return AssetEditorPrivate::bEnableLiveAudition;
+		}
+
 		bool FEditor::IsPlaying() const
 		{
 			if (UObject* MetaSound = GetMetasoundObject())
@@ -2555,16 +2552,20 @@ namespace Metasound
 				// re-registering to make sure everything is up-to-date.
 				FGraphBuilder::RegisterGraphWithFrontend(*MetaSoundToPlay);
 
-				// Set the send to the audio bus that is used for analyzing the metasound output
-				check(GEditor);
-
 				UpdateRenderInfo(true /* bIsPlaying */);
 				UpdatePageInfo(true);
 
 				if (UMetaSoundSource* Source = Cast<UMetaSoundSource>(GetMetasoundObject()))
 				{
-					if (UAudioComponent* PreviewComp = GEditor->PlayPreviewSound(Source))
+					check(GEditor);
+					if (UAudioComponent* PreviewComp = GEditor->ResetPreviewAudioComponent())
 					{
+						if (Builder.IsValid())
+						{
+							FOnCreateAuditionGeneratorHandleDelegate GeneratorHandleDelegate;
+							CastChecked<UMetaSoundSourceBuilder>(Builder.Get())->Audition(nullptr, PreviewComp, GeneratorHandleDelegate, AssetEditorPrivate::bEnableLiveAudition);
+						}
+
 						SetPreviewID(PreviewComp->GetUniqueID());
 
 						if (UAudioBus* AudioBus = AnalyzerRack->GetAudioBus())
@@ -3762,7 +3763,11 @@ namespace Metasound
 				}
 			}
 
-			FGraphBuilder::RegisterGraphWithFrontend(*MetaSound, true);
+			{
+				Frontend::FMetaSoundAssetRegistrationOptions RegOptions = FGraphBuilder::GetDefaultRegistrationOptions();
+				RegOptions.bForceViewSynchronization = true;
+				FGraphBuilder::RegisterGraphWithFrontend(*MetaSound, MoveTemp(RegOptions));
+			}
 
 			if (GraphMembersMenu.IsValid())
 			{
@@ -3952,7 +3957,11 @@ namespace Metasound
 					checkNoEntry();
 				}
 
-				FGraphBuilder::RegisterGraphWithFrontend(*MetaSound, true);
+				{
+					Frontend::FMetaSoundAssetRegistrationOptions RegOptions = FGraphBuilder::GetDefaultRegistrationOptions();
+					RegOptions.bForceViewSynchronization = true;
+					FGraphBuilder::RegisterGraphWithFrontend(*MetaSound, MoveTemp(RegOptions));
+				}
 			}
 		}
 
@@ -4952,7 +4961,11 @@ PRAGMA_ENABLE_DEPRECATION_WARNINGS
 				return FReply::Unhandled();
 			}
 
-			FGraphBuilder::RegisterGraphWithFrontend(*MetaSound, true);
+			{
+				Frontend::FMetaSoundAssetRegistrationOptions RegOptions = FGraphBuilder::GetDefaultRegistrationOptions();
+				RegOptions.bForceViewSynchronization = true;
+				FGraphBuilder::RegisterGraphWithFrontend(*MetaSound, MoveTemp(RegOptions));
+			}
 
 			if (GraphMembersMenu.IsValid())
 			{
@@ -5411,7 +5424,11 @@ PRAGMA_ENABLE_DEPRECATION_WARNINGS
 				}
 			}
 
-			FGraphBuilder::RegisterGraphWithFrontend(ParentMetasound, true);
+			{
+				Frontend::FMetaSoundAssetRegistrationOptions RegOptions = FGraphBuilder::GetDefaultRegistrationOptions();
+				RegOptions.bForceViewSynchronization = true;
+				FGraphBuilder::RegisterGraphWithFrontend(ParentMetasound, MoveTemp(RegOptions));
+			}
 		}
 
 		bool FEditor::CanPromoteAllToCommonInputs()
@@ -5521,7 +5538,11 @@ PRAGMA_ENABLE_DEPRECATION_WARNINGS
 				}
 			}
 
-			FGraphBuilder::RegisterGraphWithFrontend(ParentMetasound, true);
+			{
+				Frontend::FMetaSoundAssetRegistrationOptions RegOptions = FGraphBuilder::GetDefaultRegistrationOptions();
+				RegOptions.bForceViewSynchronization = true;
+				FGraphBuilder::RegisterGraphWithFrontend(ParentMetasound, MoveTemp(RegOptions));
+			}
 		}
 	}
 }
diff --git a/Engine/Plugins/Runtime/Metasound/Source/MetasoundEditor/Private/MetasoundEditor.h b/Engine/Plugins/Runtime/Metasound/Source/MetasoundEditor/Private/MetasoundEditor.h
--- a/Engine/Plugins/Runtime/Metasound/Source/MetasoundEditor/Private/MetasoundEditor.h
+++ b/Engine/Plugins/Runtime/Metasound/Source/MetasoundEditor/Private/MetasoundEditor.h
@@ -98,7 +98,6 @@ class UMetasoundPagesView : public UMetasoundEditorViewBase
 
 namespace Metasound::Editor
 {
-	bool IsPreviewingMetaSound(const UObject& InMetaSound);
 	bool IsPreviewingPageInputDefault(const FMetaSoundFrontendDocumentBuilder& Builder, const FMetasoundFrontendClassInput& InClassInput, const FGuid& InPageID);
 	bool IsPreviewingPageGraph(const FMetaSoundFrontendDocumentBuilder& Builder, const FGuid& InPageID);
 	bool PageEditorEnabled(const FMetaSoundFrontendDocumentBuilder& Builder, bool bHasProjectPageValues, bool bPresetCanEditPageValues = false);
@@ -171,6 +170,9 @@ namespace Metasound::Editor
 		virtual void Play() override;
 		virtual void Stop() override;
 
+		/** Returns whether or not live audition is enabled for previewing */
+		static bool IsLiveAuditionEnabled();
+
 		virtual bool IsPlaying() const override;
 
 		/** Whether pasting the current data on the clipboard to the focused graph is permissible */
diff --git a/Engine/Plugins/Runtime/Metasound/Source/MetasoundEditor/Private/MetasoundEditorGraph.cpp b/Engine/Plugins/Runtime/Metasound/Source/MetasoundEditor/Private/MetasoundEditorGraph.cpp
--- a/Engine/Plugins/Runtime/Metasound/Source/MetasoundEditor/Private/MetasoundEditorGraph.cpp
+++ b/Engine/Plugins/Runtime/Metasound/Source/MetasoundEditor/Private/MetasoundEditorGraph.cpp
@@ -8,6 +8,7 @@
 #include "EdGraph/EdGraphNode.h"
 #include "Interfaces/ITargetPlatform.h"
 #include "MetasoundAssetBase.h"
+#include "MetasoundAssetManager.h"
 #include "MetasoundDocumentBuilderRegistry.h"
 #include "MetasoundEditorGraphCommentNode.h"
 #include "MetasoundEditorGraphInputNode.h"
@@ -70,7 +71,12 @@ namespace Metasound::Editor
 			{
 				if (UObject* MetaSound = InMember.GetOutermostObject())
 				{
-					FGraphBuilder::RegisterGraphWithFrontend(*MetaSound);
+					{
+						Frontend::FMetaSoundAssetRegistrationOptions RegOptions = FGraphBuilder::GetDefaultRegistrationOptions();
+						RegOptions.bIgnoreIfLiveAuditioning = true;
+						FGraphBuilder::RegisterGraphWithFrontend(*MetaSound, MoveTemp(RegOptions));
+					}
+
 					if (FMetasoundAssetBase* MetaSoundAsset = IMetasoundUObjectRegistry::Get().GetObjectAsAssetBase(MetaSound))
 					{
 						MetaSoundAsset->GetModifyContext().AddMemberIDsModified({ InMember.GetMemberID() });
diff --git a/Engine/Plugins/Runtime/Metasound/Source/MetasoundEditor/Private/MetasoundEditorGraphBuilder.cpp b/Engine/Plugins/Runtime/Metasound/Source/MetasoundEditor/Private/MetasoundEditorGraphBuilder.cpp
--- a/Engine/Plugins/Runtime/Metasound/Source/MetasoundEditor/Private/MetasoundEditorGraphBuilder.cpp
+++ b/Engine/Plugins/Runtime/Metasound/Source/MetasoundEditor/Private/MetasoundEditorGraphBuilder.cpp
@@ -1001,6 +1001,23 @@ PRAGMA_ENABLE_DEPRECATION_WARNINGS
 			return IOutputController::GetInvalidHandle();
 		}
 
+		bool FGraphBuilder::IsPreviewingMetaSound(const UObject& InMetaSound)
+		{
+			using namespace Engine;
+			if (const UAudioComponent* PreviewComponent = GEditor->GetPreviewAudioComponent())
+			{
+				if (PreviewComponent->IsPlaying())
+				{
+					if (const USoundBase* Sound = PreviewComponent->Sound)
+					{
+						return Sound->GetUniqueID() == InMetaSound.GetUniqueID();
+					}
+				}
+			}
+
+			return false;
+		}
+
 		UEdGraphPin* FGraphBuilder::FindReroutedOutputPin(UEdGraphPin* OutputPin)
 		{
 			using namespace Frontend;
@@ -1055,35 +1072,6 @@ PRAGMA_ENABLE_DEPRECATION_WARNINGS
 			return OutputPin;
 		}
 
-		Frontend::FOutputHandle FGraphBuilder::FindReroutedOutputHandleFromPin(const UEdGraphPin* OutputPin)
-		{
-			using namespace Frontend;
-
-			if (OutputPin)
-			{
-				if (UMetasoundEditorGraphExternalNode* ExternalNode = Cast<UMetasoundEditorGraphExternalNode>(OutputPin->GetOwningNode()))
-				{
-					if (ExternalNode->GetBreadcrumb().ClassName == FRerouteNodeTemplate::ClassName)
-					{
-						auto IsInput = [](UEdGraphPin* Pin) { check(Pin); return Pin->Direction == EGPD_Input; };
-						if (UEdGraphPin* RerouteInput = *ExternalNode->Pins.FindByPredicate(IsInput))
-						{
-							TArray<UEdGraphPin*>& LinkedTo = RerouteInput->LinkedTo;
-							if (!LinkedTo.IsEmpty())
-							{
-								UEdGraphPin* ReroutedOutput = LinkedTo.Last();
-								return FindReroutedOutputHandleFromPin(ReroutedOutput);
-							}
-						}
-					}
-				}
-
-				return GetOutputHandleFromPin(OutputPin);
-			}
-
-			return IOutputController::GetInvalidHandle();
-		}
-
 		Frontend::FConstOutputHandle FGraphBuilder::FindReroutedConstOutputHandleFromPin(const UEdGraphPin* OutputPin)
 		{
 			using namespace Frontend;
@@ -1530,7 +1518,7 @@ PRAGMA_ENABLE_DEPRECATION_WARNINGS
 			}
 		}
 
-		void FGraphBuilder::RegisterGraphWithFrontend(UObject& InMetaSound, bool bInForceViewSynchronization)
+		void FGraphBuilder::RegisterGraphWithFrontend(UObject& InMetaSound, Frontend::FMetaSoundAssetRegistrationOptions RegOptions)
 		{
 			using namespace Frontend;
 
@@ -1560,24 +1548,46 @@ PRAGMA_ENABLE_DEPRECATION_WARNINGS
 				}
 			}
 
-			FMetaSoundAssetRegistrationOptions RegOptions;
-			RegOptions.bForceReregister = true;
-			RegOptions.bForceViewSynchronization = bInForceViewSynchronization;
-			RegOptions.PageOrder = UMetaSoundSettings::GetPageOrder();
-			
+			auto CheckShouldIgnore = [&](const FMetasoundAssetBase& InAsset)
+			{
+				if (FEditor::IsLiveAuditionEnabled())
+				{
+					if (RegOptions.bIgnoreIfLiveAuditioning)
+					{
+						const UObject* AssetToTest = InAsset.GetOwningAsset();
+						check(AssetToTest);
+						const bool bIsPreviewing = FGraphBuilder::IsPreviewingMetaSound(*AssetToTest);
+						if (bIsPreviewing)
+						{
+							return true;
+						}
+
+						return false;
+					}
+				}
+
+				return false;
+			};
+
 			// if EditedReferencingMetaSounds is empty, then no MetaSounds are open
 			// that reference this MetaSound, so just register this asset. Otherwise,
 			// this graph will recursively get updated when the open referencing graphs
 			// are registered recursively via bRegisterDependencies flag.
 			if (EditedReferencingMetaSounds.IsEmpty())
 			{
-				MetaSoundAsset->UpdateAndRegisterForExecution(RegOptions);
+				if (!CheckShouldIgnore(*MetaSoundAsset))
+				{
+					MetaSoundAsset->UpdateAndRegisterForExecution(RegOptions);
+				}
 			}
 			else
 			{
-				for (FMetasoundAssetBase* MetaSound : EditedReferencingMetaSounds)
+				for (FMetasoundAssetBase* RefMetaSound : EditedReferencingMetaSounds)
 				{
-					MetaSound->UpdateAndRegisterForExecution(RegOptions);
+					if (!CheckShouldIgnore(*RefMetaSound))
+					{
+						RefMetaSound->UpdateAndRegisterForExecution(RegOptions);
+					}
 				}
 			}
 		}
diff --git a/Engine/Plugins/Runtime/Metasound/Source/MetasoundEditor/Private/MetasoundEditorGraphBuilder.h b/Engine/Plugins/Runtime/Metasound/Source/MetasoundEditor/Private/MetasoundEditorGraphBuilder.h
--- a/Engine/Plugins/Runtime/Metasound/Source/MetasoundEditor/Private/MetasoundEditorGraphBuilder.h
+++ b/Engine/Plugins/Runtime/Metasound/Source/MetasoundEditor/Private/MetasoundEditorGraphBuilder.h
@@ -3,12 +3,14 @@
 
 #include "Internationalization/Text.h"
 #include "Logging/TokenizedMessage.h"
+#include "MetasoundAssetManager.h"
 #include "MetasoundEditorGraphValidation.h"
 #include "MetasoundFrontend.h"
 #include "MetasoundFrontendController.h"
 #include "MetasoundFrontendDocumentBuilder.h"
 #include "MetasoundFrontendDocument.h"
 #include "MetasoundFrontendNodeTemplateRegistry.h"
+#include "MetasoundSettings.h"
 #include "Templates/SharedPointer.h"
 #include "UObject/NameTypes.h"
 #include "UObject/NoExportTypes.h"
@@ -75,11 +77,19 @@ namespace Metasound
 			// EdGraph synchronization, or if any of its referenced asset graphs have been marked as modified.
 			static bool RecurseGetDocumentModified(FMetasoundAssetBase& InAssetBase);
 
+			// Returns default registration options when registering from this graph builder API.
+			static Frontend::FMetaSoundAssetRegistrationOptions GetDefaultRegistrationOptions()
+			{
+				Frontend::FMetaSoundAssetRegistrationOptions RegOptions;
+				RegOptions.bForceViewSynchronization = false;
+				RegOptions.PageOrder = UMetaSoundSettings::GetPageOrder();
+				return RegOptions;
+			}
+
 			// Wraps RegisterGraphWithFrontend logic in Frontend with any additional logic required to refresh editor & respective editor object state.
 			// @param InMetaSound - MetaSound to register
-			// @param bInForceSynchronize - Forces the synchronize flag for all open graphs being registered by this call (all referenced graphs and
-			// referencing graphs open in editors)
-			static void RegisterGraphWithFrontend(UObject& InMetaSound, bool bInForceViewSynchronization = false);
+			// @param RegOptions - Registration options utilized
+			static void RegisterGraphWithFrontend(UObject& InMetaSound, Frontend::FMetaSoundAssetRegistrationOptions RegOptions = GetDefaultRegistrationOptions());
 
 			// Returns whether pin category is a custom MetaSound DataType
 			static bool IsPinCategoryMetaSoundCustomDataType(FName InPinCategoryName);
@@ -109,12 +119,6 @@ namespace Metasound
 			// only DisplayName, then the VariableName is used. 
 			static FText GetDisplayName(const Frontend::IVariableController& InFrontendVariable, bool bInIncludeNamespace = false);
 
-			// Returns the PinName for an IOutputController.
-			static FName GetPinName(const Frontend::IOutputController& InFrontendOutput);
-
-			// Returns the PinName for an IInputController.
-			static FName GetPinName(const Frontend::IInputController& InFrontendInput);
-
 			// Adds a new EdGraph comment node associated with the given MetaSoundFrontendGraph comment ID
 			static UMetasoundEditorGraphCommentNode* CreateCommentNode(UObject& InMetaSound, bool bInSelectNewNode = true, FGuid InCommentID = FGuid::NewGuid());
 
@@ -190,12 +194,13 @@ namespace Metasound
 			static Frontend::FOutputHandle GetOutputHandleFromPin(const UEdGraphPin* InPin);
 			static Frontend::FConstOutputHandle GetConstOutputHandleFromPin(const UEdGraphPin* InPin);
 
+			static bool IsPreviewingMetaSound(const UObject& InObject);
+
 			static UEdGraphPin* FindReroutedOutputPin(UEdGraphPin* InPin);
 			static const UEdGraphPin* FindReroutedOutputPin(const UEdGraphPin* InPin);
 
 			// Find the "concrete" output handle associated with an output pin.  If the given output pin is on
 			// a reroute node, will recursively search for the non-rerouted output its representing.
-			static Frontend::FOutputHandle FindReroutedOutputHandleFromPin(const UEdGraphPin* InOutputPin);
 			static Frontend::FConstOutputHandle FindReroutedConstOutputHandleFromPin(const UEdGraphPin* InPin);
 
 			// Find the "concrete" input handle associated with an output pin.  If the given input pin is on
diff --git a/Engine/Plugins/Runtime/Metasound/Source/MetasoundEditor/Private/MetasoundEditorGraphNode.cpp b/Engine/Plugins/Runtime/Metasound/Source/MetasoundEditor/Private/MetasoundEditorGraphNode.cpp
--- a/Engine/Plugins/Runtime/Metasound/Source/MetasoundEditor/Private/MetasoundEditorGraphNode.cpp
+++ b/Engine/Plugins/Runtime/Metasound/Source/MetasoundEditor/Private/MetasoundEditorGraphNode.cpp
@@ -562,16 +562,24 @@ void UMetasoundEditorGraphNode::PinDefaultValueChanged(UEdGraphPin* Pin)
 
 	if (Pin && Pin->Direction == EGPD_Input)
 	{
-		GetMetasoundChecked().Modify();
-
-		FInputHandle InputHandle = FGraphBuilder::GetInputHandleFromPin(Pin);
-		if (InputHandle->IsValid())
+		FMetasoundFrontendLiteral LiteralValue;
+		if (FGraphBuilder::GetPinLiteral(*Pin, LiteralValue))
 		{
-			FMetasoundFrontendLiteral LiteralValue;
-			if (FGraphBuilder::GetPinLiteral(*Pin, LiteralValue))
+			FMetaSoundFrontendDocumentBuilder& DocBuilder = GetBuilderChecked().GetBuilder();
+			FMetasoundFrontendVertexHandle InputVertexHandle = FGraphBuilder::GetPinVertexHandle(DocBuilder, Pin);
+			if (!ensure(InputVertexHandle.IsSet()))
 			{
-				InputHandle->SetLiteral(LiteralValue);
+				return;
 			}
+
+			const FMetasoundFrontendVertex* InputVertex = DocBuilder.FindNodeInput(InputVertexHandle.NodeID, InputVertexHandle.VertexID);
+			if (!ensure(InputVertex))
+			{
+				return;
+			}
+
+			GetMetasoundChecked().Modify();
+			DocBuilder.SetNodeInputDefault(InputVertexHandle.NodeID, InputVertexHandle.VertexID, LiteralValue);
 		}
 	}
 }
diff --git a/Engine/Plugins/Runtime/Metasound/Source/MetasoundEditor/Private/MetasoundEditorGraphSchema.cpp b/Engine/Plugins/Runtime/Metasound/Source/MetasoundEditor/Private/MetasoundEditorGraphSchema.cpp
--- a/Engine/Plugins/Runtime/Metasound/Source/MetasoundEditor/Private/MetasoundEditorGraphSchema.cpp
+++ b/Engine/Plugins/Runtime/Metasound/Source/MetasoundEditor/Private/MetasoundEditorGraphSchema.cpp
@@ -773,7 +773,9 @@ PRAGMA_ENABLE_DEPRECATION_WARNINGS
 
 							if (ensure(SchemaPrivate::TryConnectNewNodeToMatchingDataTypePin(*EdGraphNode, FromPin)))
 							{
-								FGraphBuilder::RegisterGraphWithFrontend(ParentMetasound);
+								Frontend::FMetaSoundAssetRegistrationOptions RegOptions = FGraphBuilder::GetDefaultRegistrationOptions();
+								RegOptions.bIgnoreIfLiveAuditioning = true;
+								FGraphBuilder::RegisterGraphWithFrontend(ParentMetasound, MoveTemp(RegOptions));
 								if (bSelectNewNode)
 								{
 									SchemaPrivate::SelectNodeInEditorForRename(*MetasoundGraph, *NewGraphNode);
diff --git a/Engine/Plugins/Runtime/Metasound/Source/MetasoundEditor/Private/MetasoundEditorModule.cpp b/Engine/Plugins/Runtime/Metasound/Source/MetasoundEditor/Private/MetasoundEditorModule.cpp
--- a/Engine/Plugins/Runtime/Metasound/Source/MetasoundEditor/Private/MetasoundEditorModule.cpp
+++ b/Engine/Plugins/Runtime/Metasound/Source/MetasoundEditor/Private/MetasoundEditorModule.cpp
@@ -972,9 +972,11 @@ namespace Metasound::Editor
 			{
 				using namespace Engine;
 
-				// Use the editor version of RegisterWithFrontend so it refreshes any open MetaSound editors.
-				const bool bForceViewSynchronization = AssetContext == ERegistrationAssetContext::Renaming;
-				FGraphBuilder::RegisterGraphWithFrontend(InMetaSound, bForceViewSynchronization);
+				{
+					Frontend::FMetaSoundAssetRegistrationOptions RegOptions = FGraphBuilder::GetDefaultRegistrationOptions();
+					RegOptions.bForceViewSynchronization = AssetContext == ERegistrationAssetContext::Renaming;
+					FGraphBuilder::RegisterGraphWithFrontend(InMetaSound, MoveTemp(RegOptions));
+				}
 			});
 			MetaSoundEngineModule.GetOnGraphUnregisteredDelegate().BindLambda([](UObject& InMetaSound, Engine::ERegistrationAssetContext AssetContext)
 			{
diff --git a/Engine/Plugins/Runtime/Metasound/Source/MetasoundEditor/Private/MetasoundEditorSubsystem.cpp b/Engine/Plugins/Runtime/Metasound/Source/MetasoundEditor/Private/MetasoundEditorSubsystem.cpp
--- a/Engine/Plugins/Runtime/Metasound/Source/MetasoundEditor/Private/MetasoundEditorSubsystem.cpp
+++ b/Engine/Plugins/Runtime/Metasound/Source/MetasoundEditor/Private/MetasoundEditorSubsystem.cpp
@@ -409,7 +409,12 @@ bool UMetaSoundEditorSubsystem::IsPageAuditionPlatformCookTarget(const FGuid& In
 
 void UMetaSoundEditorSubsystem::RegisterGraphWithFrontend(UObject& InMetaSound, bool bInForceViewSynchronization) const
 {
-	Metasound::Editor::FGraphBuilder::RegisterGraphWithFrontend(InMetaSound, bInForceViewSynchronization);
+	using namespace Metasound::Editor;
+	using namespace Metasound::Frontend;
+	
+	FMetaSoundAssetRegistrationOptions RegOptions = FGraphBuilder::GetDefaultRegistrationOptions();
+	RegOptions.bForceViewSynchronization = bInForceViewSynchronization;
+	FGraphBuilder::RegisterGraphWithFrontend(InMetaSound, MoveTemp(RegOptions));
 }
 
 void UMetaSoundEditorSubsystem::RegisterToolbarExtender(TSharedRef<FExtender> InExtender)
diff --git a/Engine/Plugins/Runtime/Metasound/Source/MetasoundEditor/Private/MetasoundNodeDetailCustomization.cpp b/Engine/Plugins/Runtime/Metasound/Source/MetasoundEditor/Private/MetasoundNodeDetailCustomization.cpp
--- a/Engine/Plugins/Runtime/Metasound/Source/MetasoundEditor/Private/MetasoundNodeDetailCustomization.cpp
+++ b/Engine/Plugins/Runtime/Metasound/Source/MetasoundEditor/Private/MetasoundNodeDetailCustomization.cpp
@@ -1578,8 +1578,11 @@ PRAGMA_ENABLE_DEPRECATION_WARNINGS
 
 						VertexPtr->SetSortOrderIndex(NewValue);
 
-						constexpr bool bInForceViewSynchronization = true;
-						FGraphBuilder::RegisterGraphWithFrontend(*MetaSoundObject, bInForceViewSynchronization);
+						{
+							Frontend::FMetaSoundAssetRegistrationOptions RegOptions = FGraphBuilder::GetDefaultRegistrationOptions();
+							RegOptions.bForceViewSynchronization = true;
+							FGraphBuilder::RegisterGraphWithFrontend(*MetaSoundObject, MoveTemp(RegOptions));
+						}
 					})
 					.Font(IDetailLayoutBuilder::GetDetailFont())
 				];
diff --git a/Engine/Plugins/Runtime/Metasound/Source/MetasoundEditor/Private/MetasoundPinAudioInspector.cpp b/Engine/Plugins/Runtime/Metasound/Source/MetasoundEditor/Private/MetasoundPinAudioInspector.cpp
--- a/Engine/Plugins/Runtime/Metasound/Source/MetasoundEditor/Private/MetasoundPinAudioInspector.cpp
+++ b/Engine/Plugins/Runtime/Metasound/Source/MetasoundEditor/Private/MetasoundPinAudioInspector.cpp
@@ -68,7 +68,7 @@ namespace Metasound::Editor
 		{
 			if (FGraphConnectionManager* ConnectionManager = GetConnectionManager())
 			{
-				const Frontend::FConstOutputHandle OutputHandle = FGraphBuilder::FindReroutedOutputHandleFromPin(GraphPinObj);
+				const Frontend::FConstOutputHandle OutputHandle = FGraphBuilder::FindReroutedConstOutputHandleFromPin(GraphPinObj);
 				const FGuid NodeID = OutputHandle->GetOwningNodeID();
 				const FName OutputName = OutputHandle->GetName();
 				AnalyzerInstanceID = ConnectionManager->AddAudioBusWriter(NodeID, OutputName, AudioDeviceId, Oscilloscope->GetAudioBus());
diff --git a/Engine/Plugins/Runtime/Metasound/Source/MetasoundEditor/Private/SMetasoundPinValueInspector.cpp b/Engine/Plugins/Runtime/Metasound/Source/MetasoundEditor/Private/SMetasoundPinValueInspector.cpp
--- a/Engine/Plugins/Runtime/Metasound/Source/MetasoundEditor/Private/SMetasoundPinValueInspector.cpp
+++ b/Engine/Plugins/Runtime/Metasound/Source/MetasoundEditor/Private/SMetasoundPinValueInspector.cpp
@@ -125,7 +125,7 @@ namespace Metasound
 		{
 			using namespace Frontend;
 
-			FConstOutputHandle ReroutedOutputHandle = GetReroutedOutputHandle();
+			FConstOutputHandle ReroutedOutputHandle = GetConstReroutedOutputHandle();
 			ColorLiterals = PinValueInspectorPrivate::InitializeLiterals(ReroutedOutputHandle->GetDataType());
 			Update();
 			DisplayName = FGraphBuilder::GetDisplayName(*ReroutedOutputHandle);
@@ -177,11 +177,6 @@ namespace Metasound
 			return FGraphBuilder::FindReroutedConstOutputHandleFromPin(GraphPinObj);
 		}
 
-		Frontend::FOutputHandle FMetasoundNumericDebugLineItem::GetReroutedOutputHandle()
-		{
-			return FGraphBuilder::FindReroutedOutputHandleFromPin(GraphPinObj);
-		}
-
 		bool FMetasoundNumericDebugLineItem::GetNodeIDAndOutputName(FGuid& OutNodeID, FName& OutName) const
 		{
 			if (UEdGraphPin* ReroutedOutputPin = FGraphBuilder::FindReroutedOutputPin(GraphPinObj))
@@ -285,7 +280,7 @@ namespace Metasound
 
 			static const FText SetValueFormat = LOCTEXT("PinValueInspector_SetMinColorFormat", "Set '{0}[{1}]' Color to '{2}'");
 
-			const FConstOutputHandle OutputHandle = GetReroutedOutputHandle();
+			const FConstOutputHandle OutputHandle = GetConstReroutedOutputHandle();
 			const FName Name = OutputHandle->GetName();
 			FScopedTransaction Transaction(FText::Format(SetValueFormat, FText::FromName(Name), FText::AsNumber(InIndex), FText::FromString(InColor.ToString())));
 
@@ -306,7 +301,7 @@ namespace Metasound
 		{
 			using namespace Frontend;
 
-			FOutputHandle OutputHandle = GetReroutedOutputHandle();
+			FConstOutputHandle OutputHandle = GetConstReroutedOutputHandle();
 
 			static const FText SetValueFormat = LOCTEXT("PinValueInspector_SetColorLiteralValueFormat", "Set '{0}[{1}]' EdgeStyle Color Value to '{2}'");
 			const FName Name = OutputHandle->GetName();
@@ -327,19 +322,13 @@ namespace Metasound
 			}
 		}
 
-		Frontend::FGraphHandle FMetasoundNumericDebugLineItem::GetGraphHandle()
-		{
-			Frontend::FOutputHandle OutputHandle = FGraphBuilder::GetOutputHandleFromPin(GraphPinObj);
-			return OutputHandle->GetOwningNode()->GetOwningGraph();
-		}
-
 		Frontend::FConstGraphHandle FMetasoundNumericDebugLineItem::GetConstGraphHandle() const
 		{
 			Frontend::FConstOutputHandle OutputHandle = FGraphBuilder::GetConstOutputHandleFromPin(GraphPinObj);
 			return OutputHandle->GetOwningNode()->GetOwningGraph();
 		}
 
-		FGraphConnectionManager* FMetasoundNumericDebugLineItem::GetConnectionManager()
+		const FGraphConnectionManager* FMetasoundNumericDebugLineItem::GetConnectionManager() const
 		{
 			using namespace Frontend;
 
@@ -429,7 +418,7 @@ namespace Metasound
 					}
 				}
 
-				if (FGraphConnectionManager* ConnectionManager = GetConnectionManager())
+				if (const FGraphConnectionManager* ConnectionManager = GetConnectionManager())
 				{
 					Message = FString::Format(TEXT("Value: {0}"), { GetValueStringFunction(*ConnectionManager, NodeID, OutputName) });
 				}
@@ -494,7 +483,7 @@ namespace Metasound
 					.Font(FontInfo)
 				];
 
-			const FName DataType = GetReroutedOutputHandle()->GetDataType();
+			const FName DataType = GetConstReroutedOutputHandle()->GetDataType();
 			const bool bConnectionAnimationsEnabled = GetDefault<UMetasoundEditorSettings>()->AnalyzerAnimationSettings.bAnimateConnections;
 			const bool bIsTypeSupported = PinValueInspectorPrivate::IsMetasoundPrimitiveValueColorizationSupported(DataType);
 			if (bConnectionAnimationsEnabled && bIsTypeSupported)
diff --git a/Engine/Plugins/Runtime/Metasound/Source/MetasoundEditor/Private/SMetasoundPinValueInspector.h b/Engine/Plugins/Runtime/Metasound/Source/MetasoundEditor/Private/SMetasoundPinValueInspector.h
--- a/Engine/Plugins/Runtime/Metasound/Source/MetasoundEditor/Private/SMetasoundPinValueInspector.h
+++ b/Engine/Plugins/Runtime/Metasound/Source/MetasoundEditor/Private/SMetasoundPinValueInspector.h
@@ -65,13 +65,9 @@ namespace Metasound
 
 			Frontend::FConstOutputHandle GetConstReroutedOutputHandle() const;
 
-			Frontend::FOutputHandle GetReroutedOutputHandle();
-
-			Frontend::FGraphHandle GetGraphHandle();
-
 			Frontend::FConstGraphHandle GetConstGraphHandle() const;
 
-			FGraphConnectionManager* GetConnectionManager();
+			const FGraphConnectionManager* GetConnectionManager() const;
 
 			const FMetasoundFrontendEdgeStyle* FindConstEdgeStyle() const;
 
diff --git a/Engine/Plugins/Runtime/Metasound/Source/MetasoundEditor/Private/SMetasoundSpectrumAnalyzerGraphNode.cpp b/Engine/Plugins/Runtime/Metasound/Source/MetasoundEditor/Private/SMetasoundSpectrumAnalyzerGraphNode.cpp
--- a/Engine/Plugins/Runtime/Metasound/Source/MetasoundEditor/Private/SMetasoundSpectrumAnalyzerGraphNode.cpp
+++ b/Engine/Plugins/Runtime/Metasound/Source/MetasoundEditor/Private/SMetasoundSpectrumAnalyzerGraphNode.cpp
@@ -44,7 +44,7 @@ namespace Metasound::Editor
 					FGraphConnectionManager& ConnectionManager = Editor->GetConnectionManager();
 					if (!AnalyzerInstanceID.IsValid() || !ConnectionManager.HasAudioBusWriter(AnalyzerInstanceID))
 					{
-						const Frontend::FConstOutputHandle OutputHandle = FGraphBuilder::FindReroutedOutputHandleFromPin(SourcePin);
+						const Frontend::FConstOutputHandle OutputHandle = FGraphBuilder::FindReroutedConstOutputHandleFromPin(SourcePin);
 						const FGuid NodeID = OutputHandle->GetOwningNodeID();
 						const FName OutputName = OutputHandle->GetName();
 						const Audio::FDeviceId DeviceId = GEditor->GetMainAudioDeviceID();
diff --git a/Engine/Plugins/Runtime/Metasound/Source/MetasoundEngine/Private/MetasoundBuilderSubsystem.cpp b/Engine/Plugins/Runtime/Metasound/Source/MetasoundEngine/Private/MetasoundBuilderSubsystem.cpp
--- a/Engine/Plugins/Runtime/Metasound/Source/MetasoundEngine/Private/MetasoundBuilderSubsystem.cpp
+++ b/Engine/Plugins/Runtime/Metasound/Source/MetasoundEngine/Private/MetasoundBuilderSubsystem.cpp
@@ -21,6 +21,7 @@
 #include "MetasoundFrontendDocument.h"
 #include "MetasoundFrontendDocumentIdGenerator.h"
 #include "MetasoundFrontendGraphBuilder.h"
+#include "MetasoundFrontendNodeTemplateRegistry.h"
 #include "MetasoundFrontendPages.h"
 #include "MetasoundFrontendRegistries.h"
 #include "MetasoundFrontendSearchEngine.h"
@@ -237,6 +238,8 @@ void UMetaSoundSourceBuilder::OnLiveComponentFinished(UAudioComponent* AudioComp
 	if (LiveComponentIDs.IsEmpty())
 	{
 		AudioComponent->OnAudioFinishedNative.Remove(LiveComponentHandle);
+
+		GetMetaSoundSource().SetDynamicGeneratorEnabled(false);
 	}
 }
 
@@ -416,11 +419,56 @@ void UMetaSoundSourceBuilder::OnEdgeAdded(int32 EdgeIndex) const
 	const FMetasoundFrontendEdge& NewEdge = GetConstTargetPageGraphChecked().Edges[EdgeIndex];
 	ExecuteAuditionableTransaction([this, &NewEdge](Metasound::DynamicGraph::FDynamicOperatorTransactor& Transactor)
 	{
+		using namespace Metasound::Frontend;
+
 		const FMetaSoundFrontendDocumentBuilder& DocBuilder = GetConstBuilder();
+
 		const FMetasoundFrontendVertex* FromNodeOutput = DocBuilder.FindNodeOutput(NewEdge.FromNodeID, NewEdge.FromVertexID);
 		const FMetasoundFrontendVertex* ToNodeInput = DocBuilder.FindNodeInput(NewEdge.ToNodeID, NewEdge.ToVertexID);
 		if (FromNodeOutput && ToNodeInput)
 		{
+#if WITH_EDITOR
+			const FMetasoundFrontendNode* ToNode = DocBuilder.FindNode(NewEdge.ToNodeID);
+			const FMetasoundFrontendNode* FromNode = DocBuilder.FindNode(NewEdge.FromNodeID);
+			if (ToNode && FromNode)
+			{
+				const FMetasoundFrontendClass* ToNodeClass = DocBuilder.FindDependency(ToNode->ClassID);
+				const FMetasoundFrontendClass* FromNodeClass = DocBuilder.FindDependency(FromNode->ClassID);
+
+				// Adding edge to a template node can be safely ignored as template nodes are processed out.
+				// Once template node is connected as a source however, then the given edge can be processed accordingly.
+				if (ToNodeClass->Metadata.GetType() == EMetasoundFrontendClassType::Template)
+				{
+					return true;
+				}
+
+				if (FromNodeClass->Metadata.GetType() == EMetasoundFrontendClassType::Template)
+				{
+					bool bHandled = false;
+					const FTemplateNodeEdgeTransactionContext TransactionContext
+					{
+						.Builder = DocBuilder,
+						.PageID = DocBuilder.GetBuildPageID(),
+						.FromNode = FromNode,
+						.FromOutput = FromNodeOutput,
+						.ToNode = ToNode,
+						.ToInput = ToNodeInput
+					};
+
+					const FMetasoundFrontendClassName& FromClassName = FromNodeClass->Metadata.GetClassName();
+					if (const INodeTemplate* Template = INodeTemplateRegistry::Get().FindTemplate(FromClassName))
+					{
+						bHandled = Template->OnEdgeAdded(TransactionContext, Transactor);
+					}
+
+					if (bHandled)
+					{
+						return true;
+					}
+				}
+			}
+#endif // WITH_EDITOR
+
 			Transactor.AddDataEdge(NewEdge.FromNodeID, FromNodeOutput->Name, NewEdge.ToNodeID, ToNodeInput->Name);
 			return true;
 		}
@@ -429,18 +477,6 @@ void UMetaSoundSourceBuilder::OnEdgeAdded(int32 EdgeIndex) const
 	});
 }
 
-TOptional<Metasound::FAnyDataReference> UMetaSoundSourceBuilder::CreateDataReference(
-	const Metasound::FOperatorSettings& InOperatorSettings,
-	FName DataType,
-	const Metasound::FLiteral& InLiteral,
-	Metasound::EDataReferenceAccessType AccessType)
-{
-	using namespace Metasound;
-	using namespace Metasound::Frontend;
-
-	return IDataTypeRegistry::Get().CreateDataReference(DataType, AccessType, InLiteral, InOperatorSettings);
-};
-
 void UMetaSoundSourceBuilder::OnInputAdded(int32 InputIndex)
 {
 	using namespace Metasound::DynamicGraph;
@@ -473,7 +509,7 @@ void UMetaSoundSourceBuilder::OnInputAdded(int32 InputIndex)
 		}
 
 		const FLiteral NewInputLiteral = NewInput.FindConstDefaultChecked(Frontend::DefaultPageID).ToLiteral(NewInput.TypeName);
-		Transactor.AddInputDataDestination(NewInput.NodeID, NewInput.Name, NewInputLiteral, &UMetaSoundSourceBuilder::CreateDataReference);
+		Transactor.AddInputDataDestination(NewInput.NodeID, NewInput.Name, NewInputLiteral, &FMetaSoundFrontendDocumentBuilder::CreateDataReference);
 		return true;
 	});
 }
@@ -494,6 +530,22 @@ void UMetaSoundSourceBuilder::OnNodeAdded(int32 NodeIndex) const
 		const FMetasoundFrontendClass* NodeClass = Builder.FindDependency(AddedNode.ClassID);
 		checkf(NodeClass, TEXT("Node successfully added to graph but document is missing associated dependency"));
 
+#if WITH_EDITOR
+		if (NodeClass->Metadata.GetType() == EMetasoundFrontendClassType::Template)
+		{
+			const FMetasoundFrontendClassName& AddedClassName = NodeClass->Metadata.GetClassName();
+			if (const INodeTemplate* Template = INodeTemplateRegistry::Get().FindTemplate(AddedClassName))
+			{
+				const FTemplateNodeTransactionContext TransactionContext { Builder, Builder.GetBuildPageID(), &AddedNode };
+				const bool bHandled = Template->OnNodeAdded(TransactionContext, Transactor);
+				if (bHandled)
+				{
+					return true;
+				}
+			}
+		}
+#endif // WITH_EDITOR
+
 		FGraphBuilder::FCreateNodeParams CreateNodeParams
 		{
 			.FrontendNode = AddedNode,
@@ -533,7 +585,7 @@ void UMetaSoundSourceBuilder::OnNodeInputLiteralSet(int32 NodeIndex, int32 Verte
 
 			const FMetasoundFrontendLiteral& InputDefault = Node.InputLiterals[LiteralIndex].Value;
 			
-			Transactor.SetValue(Node.GetID(), Input.Name, InputDefault.ToLiteral(Input.TypeName), &UMetaSoundSourceBuilder::CreateDataReference);
+			Transactor.SetValue(Node.GetID(), Input.Name, InputDefault.ToLiteral(Input.TypeName), &FMetaSoundFrontendDocumentBuilder::CreateDataReference);
 			return true;
 		});
 	}
@@ -581,17 +633,66 @@ void UMetaSoundSourceBuilder::OnRemoveSwappingEdge(int32 SwapIndex, int32 LastIn
 	{
 		using namespace Metasound;
 		using namespace Metasound::Engine;
+		using namespace Metasound::Frontend;
 
-		const FMetaSoundFrontendDocumentBuilder& Builder = GetConstBuilder();
-		const FMetasoundFrontendVertex* FromNodeOutput = Builder.FindNodeOutput(EdgeBeingRemoved.FromNodeID, EdgeBeingRemoved.FromVertexID);
-		const FMetasoundFrontendVertex* ToNodeInput = Builder.FindNodeInput(EdgeBeingRemoved.ToNodeID, EdgeBeingRemoved.ToVertexID);
+		const FMetaSoundFrontendDocumentBuilder& DocBuilder = GetConstBuilder();
+		const FMetasoundFrontendVertex* FromNodeOutput = DocBuilder.FindNodeOutput(EdgeBeingRemoved.FromNodeID, EdgeBeingRemoved.FromVertexID);
+		const FMetasoundFrontendVertex* ToNodeInput = DocBuilder.FindNodeInput(EdgeBeingRemoved.ToNodeID, EdgeBeingRemoved.ToVertexID);
 		if (FromNodeOutput && ToNodeInput)
 		{
 			TOptional<FMetasoundFrontendLiteral> InputDefault = BuilderSubsystemPrivate::TryResolveNodeInputDefault(Builder, EdgeBeingRemoved.ToNodeID, ToNodeInput->Name);
-
 			if (ensureAlwaysMsgf(InputDefault.IsSet(), TEXT("Could not dynamically assign default literal upon removing edge: literal should be assigned by either the frontend document's input or the class definition")))
 			{
-				Transactor.RemoveDataEdge(EdgeBeingRemoved.FromNodeID, FromNodeOutput->Name, EdgeBeingRemoved.ToNodeID, ToNodeInput->Name, InputDefault->ToLiteral(ToNodeInput->TypeName), &UMetaSoundSourceBuilder::CreateDataReference);
+#if WITH_EDITOR
+				const FMetasoundFrontendNode* ToNode = DocBuilder.FindNode(EdgeBeingRemoved.ToNodeID);
+				const FMetasoundFrontendNode* FromNode = DocBuilder.FindNode(EdgeBeingRemoved.FromNodeID);
+				if (ToNode && FromNode)
+				{
+					const FMetasoundFrontendClass* ToNodeClass = DocBuilder.FindDependency(ToNode->ClassID);
+					const FMetasoundFrontendClass* FromNodeClass = DocBuilder.FindDependency(FromNode->ClassID);
+
+					// Adding edge to a template node can be safely ignored as template nodes are processed out.
+					// Once template node is connected as a source however, then the given edge can be processed accordingly.
+					if (ToNodeClass->Metadata.GetType() == EMetasoundFrontendClassType::Template)
+					{
+						return true;
+					}
+
+					if (FromNodeClass->Metadata.GetType() == EMetasoundFrontendClassType::Template)
+					{
+						const FTemplateNodeEdgeTransactionContext TransactionContext
+						{
+							.Builder = DocBuilder,
+							.PageID = DocBuilder.GetBuildPageID(),
+							.FromNode = FromNode,
+							.FromOutput = FromNodeOutput,
+							.ToNode = ToNode,
+							.ToInput = ToNodeInput,
+							.ToInputDefault = &InputDefault.GetValue()
+						};
+
+						bool bHandled = false;
+						const FMetasoundFrontendClassName& FromClassName = FromNodeClass->Metadata.GetClassName();
+						if (const INodeTemplate* Template = INodeTemplateRegistry::Get().FindTemplate(FromClassName))
+						{
+							bHandled = Template->OnRemoveSwappingEdge(TransactionContext, Transactor);
+						}
+
+						if (bHandled)
+						{
+							return true;
+						}
+					}
+				}
+#endif // WITH_EDITOR
+
+				Transactor.RemoveDataEdge(
+					EdgeBeingRemoved.FromNodeID,
+					FromNodeOutput->Name,
+					EdgeBeingRemoved.ToNodeID,
+					ToNodeInput->Name,
+					InputDefault->ToLiteral(ToNodeInput->TypeName),
+					&FMetaSoundFrontendDocumentBuilder::CreateDataReference);
 				return true;
 			}
 		}
@@ -660,6 +761,26 @@ void UMetaSoundSourceBuilder::OnRemoveSwappingNode(int32 SwapIndex, int32 LastIn
 		using namespace Metasound::Frontend;
 
 		const FMetasoundFrontendNode& NodeBeingRemoved = GetConstTargetPageGraphChecked().Nodes[SwapIndex];
+
+#if WITH_EDITOR
+		const FMetasoundFrontendClass* RemovingNodeClass = Builder.FindDependency(NodeBeingRemoved.ClassID);
+		checkf(RemovingNodeClass, TEXT("Node to be removed from frontend graph is missing associated dependency"));
+		if (RemovingNodeClass->Metadata.GetType() == EMetasoundFrontendClassType::Template)
+		{
+			const FMetasoundFrontendClassName& RemovingClassName = RemovingNodeClass->Metadata.GetClassName();
+			if (const INodeTemplate* Template = INodeTemplateRegistry::Get().FindTemplate(RemovingClassName))
+			{
+				const FMetaSoundFrontendDocumentBuilder& DocBuilder = GetConstBuilder();
+				const FTemplateNodeTransactionContext TransactionContext { DocBuilder, DocBuilder.GetBuildPageID(), &NodeBeingRemoved };
+				const bool bHandled = Template->OnRemoveSwappingNode(TransactionContext, Transactor);
+				if (bHandled)
+				{
+					return true;
+				}
+			}
+		}
+#endif // WITH_EDITOR
+
 		const FGuid& NodeID = NodeBeingRemoved.GetID();
 		Transactor.RemoveNode(NodeID);
 		return true;
@@ -686,7 +807,7 @@ void UMetaSoundSourceBuilder::OnRemovingNodeInputLiteral(int32 NodeIndex, int32
 			TOptional<FMetasoundFrontendLiteral> InputDefault = BuilderSubsystemPrivate::TryResolveNodeInputDefault(Builder, Node.GetID(), Input.Name);
 			if (ensureAlwaysMsgf(InputDefault.IsSet(), TEXT("Could not dynamically assign default literal from class definition upon removing input '%s' literal: document's dependency entry invalid and has no default assigned"), *Input.Name.ToString()))
 			{
-				Transactor.SetValue(Node.GetID(), Input.Name, InputDefault->ToLiteral(Input.TypeName), CreateDataReference);
+				Transactor.SetValue(Node.GetID(), Input.Name, InputDefault->ToLiteral(Input.TypeName), &FMetaSoundFrontendDocumentBuilder::CreateDataReference);
 				return true;
 			}
 
diff --git a/Engine/Plugins/Runtime/Metasound/Source/MetasoundEngine/Public/MetasoundBuilderSubsystem.h b/Engine/Plugins/Runtime/Metasound/Source/MetasoundEngine/Public/MetasoundBuilderSubsystem.h
--- a/Engine/Plugins/Runtime/Metasound/Source/MetasoundEngine/Public/MetasoundBuilderSubsystem.h
+++ b/Engine/Plugins/Runtime/Metasound/Source/MetasoundEngine/Public/MetasoundBuilderSubsystem.h
@@ -116,8 +116,6 @@ protected:
 	UE_API virtual void OnRemovingAssetReference(TScriptInterface<IMetaSoundDocumentInterface> DocInterface) override;
 
 private:
-	static UE_API TOptional<Metasound::FAnyDataReference> CreateDataReference(const Metasound::FOperatorSettings& InOperatorSettings, FName DataType, const Metasound::FLiteral& InLiteral, Metasound::EDataReferenceAccessType AccessType);
-
 	UE_API const FMetasoundFrontendGraph& GetConstTargetPageGraphChecked() const;
 
 	UE_API const UMetaSoundSource& GetMetaSoundSource() const;
diff --git a/Engine/Plugins/Runtime/Metasound/Source/MetasoundFrontend/Private/MetasoundFrontendDocumentBuilder.cpp b/Engine/Plugins/Runtime/Metasound/Source/MetasoundFrontend/Private/MetasoundFrontendDocumentBuilder.cpp
--- a/Engine/Plugins/Runtime/Metasound/Source/MetasoundFrontend/Private/MetasoundFrontendDocumentBuilder.cpp
+++ b/Engine/Plugins/Runtime/Metasound/Source/MetasoundFrontend/Private/MetasoundFrontendDocumentBuilder.cpp
@@ -2264,6 +2264,18 @@ bool FMetaSoundFrontendDocumentBuilder::ConvertToPreset(const FMetasoundFrontend
 	return false;
 }
 
+TOptional<Metasound::FAnyDataReference> FMetaSoundFrontendDocumentBuilder::CreateDataReference(
+	const Metasound::FOperatorSettings& InOperatorSettings,
+	FName DataType,
+	const Metasound::FLiteral& InLiteral,
+	Metasound::EDataReferenceAccessType AccessType)
+{
+	using namespace Metasound;
+	using namespace Metasound::Frontend;
+
+	return IDataTypeRegistry::Get().CreateDataReference(DataType, AccessType, InLiteral, InOperatorSettings);
+};
+
 #if WITH_EDITORONLY_DATA
 bool FMetaSoundFrontendDocumentBuilder::DiffAgainstRegistryInterface(const FGuid& InNodeID, const FGuid& InPageID, bool bInUseHighestMinorVersion, Metasound::Frontend::FClassInterfaceUpdates& OutInterfaceUpdates) const
 {
diff --git a/Engine/Plugins/Runtime/Metasound/Source/MetasoundFrontend/Private/MetasoundFrontendNodeTemplateRegistry.cpp b/Engine/Plugins/Runtime/Metasound/Source/MetasoundFrontend/Private/MetasoundFrontendNodeTemplateRegistry.cpp
--- a/Engine/Plugins/Runtime/Metasound/Source/MetasoundFrontend/Private/MetasoundFrontendNodeTemplateRegistry.cpp
+++ b/Engine/Plugins/Runtime/Metasound/Source/MetasoundFrontend/Private/MetasoundFrontendNodeTemplateRegistry.cpp
@@ -2,6 +2,7 @@
 
 #include "MetasoundFrontendNodeTemplateRegistry.h"
 
+#include "MetasoundDynamicOperatorTransactor.h"
 #include "MetasoundFrontendDocument.h"
 #include "MetasoundFrontendDocumentBuilder.h"
 #include "MetasoundFrontendNodeClassRegistry.h"
@@ -148,6 +149,27 @@ namespace Metasound::Frontend
 	{
 		return true;
 	}
+
+	bool FNodeTemplateBase::OnEdgeAdded(const FTemplateNodeEdgeTransactionContext& TransactionContext, DynamicGraph::FDynamicOperatorTransactor& OutTransactor) const
+	{
+		return false;
+	}
+
+	bool FNodeTemplateBase::OnNodeAdded(const FTemplateNodeTransactionContext& TransactionContext, DynamicGraph::FDynamicOperatorTransactor& OutTransactor) const
+	{
+		return false;
+	}
+
+	bool FNodeTemplateBase::OnRemoveSwappingEdge(const FTemplateNodeEdgeTransactionContext& TransactionContext, DynamicGraph::FDynamicOperatorTransactor& OutTransactor) const
+	{
+		return false;
+	}
+
+	bool FNodeTemplateBase::OnRemoveSwappingNode(const FTemplateNodeTransactionContext& TransactionContext, DynamicGraph::FDynamicOperatorTransactor& OutTransactor) const
+	{
+		return false;
+	}
+
 #endif // WITH_EDITOR
 
 	void RegisterNodeTemplate(TUniquePtr<INodeTemplate>&& InTemplate)
diff --git a/Engine/Plugins/Runtime/Metasound/Source/MetasoundFrontend/Private/NodeTemplates/MetasoundFrontendNodeTemplateReroute.cpp b/Engine/Plugins/Runtime/Metasound/Source/MetasoundFrontend/Private/NodeTemplates/MetasoundFrontendNodeTemplateReroute.cpp
--- a/Engine/Plugins/Runtime/Metasound/Source/MetasoundFrontend/Private/NodeTemplates/MetasoundFrontendNodeTemplateReroute.cpp
+++ b/Engine/Plugins/Runtime/Metasound/Source/MetasoundFrontend/Private/NodeTemplates/MetasoundFrontendNodeTemplateReroute.cpp
@@ -4,6 +4,7 @@
 
 #include "Algo/AnyOf.h"
 #include "MetasoundAssetManager.h"
+#include "MetasoundDynamicOperatorTransactor.h"
 #include "MetasoundFrontendDataTypeRegistry.h"
 #include "MetasoundFrontendDocument.h"
 #include "MetasoundFrontendDocumentBuilder.h"
@@ -16,6 +17,17 @@ namespace Metasound::Frontend
 {
 	namespace ReroutePrivate
 	{
+		struct FRerouteNodeVertexPair
+		{
+			const FMetasoundFrontendNode* Node= nullptr;
+			const FMetasoundFrontendVertex* Vertex = nullptr;
+
+			bool IsValid() const
+			{
+				return Node && Vertex;
+			}
+		};
+
 		class FRerouteNodeTemplateTransform : public INodeTemplateTransform
 		{
 		public:
@@ -88,12 +100,17 @@ namespace Metasound::Frontend
 			return false;
 		}
 
-		const FMetasoundFrontendVertex* FindReroutedOutputVertex(const FMetaSoundFrontendDocumentBuilder& InBuilder, const FGuid& InPageID, const FMetasoundFrontendNode& InOutputOwningNode, const FMetasoundFrontendVertex& InOutputVertex)
+		FRerouteNodeVertexPair FindReroutedOutputVertex(
+			const FMetasoundFrontendClassName& InClassName,
+			const FMetaSoundFrontendDocumentBuilder& InBuilder,
+			const FGuid& InPageID, const FMetasoundFrontendNode& InOutputOwningNode, const FMetasoundFrontendVertex& InOutputVertex)
 		{
 			const FMetasoundFrontendVertex* ReroutedVertex = &InOutputVertex;
+			const FMetasoundFrontendNode* ReroutedNode = &InOutputOwningNode;
+
 			if (const FMetasoundFrontendClass* Class = InBuilder.FindDependency(InOutputOwningNode.ClassID))
 			{
-				if (Class->Metadata.GetClassName() == FRerouteNodeTemplate::ClassName)
+				if (Class->Metadata.GetClassName() == InClassName)
 				{
 					const FGuid& OutputOwningNodeID = InOutputOwningNode.GetID();
 					TArray<const FMetasoundFrontendVertex*> Inputs = InBuilder.FindNodeInputs(OutputOwningNodeID, { }, &InPageID);
@@ -101,23 +118,23 @@ namespace Metasound::Frontend
 					{
 						if (const FMetasoundFrontendVertex* RerouteInput = Inputs.Last())
 						{
-							const FMetasoundFrontendNode* ConnectedNode = nullptr;
-							if (const FMetasoundFrontendVertex* OutputVertex = InBuilder.FindNodeOutputConnectedToNodeInput(OutputOwningNodeID, RerouteInput->VertexID, &ConnectedNode, &InPageID))
+							if (const FMetasoundFrontendVertex* OutputVertex = InBuilder.FindNodeOutputConnectedToNodeInput(OutputOwningNodeID, RerouteInput->VertexID, &ReroutedNode, &InPageID))
 							{
-								check(ConnectedNode);
-								return FindReroutedOutputVertex(InBuilder, InPageID, *ConnectedNode, *OutputVertex);
+								check(ReroutedNode);
+								return FindReroutedOutputVertex(InClassName, InBuilder, InPageID, *ReroutedNode, *OutputVertex);
 							}
 						}
 					}
 
-					return nullptr;
+					return { };
 				}
 			}
 
-			return ReroutedVertex;
+			return FRerouteNodeVertexPair { ReroutedNode, ReroutedVertex};
 		}
 
 		void FindReroutedInputVertices(
+			const FMetasoundFrontendClassName& InClassName,
 			const FMetaSoundFrontendDocumentBuilder& InBuilder,
 			const FGuid& InPageID,
 			const FMetasoundFrontendNode& InInputOwningNode,
@@ -129,7 +146,7 @@ namespace Metasound::Frontend
 
 			if (const FMetasoundFrontendClass* Class = InBuilder.FindDependency(InInputOwningNode.ClassID))
 			{
-				if (Class->Metadata.GetClassName() == FRerouteNodeTemplate::ClassName)
+				if (Class->Metadata.GetClassName() == InClassName)
 				{
 					const FGuid& InputOwningNodeID = InInputOwningNode.GetID();
 					TArray<const FMetasoundFrontendVertex*> Outputs = InBuilder.FindNodeOutputs(InputOwningNodeID);
@@ -148,7 +165,7 @@ namespace Metasound::Frontend
 							const FMetasoundFrontendVertex* ConnectedInputVertex = ConnectedInputVertices[Index];
 							check(ConnectedInputVertex);
 
-							FindReroutedInputVertices(InBuilder, InPageID, *ConnectedInputOwningNode, *ConnectedInputVertex, InOutReroutedInputOwningNodes, InOutReroutedInputVertices);
+							FindReroutedInputVertices(InClassName, InBuilder, InPageID, *ConnectedInputOwningNode, *ConnectedInputVertex, InOutReroutedInputOwningNodes, InOutReroutedInputVertices);
 						}
 					}
 
@@ -353,7 +370,7 @@ namespace Metasound::Frontend
 		const bool bConnectedToNonRerouteOutputs = Algo::AnyOf(Outputs, [&InBuilder, &InPageID, &Node](const FMetasoundFrontendVertex* OutputVertex)
 		{
 			using namespace ReroutePrivate;
-			return FindReroutedOutputVertex(InBuilder, InPageID, *Node, *OutputVertex) != nullptr;
+			return FindReroutedOutputVertex(ClassName, InBuilder, InPageID, *Node, *OutputVertex).IsValid();
 		});
 
 		TArray<const FMetasoundFrontendVertex*> Inputs = InBuilder.FindNodeInputs(InNodeID, { }, &InPageID);
@@ -362,7 +379,7 @@ namespace Metasound::Frontend
 			using namespace ReroutePrivate;
 			TArray<const FMetasoundFrontendVertex*> InputVertices;
 			TArray<const FMetasoundFrontendNode*> InputVerticesOwningNodes;
-			FindReroutedInputVertices(InBuilder, InPageID, *Node, *InputVertex, InputVerticesOwningNodes, InputVertices);
+			FindReroutedInputVertices(ClassName, InBuilder, InPageID, *Node, *InputVertex, InputVerticesOwningNodes, InputVertices);
 			return !InputVertices.IsEmpty();
 		});
 
@@ -374,6 +391,82 @@ namespace Metasound::Frontend
 
 		return bHasRequiredConnections;
 	}
+
+	bool FRerouteNodeTemplate::OnEdgeAdded(const FTemplateNodeEdgeTransactionContext& TransactionContext, DynamicGraph::FDynamicOperatorTransactor& OutTransactor) const
+	{
+		using namespace ReroutePrivate;
+
+		const FRerouteNodeVertexPair RerouteOutputPair = FindReroutedOutputVertex(GetClassName(), TransactionContext.Builder, TransactionContext.PageID, *TransactionContext.FromNode, *TransactionContext.FromOutput);
+		if (RerouteOutputPair.IsValid())
+		{
+			check(RerouteOutputPair.Node && RerouteOutputPair.Vertex);
+			if (const FMetasoundFrontendClass* RerouteClass = TransactionContext.Builder.FindDependency(RerouteOutputPair.Node->ClassID))
+			{
+				if (RerouteClass->Metadata.GetType() == EMetasoundFrontendClassType::Template)
+				{
+					if (const INodeTemplate* Template = INodeTemplateRegistry::Get().FindTemplate(RerouteClass->Metadata.GetClassName()))
+					{
+						FTemplateNodeEdgeTransactionContext RerouteContext = TransactionContext;
+						RerouteContext.FromNode = RerouteOutputPair.Node;
+						RerouteContext.FromOutput = RerouteOutputPair.Vertex;
+
+						return Template->OnEdgeAdded(TransactionContext, OutTransactor);
+					}
+				}
+			}
+			OutTransactor.AddDataEdge(RerouteOutputPair.Node->GetID(), RerouteOutputPair.Vertex->Name, TransactionContext.ToNode->GetID(), TransactionContext.ToInput->Name);
+			return true;
+		}
+
+		return false;
+	}
+
+	bool FRerouteNodeTemplate::OnNodeAdded(const FTemplateNodeTransactionContext& TransactionContext, DynamicGraph::FDynamicOperatorTransactor& OutTransactor) const
+	{
+		// Reroutes do not apply transaction as they are essentially pass through nodes that can provide edge context on demand but do not apply operators with any execution.
+		return true;
+	}
+
+	bool FRerouteNodeTemplate::OnRemoveSwappingEdge(const FTemplateNodeEdgeTransactionContext& TransactionContext, DynamicGraph::FDynamicOperatorTransactor& OutTransactor) const
+	{
+		using namespace ReroutePrivate;
+
+		const FRerouteNodeVertexPair RerouteOutputPair = FindReroutedOutputVertex(GetClassName(), TransactionContext.Builder, TransactionContext.PageID, *TransactionContext.FromNode, *TransactionContext.FromOutput);
+		if (RerouteOutputPair.IsValid())
+		{
+			check(RerouteOutputPair.Node && RerouteOutputPair.Vertex);
+			if (const FMetasoundFrontendClass* RerouteClass = TransactionContext.Builder.FindDependency(RerouteOutputPair.Node->ClassID))
+			{
+				if (RerouteClass->Metadata.GetType() == EMetasoundFrontendClassType::Template)
+				{
+					if (const INodeTemplate* Template = INodeTemplateRegistry::Get().FindTemplate(RerouteClass->Metadata.GetClassName()))
+					{
+						FTemplateNodeEdgeTransactionContext RerouteContext = TransactionContext;
+						RerouteContext.FromNode = RerouteOutputPair.Node;
+						RerouteContext.FromOutput = RerouteOutputPair.Vertex;
+						return Template->OnRemoveSwappingEdge(TransactionContext, OutTransactor);
+					}
+				}
+			}
+
+			OutTransactor.RemoveDataEdge(
+				RerouteOutputPair.Node->GetID(),
+				RerouteOutputPair.Vertex->Name,
+				TransactionContext.ToNode->GetID(),
+				TransactionContext.ToInput->Name,
+				TransactionContext.ToInputDefault->ToLiteral(TransactionContext.ToInput->TypeName),
+				&FMetaSoundFrontendDocumentBuilder::CreateDataReference);
+			return true;
+		}
+
+		return false;
+	}
+
+	bool FRerouteNodeTemplate::OnRemoveSwappingNode(const FTemplateNodeTransactionContext& TransactionContext, DynamicGraph::FDynamicOperatorTransactor& OutTransactor) const
+	{
+		// Reroutes do not apply transaction as they are essentially pass through nodes that can provide edge context on demand but do not apply operators with any execution.
+		return true;
+	}
 #endif // WITH_EDITOR
 
 	bool FRerouteNodeTemplate::IsInputAccessTypeDynamic() const
diff --git a/Engine/Plugins/Runtime/Metasound/Source/MetasoundFrontend/Public/MetasoundFrontendDocumentBuilder.h b/Engine/Plugins/Runtime/Metasound/Source/MetasoundFrontend/Public/MetasoundFrontendDocumentBuilder.h
--- a/Engine/Plugins/Runtime/Metasound/Source/MetasoundFrontend/Public/MetasoundFrontendDocumentBuilder.h
+++ b/Engine/Plugins/Runtime/Metasound/Source/MetasoundFrontend/Public/MetasoundFrontendDocumentBuilder.h
@@ -4,11 +4,14 @@
 #include "Containers/Map.h"
 #include "Interfaces/MetasoundFrontendInterfaceRegistry.h"
 #include "MetasoundFrontendDocumentCacheInterface.h"
+#include "MetasoundDataReference.h"
 #include "MetasoundDocumentInterface.h"
 #include "MetasoundFrontendDocument.h"
 #include "MetasoundFrontendDocumentModifyDelegates.h"
 #include "MetasoundFrontendRegistries.h"
 #include "MetasoundFrontendTransform.h"
+#include "MetasoundLiteral.h"
+#include "MetasoundOperatorSettings.h"
 #include "MetasoundVertex.h"
 #include "Templates/Function.h"
 #include "UObject/ScriptInterface.h"
@@ -225,6 +228,8 @@ public:
 	UE_API bool ConvertFromPreset();
 	UE_API bool ConvertToPreset(const FMetasoundFrontendDocument& InReferencedDocument, TSharedPtr<Metasound::Frontend::FDocumentModifyDelegates> ModifyDelegates = { });
 
+	static UE_API TOptional<Metasound::FAnyDataReference> CreateDataReference(const Metasound::FOperatorSettings& InOperatorSettings, FName DataType, const Metasound::FLiteral& InLiteral, Metasound::EDataReferenceAccessType AccessType);
+
 #if WITH_EDITORONLY_DATA
 	// Fills out the provided ClassInterfaceUpdate struct with the differences between
 	// the registry's version of the class interface and that of the node.
diff --git a/Engine/Plugins/Runtime/Metasound/Source/MetasoundFrontend/Public/MetasoundFrontendNodeTemplateRegistry.h b/Engine/Plugins/Runtime/Metasound/Source/MetasoundFrontend/Public/MetasoundFrontendNodeTemplateRegistry.h
--- a/Engine/Plugins/Runtime/Metasound/Source/MetasoundFrontend/Public/MetasoundFrontendNodeTemplateRegistry.h
+++ b/Engine/Plugins/Runtime/Metasound/Source/MetasoundFrontend/Public/MetasoundFrontendNodeTemplateRegistry.h
@@ -18,6 +18,15 @@
 class IMetaSoundDocumentInterface;
 struct FMetaSoundFrontendDocumentBuilder;
 
+namespace Metasound
+{
+	namespace DynamicGraph
+	{
+		class FDynamicOperatorTransactor;
+	} // DynamicGraph
+} // namespace Metasound
+
+
 USTRUCT()
 struct FNodeTemplateGenerateInterfaceParams
 {
@@ -43,6 +52,42 @@ namespace Metasound::Frontend
 	class IOutputController;
 
 	using FConstNodeHandle = TSharedRef<const INodeController>;
+
+	struct FTemplateNodeTransactionContext
+	{
+		// Builder executing the transaction
+		const FMetaSoundFrontendDocumentBuilder& Builder;
+
+		// PageID of node being acted on
+		const FGuid& PageID;
+
+		// Node being transacted on
+		const FMetasoundFrontendNode* Node = nullptr;
+	};
+
+	struct FTemplateNodeEdgeTransactionContext
+	{
+		// Builder executing the transaction
+		const FMetaSoundFrontendDocumentBuilder& Builder;
+
+		// PageID of node being acted on
+		const FGuid& PageID;
+
+		// From node being transacted on
+		const FMetasoundFrontendNode* FromNode = nullptr;
+
+		// Node Output being transacted on
+		const FMetasoundFrontendVertex* FromOutput = nullptr;
+
+		// To Node being transacted on
+		const FMetasoundFrontendNode* ToNode = nullptr;
+
+		// Node Input being transacted on
+		const FMetasoundFrontendVertex* ToInput = nullptr;
+
+		// Resolved default literal of input (value may be from class definition or set on node)
+		const FMetasoundFrontendLiteral* ToInputDefault = nullptr;
+	};
 #endif // WITH_EDITOR
 
 	class INodeTemplateTransform
@@ -78,14 +123,8 @@ namespace Metasound::Frontend
 #if WITH_EDITOR
 		virtual FText GetNodeDisplayName(const IMetaSoundDocumentInterface& Interface, const FGuid& InPageID, const FGuid& InNodeID) const = 0;
 
-		UE_DEPRECATED(5.5, "Use overload GetOutputVertexDisplayName with supplied builder & pageID instead")
-		virtual FText GetInputPinDisplayName(const Frontend::IInputController& InInput) const { return FText(); }
-
 		virtual FText GetInputVertexDisplayName(const FMetaSoundFrontendDocumentBuilder& InBuilder, const FGuid& InPageID, const FGuid& InNodeID, FName InputName) const = 0;
 
-		UE_DEPRECATED(5.5, "Use overload GetOutputVertexDisplayName with supplied builder & pageID instead")
-		virtual FText GetOutputPinDisplayName(const Frontend::IOutputController& InOutput) const { return FText(); };
-
 		virtual FText GetOutputVertexDisplayName(const FMetaSoundFrontendDocumentBuilder& InBuilder, const FGuid& InPageID, const FGuid& InNodeID, FName OutputName) const = 0;
 
 		static UE_API FText ResolveMemberDisplayName(FName VertexName, FText DisplayName, bool bIncludeNamespace);
@@ -107,9 +146,6 @@ namespace Metasound::Frontend
 		virtual const FMetasoundFrontendVersionNumber& GetVersionNumber() const = 0;
 
 #if WITH_EDITOR
-		UE_DEPRECATED(5.5, "Use version that takes a Frontend PageID, NodeID and builder instead")
-		virtual bool HasRequiredConnections(FConstNodeHandle InNodeHandle, FString* OutMessage = nullptr) const { return false; }
-
 		// Returns whether or not the given node template has the necessary
 		// required connections to be preprocessed (editor only).
 		virtual bool HasRequiredConnections(const FMetaSoundFrontendDocumentBuilder& InBuilder, const FGuid& InPageID, const FGuid& InNodeID, FString* OutMessage = nullptr) const = 0;
@@ -130,6 +166,24 @@ namespace Metasound::Frontend
 		// Given the provided node interface, returns whether or not it conforms to an expected format
 		// that can be successfully manipulated by a generated node template transform.
 		virtual bool IsValidNodeInterface(const FMetasoundFrontendNodeInterface& InNodeInterface) const = 0;
+
+#if WITH_EDITOR
+		// Provided a Frontend transaction context, apply adding edge transaction to the given dynamic core transactor.
+		// Returns whether or not Frontend operation was handled and applied to the given transactor.
+		virtual bool OnEdgeAdded(const FTemplateNodeEdgeTransactionContext& TransactionContext, DynamicGraph::FDynamicOperatorTransactor& OutTransactor) const = 0;
+
+		// Provided a Frontend transaction context, apply adding node transaction to the given dynamic core transactor.
+		// Returns whether or not Frontend operation was handled and applied to the given transactor.
+		virtual bool OnNodeAdded(const FTemplateNodeTransactionContext& TransactionContext, DynamicGraph::FDynamicOperatorTransactor& OutTransactor) const = 0;
+
+		// Provided a Frontend transaction context, apply adding edge transaction to the given dynamic core transactor.
+		// Returns whether or not Frontend operation was handled and applied to the given transactor.
+		virtual bool OnRemoveSwappingEdge(const FTemplateNodeEdgeTransactionContext& TransactionContext, DynamicGraph::FDynamicOperatorTransactor& OutTransactor) const = 0;
+
+		// Provided a Frontend transaction context, apply removing node transaction to the given dynamic core transactor.
+		// Returns whether or not Frontend operation was handled and applied to the given transactor.
+		virtual bool OnRemoveSwappingNode(const FTemplateNodeTransactionContext& TransactionContext, DynamicGraph::FDynamicOperatorTransactor& OutTransactor) const = 0;
+#endif // WITH_EDITOR
 	};
 
 	class INodeTemplateRegistry
@@ -157,6 +211,11 @@ namespace Metasound::Frontend
 		UE_API virtual FText GetInputVertexDisplayName(const FMetaSoundFrontendDocumentBuilder& InBuilder, const FGuid& InPageID, const FGuid& InNodeID, FName InputName) const override;
 		UE_API virtual FText GetOutputVertexDisplayName(const FMetaSoundFrontendDocumentBuilder& InBuilder, const FGuid& InPageID, const FGuid& InNodeID, FName OutputName) const override;
 		UE_API virtual bool HasRequiredConnections(const FMetaSoundFrontendDocumentBuilder& InBuilder, const FGuid& InPageID, const FGuid& InNodeID, FString* OutMessage = nullptr) const override;
+
+		UE_API virtual bool OnEdgeAdded(const FTemplateNodeEdgeTransactionContext& TransactionContext, DynamicGraph::FDynamicOperatorTransactor& OutTransactor) const override;
+		UE_API virtual bool OnNodeAdded(const FTemplateNodeTransactionContext& TransactionContext, DynamicGraph::FDynamicOperatorTransactor& OutTransactor) const override;
+		UE_API virtual bool OnRemoveSwappingEdge(const FTemplateNodeEdgeTransactionContext& TransactionContext, DynamicGraph::FDynamicOperatorTransactor& OutTransactor) const override;
+		UE_API virtual bool OnRemoveSwappingNode(const FTemplateNodeTransactionContext& TransactionContext, DynamicGraph::FDynamicOperatorTransactor& OutTransactor) const override;
 #endif // WITH_EDITOR
 	};
 
diff --git a/Engine/Plugins/Runtime/Metasound/Source/MetasoundFrontend/Public/NodeTemplates/MetasoundFrontendNodeTemplateReroute.h b/Engine/Plugins/Runtime/Metasound/Source/MetasoundFrontend/Public/NodeTemplates/MetasoundFrontendNodeTemplateReroute.h
--- a/Engine/Plugins/Runtime/Metasound/Source/MetasoundFrontend/Public/NodeTemplates/MetasoundFrontendNodeTemplateReroute.h
+++ b/Engine/Plugins/Runtime/Metasound/Source/MetasoundFrontend/Public/NodeTemplates/MetasoundFrontendNodeTemplateReroute.h
@@ -49,6 +49,11 @@ namespace Metasound::Frontend
 
 #if WITH_EDITOR
 		UE_API virtual bool HasRequiredConnections(const FMetaSoundFrontendDocumentBuilder& InBuilder, const FGuid& InPageID, const FGuid& InNodeID, FString* OutMessage = nullptr) const override;
+
+		UE_API virtual bool OnEdgeAdded(const FTemplateNodeEdgeTransactionContext& TransactionContext, DynamicGraph::FDynamicOperatorTransactor& OutTransactor) const override;
+		UE_API virtual bool OnNodeAdded(const FTemplateNodeTransactionContext& TransactionContext, DynamicGraph::FDynamicOperatorTransactor& OutTransactor) const override;
+		UE_API virtual bool OnRemoveSwappingEdge(const FTemplateNodeEdgeTransactionContext& TransactionContext, DynamicGraph::FDynamicOperatorTransactor& OutTransactor) const override;
+		UE_API virtual bool OnRemoveSwappingNode(const FTemplateNodeTransactionContext& TransactionContext, DynamicGraph::FDynamicOperatorTransactor& OutTransactor) const override;
 #endif // WITH_EDITOR
 	};
 } // namespace Metasound::Frontend
diff --git a/Engine/Plugins/Runtime/Metasound/Source/MetasoundGraphCore/MetasoundGraphCore.Build.cs b/Engine/Plugins/Runtime/Metasound/Source/MetasoundGraphCore/MetasoundGraphCore.Build.cs
--- a/Engine/Plugins/Runtime/Metasound/Source/MetasoundGraphCore/MetasoundGraphCore.Build.cs
+++ b/Engine/Plugins/Runtime/Metasound/Source/MetasoundGraphCore/MetasoundGraphCore.Build.cs
@@ -9,6 +9,7 @@ namespace UnrealBuildTool.Rules
 	{
 		public MetasoundGraphCore(ReadOnlyTargetRules Target) : base(Target)
 		{
+			OptimizeCode = CodeOptimization.Never;
 			PublicDependencyModuleNames.AddRange(
 				new string[]
 				{
diff --git a/Engine/Source/Editor/UnrealEd/Private/EditorEngine.cpp b/Engine/Source/Editor/UnrealEd/Private/EditorEngine.cpp
--- a/Engine/Source/Editor/UnrealEd/Private/EditorEngine.cpp
+++ b/Engine/Source/Editor/UnrealEd/Private/EditorEngine.cpp
@@ -2859,8 +2859,10 @@ UAudioComponent* UEditorEngine::GetPreviewAudioComponent()
 	return PreviewAudioComponent;
 }
 
-UAudioComponent* UEditorEngine::ResetPreviewAudioComponent( USoundBase* Sound, USoundNode* SoundNode )
+UAudioComponent* UEditorEngine::ResetPreviewAudioComponent(USoundBase* Sound, USoundNode* SoundNode)
 {
+	PreviewSoundCue = nullptr;
+
 	if (FAudioDevice* AudioDevice = GetMainAudioDeviceRaw())
 	{
 		if (PreviewAudioComponent)
@@ -2869,15 +2871,47 @@ UAudioComponent* UEditorEngine::ResetPreviewAudioComponent( USoundBase* Sound, U
 		}
 		else
 		{
-			PreviewSoundCue = NewObject<USoundCue>();
-			// Set world to NULL as it will most likely become invalid in the next PIE/Simulate session and the
-			// component will be left with invalid pointer.
-			PreviewAudioComponent = FAudioDevice::CreateComponent(PreviewSoundCue);
+			if (Sound)
+			{
+				PreviewAudioComponent = FAudioDevice::CreateComponent(Sound);
+			}
+			else if (SoundNode)
+			{
+				PreviewSoundCue = NewObject<USoundCue>();
+				PreviewAudioComponent = FAudioDevice::CreateComponent(PreviewSoundCue);
+			}
+			else
+			{
+				PreviewAudioComponent = NewObject<UAudioComponent>();
+			}
 		}
 
 		check(PreviewAudioComponent);
-		// Mark as a preview component so the distance calculations can be ignored
+
+		// Mark as a preview component so the distance calculations is ignored
 		PreviewAudioComponent->bPreviewComponent = true;
+		PreviewAudioComponent->bIsPreviewSound = true;
+		PreviewAudioComponent->bAutoDestroy = false;
+		PreviewAudioComponent->bIsUISound = true;
+		PreviewAudioComponent->bAllowSpatialization = false;
+		PreviewAudioComponent->bReverb = false;
+		PreviewAudioComponent->bCenterChannelOnly = false;
+
+		if (!::IsValid(Sound) || Sound->EditorData.OverrideUseSubmixesOnPreview == EUseSubmixOnPreviewMode::UseEditorPreference)
+		{
+			const UAudioEditorSettings* AudioEditorSettings = GetDefault<UAudioEditorSettings>();
+			PreviewAudioComponent->bIsDryPreviewPlayback = !AudioEditorSettings->bUseSubmixesForPreviewSound;
+		}
+		else
+		{
+			const bool bIsDryPreviewPlayback = Sound->EditorData.OverrideUseSubmixesOnPreview == EUseSubmixOnPreviewMode::PlayPreviewsWithoutSubmixes;
+			PreviewAudioComponent->bIsDryPreviewPlayback = bIsDryPreviewPlayback;
+		}
+
+		if (FAudioDeviceManager* DeviceManager = GEngine->GetAudioDeviceManager())
+		{
+			PreviewAudioComponent->AudioDeviceID = DeviceManager->GetActiveAudioDevice()->DeviceID;
+		}
 
 		if (Sound)
 		{
@@ -2885,6 +2919,7 @@ UAudioComponent* UEditorEngine::ResetPreviewAudioComponent( USoundBase* Sound, U
 		}
 		else if (SoundNode)
 		{
+			check(PreviewSoundCue);
 			PreviewSoundCue->FirstNode = SoundNode;
 			PreviewAudioComponent->Sound = PreviewSoundCue;
 			PreviewSoundCue->CacheAggregateValues();
@@ -2898,30 +2933,7 @@ UAudioComponent* UEditorEngine::PlayPreviewSound(USoundBase* Sound,  USoundNode*
 {
 	if(UAudioComponent* AudioComponent = ResetPreviewAudioComponent(Sound, SoundNode))
 	{
-		AudioComponent->bAutoDestroy = false;
-		AudioComponent->bIsUISound = true;
-		AudioComponent->bAllowSpatialization = false;
-		AudioComponent->bReverb = false;
-		AudioComponent->bCenterChannelOnly = false;
-		AudioComponent->bIsPreviewSound = true;
-
-		if (!::IsValid(Sound) || Sound->EditorData.OverrideUseSubmixesOnPreview == EUseSubmixOnPreviewMode::UseEditorPreference)
-		{
-			const UAudioEditorSettings* AudioEditorSettings = GetDefault<UAudioEditorSettings>();
-			AudioComponent->bIsDryPreviewPlayback = !AudioEditorSettings->bUseSubmixesForPreviewSound;
-		}
-		else
-		{
-			AudioComponent->bIsDryPreviewPlayback = (Sound->EditorData.OverrideUseSubmixesOnPreview == EUseSubmixOnPreviewMode::PlayPreviewsWithoutSubmixes) ? true : false;
-		}
-
-		if (FAudioDeviceManager* DeviceManager = GEngine->GetAudioDeviceManager())
-		{
-			AudioComponent->AudioDeviceID = DeviceManager->GetActiveAudioDevice()->DeviceID;
-		}
-
 		AudioComponent->Play();
-
 		return AudioComponent;
 	}
 
