Metasound size cleanup easy wins.

Various cleanups mostly around global ctors in Metasound, good for about a 1.5MB code size reduction on some targets.

diff --git a/Engine/Plugins/Runtime/Metasound/Source/MetasoundFrontend/Private/MetasoundNodeRegistrationMacro.cpp b/Engine/Plugins/Runtime/Metasound/Source/MetasoundFrontend/Private/MetasoundNodeRegistrationMacro.cpp
--- a/Engine/Plugins/Runtime/Metasound/Source/MetasoundFrontend/Private/MetasoundNodeRegistrationMacro.cpp
+++ b/Engine/Plugins/Runtime/Metasound/Source/MetasoundFrontend/Private/MetasoundNodeRegistrationMacro.cpp
@@ -95,5 +95,31 @@ namespace Metasound::Frontend
 		{
 			return Factory;
 		}
+
+		bool RegisterNodeInternal(TUniquePtr<INodeClassRegistryEntry>&& InEntry)
+		{
+			Frontend::FNodeRegistryKey Key = INodeClassRegistry::Get()->RegisterNode(MoveTemp(InEntry));
+			const bool bSuccessfullyRegisteredNode = Key.IsValid();
+			ensureAlwaysMsgf(bSuccessfullyRegisteredNode, TEXT("Registering node class failed. Please check the logs."));
+
+			return bSuccessfullyRegisteredNode;
+		}
+
+		FORCENOINLINE bool RegisterNodeUsingMacroCheck(bool& bAlreadyRegistered)
+		{
+			if (bAlreadyRegistered)
+			{
+				UE_LOG(LogMetaSound, Display, TEXT("Tried to call METASOUND_REGISTER_NODE twice with the same class. ignoring the second call. Likely because METASOUND_REGISTER_NODE is in a header that's used in multiple modules. Consider moving it to a private header or cpp file."));
+				return false;
+			}
+
+			bAlreadyRegistered = true;
+			return true;
+		}
+
+		FORCENOINLINE bool RegisterNodeDeferred(void (*InRegisterFunc)())
+		{
+			return INodeClassRegistry::Get()->EnqueueInitCommand(InRegisterFunc);
+		}
 	}
 }
diff --git a/Engine/Plugins/Runtime/Metasound/Source/MetasoundFrontend/Public/MetasoundEnum.h b/Engine/Plugins/Runtime/Metasound/Source/MetasoundFrontend/Public/MetasoundEnum.h
--- a/Engine/Plugins/Runtime/Metasound/Source/MetasoundFrontend/Public/MetasoundEnum.h
+++ b/Engine/Plugins/Runtime/Metasound/Source/MetasoundFrontend/Public/MetasoundEnum.h
@@ -42,13 +42,19 @@ namespace Metasound
 	{
 		T Value;
 		FName Name;
-		FText DisplayName;		// TODO: Remove this from runtime.
-		FText Tooltip;			// TODO: Remove this from runtime.
+#if WITH_EDITOR
+		FText DisplayName;
+		FText Tooltip;
+#endif
 
 		// Allow implicit conversion to int32 entry
 		operator TEnumEntry<int32>() const
 		{
+#if WITH_EDITOR
 			return TEnumEntry<int32>{ static_cast<int32>(Value), Name, DisplayName, Tooltip };
+#else
+			return TEnumEntry<int32>{ static_cast<int32>(Value), Name };
+#endif
 		}	
 	};
 
diff --git a/Engine/Plugins/Runtime/Metasound/Source/MetasoundFrontend/Public/MetasoundEnumRegistrationMacro.h b/Engine/Plugins/Runtime/Metasound/Source/MetasoundFrontend/Public/MetasoundEnumRegistrationMacro.h
--- a/Engine/Plugins/Runtime/Metasound/Source/MetasoundFrontend/Public/MetasoundEnumRegistrationMacro.h
+++ b/Engine/Plugins/Runtime/Metasound/Source/MetasoundFrontend/Public/MetasoundEnumRegistrationMacro.h
@@ -68,8 +68,8 @@ namespace Metasound
 	#define DEFINE_METASOUND_ENUM_ENTRY(ENTRY, DISPLAYNAME_KEY, DISPLAYNAME, TOOLTIP_KEY, TOOLTIP) { ENTRY, TEXT(#ENTRY), LOCTEXT(DISPLAYNAME_KEY, DISPLAYNAME), LOCTEXT(TOOLTIP_KEY, TOOLTIP) }
 	#define DEFINE_METASOUND_ENUM_ENTRY_NOTOOLTIP(ENTRY, DISPLAYNAME_KEY, DISPLAYNAME) { ENTRY, TEXT(#ENTRY), LOCTEXT(DISPLAYNAME_KEY, DISPLAYNAME) }
 #else
-	#define DEFINE_METASOUND_ENUM_ENTRY(ENTRY, DISPLAYNAME_KEY, DISPLAYNAME, TOOLTIP_KEY, TOOLTIP) { ENTRY, TEXT(#ENTRY), FText::GetEmpty(), FText::GetEmpty() }
-	#define DEFINE_METASOUND_ENUM_ENTRY_NOTOOLTIP(ENTRY, DISPLAYNAME_KEY, DISPLAYNAME) { ENTRY, TEXT(#ENTRY), FText::GetEmpty() }
+	#define DEFINE_METASOUND_ENUM_ENTRY(ENTRY, DISPLAYNAME_KEY, DISPLAYNAME, TOOLTIP_KEY, TOOLTIP) { ENTRY, TEXT(#ENTRY) }
+	#define DEFINE_METASOUND_ENUM_ENTRY_NOTOOLTIP(ENTRY, DISPLAYNAME_KEY, DISPLAYNAME) { ENTRY, TEXT(#ENTRY) }
 #endif // WITH_EDITOR
 
 /** DEFINE_METASOUND_ENUM_END - macro which ends the function body of the GetAllEntries function
@@ -77,4 +77,4 @@ namespace Metasound
 #define DEFINE_METASOUND_ENUM_END() \
 		};\
 		return Entries;\
-	};
\ No newline at end of file
+	};
diff --git a/Engine/Plugins/Runtime/Metasound/Source/MetasoundFrontend/Public/MetasoundNodeRegistrationMacro.h b/Engine/Plugins/Runtime/Metasound/Source/MetasoundFrontend/Public/MetasoundNodeRegistrationMacro.h
--- a/Engine/Plugins/Runtime/Metasound/Source/MetasoundFrontend/Public/MetasoundNodeRegistrationMacro.h
+++ b/Engine/Plugins/Runtime/Metasound/Source/MetasoundFrontend/Public/MetasoundNodeRegistrationMacro.h
@@ -171,6 +171,8 @@ namespace Metasound::Frontend
 				}
 			}
 		};
+
+		UE_API bool RegisterNodeInternal(TUniquePtr<INodeClassRegistryEntry>&& InEntry);
 	} // namespace NodeRegistrationPrivate
 
 
@@ -184,18 +186,14 @@ namespace Metasound::Frontend
 		}
 		else
 		{
-			Frontend::FNodeRegistryKey Key = INodeClassRegistry::Get()->RegisterNode(MakeUnique<TNodeRegistryEntry<TNodeType, ConfigurationType>>(InMetadata));
-			const bool bSuccessfullyRegisteredNode = Key.IsValid();
-			ensureAlwaysMsgf(bSuccessfullyRegisteredNode, TEXT("Registering node class failed. Please check the logs."));
-
-			return bSuccessfullyRegisteredNode;
+			return NodeRegistrationPrivate::RegisterNodeInternal(MakeUnique<TNodeRegistryEntry<TNodeType, ConfigurationType>>(InMetadata));
 		}
 	} 
 
 	template <typename TNodeType, typename ConfigurationType=void>
 	bool RegisterNode()
 	{
-		 static_assert(::Metasound::TIsNodeConstructorSupported<TNodeType>::Value, "In order to be registered as a MetaSound node, the node needs to implement the following public constructor: Constructor(Metasound::FNodeData InNodeData) or Construct(Metasound::FNodeData InNodeData, TSharedRef<const FNodeClassMetadata> InMetadata)"); 
+		static_assert(::Metasound::TIsNodeConstructorSupported<TNodeType>::Value, "In order to be registered as a MetaSound node, the node needs to implement the following public constructor: Constructor(Metasound::FNodeData InNodeData) or Construct(Metasound::FNodeData InNodeData, TSharedRef<const FNodeClassMetadata> InMetadata)"); 
 
 		if constexpr(TIsCreateNodeClassMetadataDeclared<TNodeType>::Value)
 		{
@@ -235,21 +233,21 @@ namespace Metasound::Frontend
 
 	namespace NodeRegistrationPrivate
 	{
+		UE_API bool RegisterNodeUsingMacroCheck(bool& bAlreadyRegistered);
+		UE_API bool RegisterNodeDeferred(void (*InRegisterFunc)());
+
 		template<typename NodeType, typename ConfigurationType=void>
-		bool RegisterNodeUsingMacro()
+		void RegisterNodeUsingMacro()
 		{
 			// if we reenter this code (because METASOUND_REGISTER_NODE was called twice with the same node type),
 			// we catch it here.
 			static bool bAlreadyRegisteredThisDataType = false;
-			if (bAlreadyRegisteredThisDataType)
+			if (!RegisterNodeUsingMacroCheck(bAlreadyRegisteredThisDataType))
 			{
-				UE_LOG(LogMetaSound, Display, TEXT("Tried to call METASOUND_REGISTER_NODE twice with the same class. ignoring the second call. Likely because METASOUND_REGISTER_NODE is in a header that's used in multiple modules. Consider moving it to a private header or cpp file."));
-				return false;
+				return;
 			}
 
-			bAlreadyRegisteredThisDataType = true;
-
-			return RegisterNode<NodeType, ConfigurationType>();
+			RegisterNode<NodeType, ConfigurationType>();
 		}
 	} // namespace NodeRegistrationPrivate
 }
@@ -289,7 +287,9 @@ namespace Metasound
 #define METASOUND_REGISTER_NODE_AND_CONFIGURATION(NodeClass, ConfigurationClass) \
 	 static_assert(std::is_base_of<::Metasound::INodeBase, NodeClass>::value, "To be registered as a  Metasound Node," #NodeClass "need to be a derived class from Metasound::INodeBase, Metasound::INode, or Metasound::FNode."); \
 	 static_assert(::Metasound::TIsNodeConstructorSupported<NodeClass>::Value, "In order to be registered as a Metasound Node, " #NodeClass " needs to implement the following public constructor: " #NodeClass "(Metasound::FNodeData InNodeData);"); \
-	 static bool bSuccessfullyRegistered##NodeClass  = ::Metasound::Frontend::INodeClassRegistry::Get()->EnqueueInitCommand([](){ ::Metasound::Frontend::NodeRegistrationPrivate::RegisterNodeUsingMacro<NodeClass, ConfigurationClass>(); }); // This static bool is useful for debugging, but also is the only way the compiler will let us call this function outside of an expression.
+	 static bool bSuccessfullyRegistered##NodeClass = \
+		/* This static bool is useful for debugging, but also is the only way the compiler will let us call this function outside of an expression. */ \
+		::Metasound::Frontend::NodeRegistrationPrivate::RegisterNodeDeferred(::Metasound::Frontend::NodeRegistrationPrivate::RegisterNodeUsingMacro<NodeClass, ConfigurationClass>);
 
 #define METASOUND_REGISTER_NODE(NodeClass) METASOUND_REGISTER_NODE_AND_CONFIGURATION(NodeClass, void)
 /*
@@ -299,8 +299,8 @@ Macros to help define various FText node fields.
 #define METASOUND_LOCTEXT(KEY, NAME_TEXT) LOCTEXT(KEY, NAME_TEXT)
 #define METASOUND_LOCTEXT_FORMAT(KEY, NAME_TEXT, ...) FText::Format(LOCTEXT(KEY, NAME_TEXT), __VA_ARGS__)
 #else 
-#define METASOUND_LOCTEXT(KEY, NAME_TEXT) FText::GetEmpty()
-#define METASOUND_LOCTEXT_FORMAT(KEY, NAME_TEXT, ...) FText::GetEmpty()
+#define METASOUND_LOCTEXT(KEY, NAME_TEXT) FText{}
+#define METASOUND_LOCTEXT_FORMAT(KEY, NAME_TEXT, ...) FText{}
 #endif // WITH_EDITOR
 
 #undef UE_API
diff --git a/Engine/Plugins/Runtime/Metasound/Source/MetasoundGraphCore/Private/MetasoundNodeInterface.cpp b/Engine/Plugins/Runtime/Metasound/Source/MetasoundGraphCore/Private/MetasoundNodeInterface.cpp
--- a/Engine/Plugins/Runtime/Metasound/Source/MetasoundGraphCore/Private/MetasoundNodeInterface.cpp
+++ b/Engine/Plugins/Runtime/Metasound/Source/MetasoundGraphCore/Private/MetasoundNodeInterface.cpp
@@ -9,6 +9,10 @@
 
 namespace Metasound
 {
+	FNodeData::FNodeData() = default;
+
+	FNodeData::FNodeData(const FNodeData&) = default;
+
 	FNodeData::FNodeData(const FName& InName, const FGuid& InID, FVertexInterface InInterface)
 	: FNodeData(InName, InID, MoveTemp(InInterface), TSharedPtr<const IOperatorData>())
 	{
@@ -22,6 +26,8 @@ namespace Metasound
 	{
 	}
 
+	FNodeData::~FNodeData() = default;
+
 	const FString PluginAuthor = TEXT("Epic Games, Inc.");
 #if WITH_EDITOR
 	const FText PluginNodeMissingPrompt = NSLOCTEXT("MetasoundGraphCore", "Metasound_DefaultMissingNodePrompt", "The node was likely removed, renamed, or the Metasound plugin is not loaded.");
@@ -105,6 +111,8 @@ namespace Metasound
 		return *this != InvalidNodeClassName;
 	}
 
+	FNodeClassMetadata::~FNodeClassMetadata() = default;
+
 	const FNodeClassMetadata& FNodeClassMetadata::GetEmpty()
 	{
 		static const FNodeClassMetadata EmptyInfo;
diff --git a/Engine/Plugins/Runtime/Metasound/Source/MetasoundGraphCore/Private/MetasoundVertex.cpp b/Engine/Plugins/Runtime/Metasound/Source/MetasoundGraphCore/Private/MetasoundVertex.cpp
--- a/Engine/Plugins/Runtime/Metasound/Source/MetasoundGraphCore/Private/MetasoundVertex.cpp
+++ b/Engine/Plugins/Runtime/Metasound/Source/MetasoundGraphCore/Private/MetasoundVertex.cpp
@@ -359,6 +359,9 @@ PRAGMA_ENABLE_EXPERIMENTAL_WARNINGS
 		};
 	}
 
+	FDataVertexMetadata FDataVertexMetadata::EmptyBasic { FText{}, FText{}, false };
+	FDataVertexMetadata FDataVertexMetadata::EmptyAdvanced { FText{}, FText{}, true };
+
 	bool operator==(const FInputDataVertex& InLHS, const FInputDataVertex& InRHS)
 	{
 		return (InLHS.VertexName == InRHS.VertexName) && (InLHS.DataTypeName == InRHS.DataTypeName);
@@ -418,6 +421,8 @@ PRAGMA_ENABLE_EXPERIMENTAL_WARNINGS
 		return InLHS.VertexName.FastLess(InRHS.VertexName);
 	}
 
+	FInputVertexInterface::FInputVertexInterface() = default;
+
 	FInputVertexInterface::FInputVertexInterface(TArray<FInputDataVertex> InVertices, TArray<VertexPrivate::FSubInterfaceLayout> InSubInterfaces)
 	: TVertexInterfaceImpl<FInputDataVertex>(MoveTemp(InVertices))
 	, SubInterfaces(MoveTemp(InSubInterfaces))
@@ -464,6 +469,8 @@ PRAGMA_ENABLE_EXPERIMENTAL_WARNINGS
 		return SubInterfaces.FindByPredicate(VertexPrivate::FEqualSubInterfaceName(InName));
 	}
 
+	FOutputVertexInterface::FOutputVertexInterface() = default;
+
 	FOutputVertexInterface::FOutputVertexInterface(TArray<FOutputDataVertex> InVertices, TArray<VertexPrivate::FSubInterfaceLayout> InSubInterfaces)
 	: TVertexInterfaceImpl<FOutputDataVertex>(MoveTemp(InVertices))
 	, SubInterfaces(MoveTemp(InSubInterfaces))
@@ -513,6 +520,8 @@ PRAGMA_ENABLE_EXPERIMENTAL_WARNINGS
 	{
 	}
 
+	FVertexInterface::FVertexInterface() = default;
+
 	FVertexInterface::FVertexInterface(const FInputVertexInterface& InInputs, const FOutputVertexInterface& InOutputs)
 	:	InputInterface(InInputs)
 	,	OutputInterface(InOutputs)
@@ -526,6 +535,8 @@ PRAGMA_ENABLE_EXPERIMENTAL_WARNINGS
 	{
 	}
 
+	FVertexInterface::~FVertexInterface() = default;
+
 	const FInputVertexInterface& FVertexInterface::GetInputInterface() const
 	{
 		return InputInterface;
diff --git a/Engine/Plugins/Runtime/Metasound/Source/MetasoundGraphCore/Public/MetasoundNodeInterface.h b/Engine/Plugins/Runtime/Metasound/Source/MetasoundGraphCore/Public/MetasoundNodeInterface.h
--- a/Engine/Plugins/Runtime/Metasound/Source/MetasoundGraphCore/Public/MetasoundNodeInterface.h
+++ b/Engine/Plugins/Runtime/Metasound/Source/MetasoundGraphCore/Public/MetasoundNodeInterface.h
@@ -38,9 +38,11 @@ namespace Metasound
 	/** Data used to construct a Node. */
 	struct FNodeData
 	{
-		FNodeData() = default;
+		UE_API FNodeData();
+		UE_API FNodeData(const FNodeData&);
 		UE_API FNodeData(const FName& InName, const FGuid& InID, FVertexInterface InInterface);
 		UE_API FNodeData(const FName& InName, const FGuid& InID, FVertexInterface InInterface, TSharedPtr<const IOperatorData> InOperatorData);
+		UE_API ~FNodeData();
 
 		FName Name;
 		FGuid ID;
@@ -161,6 +163,9 @@ namespace Metasound
 		/** If true, the node is deprecated and should not be used in new MetaSounds. */
 		bool bDeprecated = false;
 
+		/** Destructor. */
+		METASOUNDGRAPHCORE_API ~FNodeClassMetadata();
+
 		/** Returns an empty FNodeClassMetadata object. */
 		static METASOUNDGRAPHCORE_API const FNodeClassMetadata& GetEmpty();
 	};
diff --git a/Engine/Plugins/Runtime/Metasound/Source/MetasoundGraphCore/Public/MetasoundParamHelper.h b/Engine/Plugins/Runtime/Metasound/Source/MetasoundGraphCore/Public/MetasoundParamHelper.h
--- a/Engine/Plugins/Runtime/Metasound/Source/MetasoundGraphCore/Public/MetasoundParamHelper.h
+++ b/Engine/Plugins/Runtime/Metasound/Source/MetasoundGraphCore/Public/MetasoundParamHelper.h
@@ -86,11 +86,11 @@
 #define METASOUND_GET_PARAM_METADATA_ADVANCED(NAME) FDataVertexMetadata { NAME##Tooltip, NAME##DisplayName, true }
 #define METASOUND_GET_PARAM_DISPLAYNAME(NAME) NAME##DisplayName
 #else
-#define METASOUND_GET_PARAM_TT(NAME) FText::GetEmpty()
-#define METASOUND_GET_PARAM_NAME_AND_TT(NAME) NAME##Name, FText::GetEmpty()
-#define METASOUND_GET_PARAM_METADATA(NAME) FDataVertexMetadata { FText::GetEmpty(), FText::GetEmpty() }
-#define METASOUND_GET_PARAM_METADATA_ADVANCED(NAME) FDataVertexMetadata { FText::GetEmpty(), FText::GetEmpty(), true }
-#define METASOUND_GET_PARAM_DISPLAYNAME(NAME) FText::GetEmpty()
+#define METASOUND_GET_PARAM_TT(NAME) FText{}
+#define METASOUND_GET_PARAM_NAME_AND_TT(NAME) NAME##Name, FText{}
+#define METASOUND_GET_PARAM_METADATA(NAME) FDataVertexMetadata::EmptyBasic
+#define METASOUND_GET_PARAM_METADATA_ADVANCED(NAME) FDataVertexMetadata::EmptyAdvanced
+#define METASOUND_GET_PARAM_DISPLAYNAME(NAME) FText{}
 #endif // WITH_EDITOR
 #define METASOUND_GET_PARAM_NAME_AND_METADATA(NAME) METASOUND_GET_PARAM_NAME(NAME), METASOUND_GET_PARAM_METADATA(NAME)
 
@@ -102,8 +102,8 @@
 #define METASOUND_GET_PARAM_TT_WITH_INDEX(NAME, INDEX)  FText::Format(NAME##Tooltip, INDEX)
 #define METASOUND_GET_PARAM_NAME_WITH_INDEX_AND_METADATA(NAME, INDEX)  METASOUND_GET_PARAM_NAME_WITH_INDEX(NAME, INDEX), FDataVertexMetadata {FText::Format(NAME##Tooltip, INDEX)}
 #else 
-#define METASOUND_GET_PARAM_NAME_WITH_INDEX_AND_TT(NAME, INDEX)  *FString::Format(*static_cast<FName>(NAME##Name).ToString(), {INDEX}), FText::GetEmpty()
-#define METASOUND_GET_PARAM_TT_WITH_INDEX(NAME, INDEX)  FText::GetEmpty();
-#define METASOUND_GET_PARAM_NAME_WITH_INDEX_AND_METADATA(NAME, INDEX)  *FString::Format(*static_cast<FName>(NAME##Name).ToString(), {INDEX}), FDataVertexMetadata {FText::GetEmpty()}
+#define METASOUND_GET_PARAM_NAME_WITH_INDEX_AND_TT(NAME, INDEX)  *FString::Format(*static_cast<FName>(NAME##Name).ToString(), {INDEX}), FText{}
+#define METASOUND_GET_PARAM_TT_WITH_INDEX(NAME, INDEX)  FText{}
+#define METASOUND_GET_PARAM_NAME_WITH_INDEX_AND_METADATA(NAME, INDEX)  *FString::Format(*static_cast<FName>(NAME##Name).ToString(), {INDEX}), FDataVertexMetadata::EmptyBasic
 #endif // WITH_EDITOR
 
diff --git a/Engine/Plugins/Runtime/Metasound/Source/MetasoundGraphCore/Public/MetasoundVertex.h b/Engine/Plugins/Runtime/Metasound/Source/MetasoundGraphCore/Public/MetasoundVertex.h
--- a/Engine/Plugins/Runtime/Metasound/Source/MetasoundGraphCore/Public/MetasoundVertex.h
+++ b/Engine/Plugins/Runtime/Metasound/Source/MetasoundGraphCore/Public/MetasoundVertex.h
@@ -75,6 +75,9 @@ namespace Metasound
 		FText Description;
 		FText DisplayName;
 		bool bIsAdvancedDisplay = false;
+
+		METASOUNDGRAPHCORE_API static FDataVertexMetadata EmptyBasic;
+		METASOUNDGRAPHCORE_API static FDataVertexMetadata EmptyAdvanced;
 	};
 
 	/** Describe how the vertex will access connected data. */
@@ -188,7 +191,19 @@ namespace Metasound
 		TInputDataVertex() = default;
 
 		template<typename... ArgTypes>
-		TInputDataVertex(const FVertexName& InVertexName, const FDataVertexMetadata& InMetadata, ArgTypes&&... Args)
+		FORCENOINLINE TInputDataVertex(const FVertexName& InVertexName, const FDataVertexMetadata& InMetadata, ArgTypes&&... Args)
+		: FInputDataVertex(InVertexName, GetMetasoundDataTypeName<DataType>(), InMetadata, EVertexAccessType::Reference, Forward<ArgTypes>(Args)...)
+		{
+		}
+
+		template<typename... ArgTypes>
+		FORCENOINLINE TInputDataVertex(const FLazyName& InVertexName, const FDataVertexMetadata& InMetadata, ArgTypes&&... Args)
+		: FInputDataVertex(InVertexName, GetMetasoundDataTypeName<DataType>(), InMetadata, EVertexAccessType::Reference, Forward<ArgTypes>(Args)...)
+		{
+		}
+
+		template<typename... ArgTypes>
+		FORCENOINLINE TInputDataVertex(const char* InVertexName, const FDataVertexMetadata& InMetadata, ArgTypes&&... Args)
 		: FInputDataVertex(InVertexName, GetMetasoundDataTypeName<DataType>(), InMetadata, EVertexAccessType::Reference, Forward<ArgTypes>(Args)...)
 		{
 		}
@@ -232,7 +247,19 @@ namespace Metasound
 		TOutputDataVertex() = default;
 
 		template<typename... ArgTypes>
-		TOutputDataVertex(const FVertexName& InVertexName, const FDataVertexMetadata& InMetadata, ArgTypes&&... Args)
+		FORCENOINLINE TOutputDataVertex(const FVertexName& InVertexName, const FDataVertexMetadata& InMetadata, ArgTypes&&... Args)
+		: FOutputDataVertex(InVertexName, GetMetasoundDataTypeName<DataType>(), InMetadata, EVertexAccessType::Reference, Forward<ArgTypes>(Args)...)
+		{
+		}
+
+		template<typename... ArgTypes>
+		FORCENOINLINE TOutputDataVertex(const FLazyName& InVertexName, const FDataVertexMetadata& InMetadata, ArgTypes&&... Args)
+		: FOutputDataVertex(InVertexName, GetMetasoundDataTypeName<DataType>(), InMetadata, EVertexAccessType::Reference, Forward<ArgTypes>(Args)...)
+		{
+		}
+
+		template<typename... ArgTypes>
+		FORCENOINLINE TOutputDataVertex(const char* InVertexName, const FDataVertexMetadata& InMetadata, ArgTypes&&... Args)
 		: FOutputDataVertex(InVertexName, GetMetasoundDataTypeName<DataType>(), InMetadata, EVertexAccessType::Reference, Forward<ArgTypes>(Args)...)
 		{
 		}
@@ -641,7 +668,7 @@ PRAGMA_ENABLE_EXPERIMENTAL_WARNINGS
 	class FInputVertexInterface : public TVertexInterfaceImpl<FInputDataVertex>
 	{
 	public:
-		FInputVertexInterface() = default;
+		UE_API FInputVertexInterface();
 
 		/** Construct an FInputVertexInterface from a parameter pack. This allows 
 		 * node interfaces to be declared in a single function.
@@ -683,7 +710,7 @@ PRAGMA_ENABLE_EXPERIMENTAL_WARNINGS
 	class FOutputVertexInterface : public TVertexInterfaceImpl<FOutputDataVertex>
 	{
 	public:
-		FOutputVertexInterface() = default;
+		UE_API FOutputVertexInterface();
 
 		/** Construct an FOutputVertexInterface from a parameter pack. This allows 
 		 * node interfaces to be declared in a single function.
@@ -701,6 +728,7 @@ PRAGMA_ENABLE_EXPERIMENTAL_WARNINGS
 		 */
 		template<typename... ArgTypes>
 		explicit FOutputVertexInterface(ArgTypes && ... InArgs)
+		: FOutputVertexInterface()
 		{
 			VertexPrivate::FOutputVertexInterfaceDeclarationBuilder InterfaceBuilder(Vertices, SubInterfaces);
 			InterfaceBuilder.Build(Forward<ArgTypes>(InArgs)...);
@@ -748,7 +776,7 @@ PRAGMA_ENABLE_EXPERIMENTAL_WARNINGS
 		public:
 
 			/** Default constructor. */
-			FVertexInterface() = default;
+			UE_API FVertexInterface();
 
 			/** Construct with an input and output interface. */
 			UE_API FVertexInterface(const FInputVertexInterface& InInputs, const FOutputVertexInterface& InOutputs);
@@ -756,6 +784,9 @@ PRAGMA_ENABLE_EXPERIMENTAL_WARNINGS
 			/** Construct with input, output and environment interface. */
 			UE_API FVertexInterface(FInputVertexInterface InInputs, FOutputVertexInterface InOutputs, FEnvironmentVertexInterface InEnvironmentVariables);
 
+			/** Destructor. */
+			UE_API ~FVertexInterface();
+
 			/** Return the input interface. */
 			UE_API const FInputVertexInterface& GetInputInterface() const;
 
