Added delegates for various MetaSound document and member metadata property changed events

#rnx

diff --git a/Engine/Plugins/Runtime/Metasound/Source/MetasoundEditor/Private/MetasoundDetailCustomization.cpp b/Engine/Plugins/Runtime/Metasound/Source/MetasoundEditor/Private/MetasoundDetailCustomization.cpp
--- a/Engine/Plugins/Runtime/Metasound/Source/MetasoundEditor/Private/MetasoundDetailCustomization.cpp
+++ b/Engine/Plugins/Runtime/Metasound/Source/MetasoundEditor/Private/MetasoundDetailCustomization.cpp
@@ -243,6 +243,8 @@ namespace Metasound::Editor
 							if (DocBuilder.IsValid())
 							{
 								DocBuilder.GetConstDocumentChecked().RootGraph.Style.UpdateChangeID();
+								const FDocumentModifyDelegates& DocumentDelegates = DocBuilder.GetDocumentDelegates();
+								DocumentDelegates.OnDocumentMetadataChanged.Broadcast();
 							}
 
 							{
@@ -252,11 +254,13 @@ namespace Metasound::Editor
 							}
 						}
 					});
-					AuthorHandle->SetOnChildPropertyValueChanged(RegisterOnChange);
+					AuthorHandle->SetOnPropertyValueChanged(RegisterOnChange);
 					DescHandle->SetOnPropertyValueChanged(RegisterOnChange);
 					DisplayNameHandle->SetOnPropertyValueChanged(RegisterOnChange);
 					KeywordsHandle->SetOnPropertyValueChanged(RegisterOnChange);
 					KeywordsHandle->SetOnChildPropertyValueChanged(RegisterOnChange);
+					CategoryHierarchyHandle->SetOnPropertyValueChanged(RegisterOnChange);
+					CategoryHierarchyHandle->SetOnChildPropertyValueChanged(RegisterOnChange);
 					IsDeprecatedHandle->SetOnPropertyValueChanged(RegisterOnChange);
 				}
 
diff --git a/Engine/Plugins/Runtime/Metasound/Source/MetasoundEditor/Private/MetasoundEditorBuilderListener.cpp b/Engine/Plugins/Runtime/Metasound/Source/MetasoundEditor/Private/MetasoundEditorBuilderListener.cpp
--- a/Engine/Plugins/Runtime/Metasound/Source/MetasoundEditor/Private/MetasoundEditorBuilderListener.cpp
+++ b/Engine/Plugins/Runtime/Metasound/Source/MetasoundEditor/Private/MetasoundEditorBuilderListener.cpp
@@ -1,4 +1,5 @@
 // Copyright Epic Games, Inc. All Rights Reserved.
+
 #include "MetasoundEditorBuilderListener.h"
 
 #include UE_INLINE_GENERATED_CPP_BY_NAME(MetasoundEditorBuilderListener)
@@ -15,6 +16,78 @@ void UMetaSoundEditorBuilderListener::Init(const TWeakObjectPtr<UMetaSoundBuilde
 	}
 }
 
+void UMetaSoundEditorBuilderListener::OnDocumentDisplayNameChanged()
+{
+	if (Builder.IsValid())
+	{
+		const FMetaSoundFrontendDocumentBuilder& ConstBuilder = Builder->GetConstBuilder();
+		const FMetasoundFrontendGraphClass& GraphClass = ConstBuilder.GetConstDocumentChecked().RootGraph;
+		const FText& DisplayName = GraphClass.Metadata.GetDisplayName();
+
+		OnDocumentDisplayNameChangedDelegate.Broadcast(DisplayName);
+	}
+}
+
+void UMetaSoundEditorBuilderListener::OnDocumentDescriptionChanged()
+{
+	if (Builder.IsValid())
+	{
+		const FMetaSoundFrontendDocumentBuilder& ConstBuilder = Builder->GetConstBuilder();
+		const FMetasoundFrontendGraphClass& GraphClass = ConstBuilder.GetConstDocumentChecked().RootGraph;
+		const FText& Description = GraphClass.Metadata.GetDescription();
+
+		OnDocumentDescriptionChangedDelegate.Broadcast(Description);
+	}
+}
+
+void UMetaSoundEditorBuilderListener::OnDocumentAuthorChanged()
+{
+	if (Builder.IsValid())
+	{
+		const FMetaSoundFrontendDocumentBuilder& ConstBuilder = Builder->GetConstBuilder();
+		const FMetasoundFrontendGraphClass& GraphClass = ConstBuilder.GetConstDocumentChecked().RootGraph;
+		const FString& Author = GraphClass.Metadata.GetAuthor();
+
+		OnDocumentAuthorChangedDelegate.Broadcast(Author);
+	}
+}
+
+void UMetaSoundEditorBuilderListener::OnDocumentKeywordsChanged()
+{
+	if (Builder.IsValid())
+	{
+		const FMetaSoundFrontendDocumentBuilder& ConstBuilder = Builder->GetConstBuilder();
+		const FMetasoundFrontendGraphClass& GraphClass = ConstBuilder.GetConstDocumentChecked().RootGraph;
+		const TArray<FText>& Keywords = GraphClass.Metadata.GetKeywords();
+
+		OnDocumentKeywordsChangedDelegate.Broadcast(Keywords);
+	}
+}
+
+void UMetaSoundEditorBuilderListener::OnDocumentCategoryHierarchyChanged()
+{
+	if (Builder.IsValid())
+	{
+		const FMetaSoundFrontendDocumentBuilder& ConstBuilder = Builder->GetConstBuilder();
+		const FMetasoundFrontendGraphClass& GraphClass = ConstBuilder.GetConstDocumentChecked().RootGraph;
+		const TArray<FText>& CategoryHierarchy = GraphClass.Metadata.GetCategoryHierarchy();
+
+		OnDocumentCategoryHierarchyChangedDelegate.Broadcast(CategoryHierarchy);
+	}
+}
+
+void UMetaSoundEditorBuilderListener::OnDocumentIsDeprecatedChanged()
+{
+	if (Builder.IsValid())
+	{
+		const FMetaSoundFrontendDocumentBuilder& ConstBuilder = Builder->GetConstBuilder();
+		const FMetasoundFrontendGraphClass& GraphClass = ConstBuilder.GetConstDocumentChecked().RootGraph;
+		const bool bIsDeprecated = GraphClass.Metadata.GetIsDeprecated();
+
+		OnDocumentIsDeprecatedChangedDelegate.Broadcast(bIsDeprecated);
+	}
+}
+
 void UMetaSoundEditorBuilderListener::OnGraphInputAdded(int32 Index)
 {
 	if (Builder.IsValid())
@@ -47,13 +120,26 @@ void UMetaSoundEditorBuilderListener::OnGraphInputDefaultChanged(int32 Index)
 	}
 }
 
+void UMetaSoundEditorBuilderListener::OnGraphInputInheritsDefaultChanged(int32 Index)
+{
+	if (Builder.IsValid())
+	{
+		const FMetaSoundFrontendDocumentBuilder& ConstBuilder = Builder->GetConstBuilder();
+		const FMetasoundFrontendGraphClass& GraphClass = ConstBuilder.GetConstDocumentChecked().RootGraph;
+		const FMetasoundFrontendClassInput& GraphInput = GraphClass.GetDefaultInterface().Inputs[Index];
+		const bool bInheritsDefault = GraphClass.PresetOptions.InputsInheritingDefault.Contains(GraphInput.Name);
+
+		OnGraphInputInheritsDefaultChangedDelegate.Broadcast(GraphInput.Name, bInheritsDefault);
+	}
+}
+
 void UMetaSoundEditorBuilderListener::OnRemovingGraphInput(int32 Index)
 {
 	if (Builder.IsValid())
 	{
 		const FMetaSoundFrontendDocumentBuilder& ConstBuilder = Builder->GetConstBuilder();
 		const FMetasoundFrontendGraphClass& GraphClass = ConstBuilder.GetConstDocumentChecked().RootGraph;
-		const FMetasoundFrontendClassOutput& GraphInput = GraphClass.GetDefaultInterface().Inputs[Index];
+		const FMetasoundFrontendClassInput& GraphInput = GraphClass.GetDefaultInterface().Inputs[Index];
 
 		OnRemovingGraphInputDelegate.Broadcast(GraphInput.Name, GraphInput.TypeName);
 	}
@@ -83,19 +169,78 @@ void UMetaSoundEditorBuilderListener::OnRemovingGraphOutput(int32 Index)
 	}
 }
 
-
 void UMetaSoundEditorBuilderListener::OnGraphInputDataTypeChanged(int32 Index)
 {
 	if (Builder.IsValid())
 	{
 		const FMetaSoundFrontendDocumentBuilder& ConstBuilder = Builder->GetConstBuilder();
 		const FMetasoundFrontendGraphClass& GraphClass = ConstBuilder.GetConstDocumentChecked().RootGraph;
-		const FMetasoundFrontendClassOutput& GraphInput = GraphClass.GetDefaultInterface().Inputs[Index];
+		const FMetasoundFrontendClassInput& GraphInput = GraphClass.GetDefaultInterface().Inputs[Index];
 
 		OnGraphInputDataTypeChangedDelegate.Broadcast(GraphInput.Name, GraphInput.TypeName);
 	}
 }
 
+void UMetaSoundEditorBuilderListener::OnGraphInputDisplayNameChanged(int32 Index)
+{
+	if (Builder.IsValid())
+	{
+		const FMetaSoundFrontendDocumentBuilder& ConstBuilder = Builder->GetConstBuilder();
+		const FMetasoundFrontendGraphClass& GraphClass = ConstBuilder.GetConstDocumentChecked().RootGraph;
+		const FMetasoundFrontendClassInput& GraphInput = GraphClass.GetDefaultInterface().Inputs[Index];
+
+		OnGraphInputDisplayNameChangedDelegate.Broadcast(GraphInput.Name, GraphInput.Metadata.GetDisplayName());
+	}
+}
+
+void UMetaSoundEditorBuilderListener::OnGraphInputDescriptionChanged(int32 Index)
+{
+	if (Builder.IsValid())
+	{
+		const FMetaSoundFrontendDocumentBuilder& ConstBuilder = Builder->GetConstBuilder();
+		const FMetasoundFrontendGraphClass& GraphClass = ConstBuilder.GetConstDocumentChecked().RootGraph;
+		const FMetasoundFrontendClassInput& GraphInput = GraphClass.GetDefaultInterface().Inputs[Index];
+
+		OnGraphInputDescriptionChangedDelegate.Broadcast(GraphInput.Name, GraphInput.Metadata.GetDescription());
+	}
+}
+
+void UMetaSoundEditorBuilderListener::OnGraphInputSortOrderIndexChanged(int32 Index)
+{
+	if (Builder.IsValid())
+	{
+		const FMetaSoundFrontendDocumentBuilder& ConstBuilder = Builder->GetConstBuilder();
+		const FMetasoundFrontendGraphClass& GraphClass = ConstBuilder.GetConstDocumentChecked().RootGraph;
+		const FMetasoundFrontendClassInput& GraphInput = GraphClass.GetDefaultInterface().Inputs[Index];
+
+		OnGraphInputSortOrderIndexChangedDelegate.Broadcast(GraphInput.Name, GraphInput.Metadata.SortOrderIndex);
+	}
+}
+
+void UMetaSoundEditorBuilderListener::OnGraphInputIsConstructorPinChanged(int32 Index)
+{
+	if (Builder.IsValid())
+	{
+		const FMetaSoundFrontendDocumentBuilder& ConstBuilder = Builder->GetConstBuilder();
+		const FMetasoundFrontendGraphClass& GraphClass = ConstBuilder.GetConstDocumentChecked().RootGraph;
+		const FMetasoundFrontendClassInput& GraphInput = GraphClass.GetDefaultInterface().Inputs[Index];
+
+		OnGraphInputIsConstructorPinChangedDelegate.Broadcast(GraphInput.Name, GraphInput.AccessType == EMetasoundFrontendVertexAccessType::Value);
+	}
+}
+
+void UMetaSoundEditorBuilderListener::OnGraphInputIsAdvancedDisplayChanged(int32 Index)
+{
+	if (Builder.IsValid())
+	{
+		const FMetaSoundFrontendDocumentBuilder& ConstBuilder = Builder->GetConstBuilder();
+		const FMetasoundFrontendGraphClass& GraphClass = ConstBuilder.GetConstDocumentChecked().RootGraph;
+		const FMetasoundFrontendClassInput& GraphInput = GraphClass.GetDefaultInterface().Inputs[Index];
+
+		OnGraphInputIsAdvancedDisplayChangedDelegate.Broadcast(GraphInput.Name, GraphInput.Metadata.bIsAdvancedDisplay);
+	}
+}
+
 void UMetaSoundEditorBuilderListener::OnGraphOutputDataTypeChanged(int32 Index)
 {
 	if (Builder.IsValid())
@@ -108,6 +253,66 @@ void UMetaSoundEditorBuilderListener::OnGraphOutputDataTypeChanged(int32 Index)
 	}
 }
 
+void UMetaSoundEditorBuilderListener::OnGraphOutputDisplayNameChanged(int32 Index)
+{
+	if (Builder.IsValid())
+	{
+		const FMetaSoundFrontendDocumentBuilder& ConstBuilder = Builder->GetConstBuilder();
+		const FMetasoundFrontendGraphClass& GraphClass = ConstBuilder.GetConstDocumentChecked().RootGraph;
+		const FMetasoundFrontendClassOutput& GraphOutput = GraphClass.GetDefaultInterface().Outputs[Index];
+
+		OnGraphOutputDisplayNameChangedDelegate.Broadcast(GraphOutput.Name, GraphOutput.Metadata.GetDisplayName());
+	}
+}
+
+void UMetaSoundEditorBuilderListener::OnGraphOutputDescriptionChanged(int32 Index)
+{
+	if (Builder.IsValid())
+	{
+		const FMetaSoundFrontendDocumentBuilder& ConstBuilder = Builder->GetConstBuilder();
+		const FMetasoundFrontendGraphClass& GraphClass = ConstBuilder.GetConstDocumentChecked().RootGraph;
+		const FMetasoundFrontendClassOutput& GraphOutput = GraphClass.GetDefaultInterface().Outputs[Index];
+
+		OnGraphOutputDescriptionChangedDelegate.Broadcast(GraphOutput.Name, GraphOutput.Metadata.GetDescription());
+	}
+}
+
+void UMetaSoundEditorBuilderListener::OnGraphOutputSortOrderIndexChanged(int32 Index)
+{
+	if (Builder.IsValid())
+	{
+		const FMetaSoundFrontendDocumentBuilder& ConstBuilder = Builder->GetConstBuilder();
+		const FMetasoundFrontendGraphClass& GraphClass = ConstBuilder.GetConstDocumentChecked().RootGraph;
+		const FMetasoundFrontendClassOutput& GraphOutput = GraphClass.GetDefaultInterface().Outputs[Index];
+
+		OnGraphOutputSortOrderIndexChangedDelegate.Broadcast(GraphOutput.Name, GraphOutput.Metadata.SortOrderIndex);
+	}
+}
+
+void UMetaSoundEditorBuilderListener::OnGraphOutputIsConstructorPinChanged(int32 Index)
+{
+	if (Builder.IsValid())
+	{
+		const FMetaSoundFrontendDocumentBuilder& ConstBuilder = Builder->GetConstBuilder();
+		const FMetasoundFrontendGraphClass& GraphClass = ConstBuilder.GetConstDocumentChecked().RootGraph;
+		const FMetasoundFrontendClassOutput& GraphOutput = GraphClass.GetDefaultInterface().Outputs[Index];
+
+		OnGraphOutputIsConstructorPinChangedDelegate.Broadcast(GraphOutput.Name, GraphOutput.AccessType == EMetasoundFrontendVertexAccessType::Value);
+	}
+}
+
+void UMetaSoundEditorBuilderListener::OnGraphOutputIsAdvancedDisplayChanged(int32 Index)
+{
+	if (Builder.IsValid())
+	{
+		const FMetaSoundFrontendDocumentBuilder& ConstBuilder = Builder->GetConstBuilder();
+		const FMetasoundFrontendGraphClass& GraphClass = ConstBuilder.GetConstDocumentChecked().RootGraph;
+		const FMetasoundFrontendClassOutput& GraphOutput = GraphClass.GetDefaultInterface().Outputs[Index];
+
+		OnGraphOutputIsAdvancedDisplayChangedDelegate.Broadcast(GraphOutput.Name, GraphOutput.Metadata.bIsAdvancedDisplay);
+	}
+}
+
 void UMetaSoundEditorBuilderListener::OnGraphInputNameChanged(FName OldName, FName NewName)
 {
 	if (Builder.IsValid())
@@ -126,37 +331,61 @@ void UMetaSoundEditorBuilderListener::OnGraphOutputNameChanged(FName OldName, FN
 
 void UMetaSoundEditorBuilderListener::RemoveAllDelegates()
 {
-	// Remove multicast BP delegates 
+	// Remove multicast BP delegates
 	OnGraphInputAddedDelegate.RemoveAll(this);
-	OnGraphInputDefaultChangedDelegate.RemoveAll(this);
 	OnRemovingGraphInputDelegate.RemoveAll(this);
+	OnGraphInputNameChangedDelegate.RemoveAll(this);
+	OnGraphInputDataTypeChangedDelegate.RemoveAll(this);
+	OnGraphInputDescriptionChangedDelegate.RemoveAll(this);
+	OnGraphInputSortOrderIndexChangedDelegate.RemoveAll(this);
+	OnGraphInputIsAdvancedDisplayChangedDelegate.RemoveAll(this);
+	OnGraphInputIsConstructorPinChangedDelegate.RemoveAll(this);
+	OnGraphInputDefaultChangedDelegate.RemoveAll(this);
 
 	OnGraphOutputAddedDelegate.RemoveAll(this);
 	OnRemovingGraphOutputDelegate.RemoveAll(this);
-
-	OnGraphInputDataTypeChangedDelegate.RemoveAll(this);
-	OnGraphOutputDataTypeChangedDelegate.RemoveAll(this);
-	OnGraphInputNameChangedDelegate.RemoveAll(this);
 	OnGraphOutputNameChangedDelegate.RemoveAll(this);
+	OnGraphOutputDataTypeChangedDelegate.RemoveAll(this);
+	OnGraphOutputDescriptionChangedDelegate.RemoveAll(this);
+	OnGraphOutputSortOrderIndexChangedDelegate.RemoveAll(this);
+	OnGraphOutputIsConstructorPinChangedDelegate.RemoveAll(this);
+	OnGraphOutputIsAdvancedDisplayChangedDelegate.RemoveAll(this);
 
-	// Remove document delegates 
+	// Remove document delegates
 	if (Builder.IsValid())
 	{
 		FMetaSoundFrontendDocumentBuilder& DocumentBuilder = Builder->GetBuilder();
 		Metasound::Frontend::FDocumentModifyDelegates& BuilderDelegates = DocumentBuilder.GetDocumentDelegates();
 		Metasound::Frontend::FInterfaceModifyDelegates& InterfaceDelegates = BuilderDelegates.InterfaceDelegates;
 
+		BuilderDelegates.OnDocumentMetadataChanged.Remove(OnDocumentDisplayNameChangedHandle);
+		BuilderDelegates.OnDocumentMetadataChanged.Remove(OnDocumentDescriptionChangedHandle);
+		BuilderDelegates.OnDocumentMetadataChanged.Remove(OnDocumentAuthorChangedHandle);
+		BuilderDelegates.OnDocumentMetadataChanged.Remove(OnDocumentKeywordsChangedHandle);
+		BuilderDelegates.OnDocumentMetadataChanged.Remove(OnDocumentCategoryHierarchyChangedHandle);
+		BuilderDelegates.OnDocumentMetadataChanged.Remove(OnDocumentIsDeprecatedChangedHandle);
+
 		InterfaceDelegates.OnInputAdded.Remove(OnInputAddedHandle);
-		InterfaceDelegates.OnInputDefaultChanged.Remove(OnInputDefaultChangedHandle);
 		InterfaceDelegates.OnRemovingInput.Remove(OnRemovingInputHandle);
+		InterfaceDelegates.OnInputNameChanged.Remove(OnInputNameChangedHandle);
+		InterfaceDelegates.OnInputDisplayNameChanged.Remove(OnInputDisplayNameChangedHandle);
+		InterfaceDelegates.OnInputDataTypeChanged.Remove(OnInputDataTypeChangedHandle);
+		InterfaceDelegates.OnInputDescriptionChanged.Remove(OnInputDescriptionChangedHandle);
+		InterfaceDelegates.OnInputSortOrderIndexChanged.Remove(OnInputSortOrderIndexChangedHandle);
+		InterfaceDelegates.OnInputIsConstructorPinChanged.Remove(OnInputIsConstructorPinChangedHandle);
+		InterfaceDelegates.OnInputIsAdvancedDisplayChanged.Remove(OnInputIsAdvancedDisplayChangedHandle);
+		InterfaceDelegates.OnInputDefaultChanged.Remove(OnInputDefaultChangedHandle);
+		InterfaceDelegates.OnInputInheritsDefaultChanged.Remove(OnInputInheritsDefaultChangedHandle);
 
 		InterfaceDelegates.OnOutputAdded.Remove(OnOutputAddedHandle);
 		InterfaceDelegates.OnRemovingOutput.Remove(OnRemovingOutputHandle);
-
-		InterfaceDelegates.OnInputDataTypeChanged.Remove(OnInputDataTypeChangedHandle);
-		InterfaceDelegates.OnOutputDataTypeChanged.Remove(OnOutputDataTypeChangedHandle);
-		InterfaceDelegates.OnInputNameChanged.Remove(OnInputNameChangedHandle);
 		InterfaceDelegates.OnOutputNameChanged.Remove(OnOutputNameChangedHandle);
+		InterfaceDelegates.OnOutputDisplayNameChanged.Remove(OnOutputDisplayNameChangedHandle);
+		InterfaceDelegates.OnOutputDataTypeChanged.Remove(OnOutputDataTypeChangedHandle);
+		InterfaceDelegates.OnOutputDescriptionChanged.Remove(OnOutputDescriptionChangedHandle);
+		InterfaceDelegates.OnOutputSortOrderIndexChanged.Remove(OnOutputSortOrderIndexChangedHandle);
+		InterfaceDelegates.OnOutputIsConstructorPinChanged.Remove(OnOutputIsConstructorPinChangedHandle);
+		InterfaceDelegates.OnOutputIsAdvancedDisplayChanged.Remove(OnOutputIsAdvancedDisplayChangedHandle);
 	}
 }
 
@@ -166,16 +395,33 @@ void UMetaSoundEditorBuilderListener::FEditorBuilderListener::OnBuilderReloaded(
 	{
 		Metasound::Frontend::FInterfaceModifyDelegates& InterfaceDelegates = OutDelegates.InterfaceDelegates;
 
+		Parent->OnDocumentDisplayNameChangedHandle = OutDelegates.OnDocumentMetadataChanged.AddUObject(Parent, &UMetaSoundEditorBuilderListener::OnDocumentDisplayNameChanged);
+		Parent->OnDocumentDescriptionChangedHandle = OutDelegates.OnDocumentMetadataChanged.AddUObject(Parent, &UMetaSoundEditorBuilderListener::OnDocumentDescriptionChanged);
+		Parent->OnDocumentAuthorChangedHandle = OutDelegates.OnDocumentMetadataChanged.AddUObject(Parent, &UMetaSoundEditorBuilderListener::OnDocumentAuthorChanged);
+		Parent->OnDocumentKeywordsChangedHandle = OutDelegates.OnDocumentMetadataChanged.AddUObject(Parent, &UMetaSoundEditorBuilderListener::OnDocumentKeywordsChanged);
+		Parent->OnDocumentCategoryHierarchyChangedHandle = OutDelegates.OnDocumentMetadataChanged.AddUObject(Parent, &UMetaSoundEditorBuilderListener::OnDocumentCategoryHierarchyChanged);
+		Parent->OnDocumentIsDeprecatedChangedHandle = OutDelegates.OnDocumentMetadataChanged.AddUObject(Parent, &UMetaSoundEditorBuilderListener::OnDocumentIsDeprecatedChanged);
+
 		Parent->OnInputAddedHandle = InterfaceDelegates.OnInputAdded.AddUObject(Parent, &UMetaSoundEditorBuilderListener::OnGraphInputAdded);
-		Parent->OnInputDefaultChangedHandle = InterfaceDelegates.OnInputDefaultChanged.AddUObject(Parent, &UMetaSoundEditorBuilderListener::OnGraphInputDefaultChanged);
 		Parent->OnRemovingInputHandle = InterfaceDelegates.OnRemovingInput.AddUObject(Parent, &UMetaSoundEditorBuilderListener::OnRemovingGraphInput);
+		Parent->OnInputNameChangedHandle = InterfaceDelegates.OnInputNameChanged.AddUObject(Parent, &UMetaSoundEditorBuilderListener::OnGraphInputNameChanged);
+		Parent->OnInputDisplayNameChangedHandle = InterfaceDelegates.OnInputDisplayNameChanged.AddUObject(Parent, &UMetaSoundEditorBuilderListener::OnGraphInputDisplayNameChanged);
+		Parent->OnInputDataTypeChangedHandle = InterfaceDelegates.OnInputDataTypeChanged.AddUObject(Parent, &UMetaSoundEditorBuilderListener::OnGraphInputDataTypeChanged);
+		Parent->OnInputDescriptionChangedHandle = InterfaceDelegates.OnInputDescriptionChanged.AddUObject(Parent, &UMetaSoundEditorBuilderListener::OnGraphInputDescriptionChanged);
+		Parent->OnInputSortOrderIndexChangedHandle = InterfaceDelegates.OnInputSortOrderIndexChanged.AddUObject(Parent, &UMetaSoundEditorBuilderListener::OnGraphInputSortOrderIndexChanged);
+		Parent->OnInputIsConstructorPinChangedHandle = InterfaceDelegates.OnInputIsConstructorPinChanged.AddUObject(Parent, &UMetaSoundEditorBuilderListener::OnGraphInputIsConstructorPinChanged);
+		Parent->OnInputIsAdvancedDisplayChangedHandle = InterfaceDelegates.OnInputIsAdvancedDisplayChanged.AddUObject(Parent, &UMetaSoundEditorBuilderListener::OnGraphInputIsAdvancedDisplayChanged);
+		Parent->OnInputDefaultChangedHandle = InterfaceDelegates.OnInputDefaultChanged.AddUObject(Parent, &UMetaSoundEditorBuilderListener::OnGraphInputDefaultChanged);
+		Parent->OnInputInheritsDefaultChangedHandle = InterfaceDelegates.OnInputInheritsDefaultChanged.AddUObject(Parent, &UMetaSoundEditorBuilderListener::OnGraphInputInheritsDefaultChanged);
 
 		Parent->OnOutputAddedHandle = InterfaceDelegates.OnOutputAdded.AddUObject(Parent, &UMetaSoundEditorBuilderListener::OnGraphOutputAdded);
 		Parent->OnRemovingOutputHandle = InterfaceDelegates.OnRemovingOutput.AddUObject(Parent, &UMetaSoundEditorBuilderListener::OnRemovingGraphOutput);
-
-		Parent->OnInputDataTypeChangedHandle = InterfaceDelegates.OnInputDataTypeChanged.AddUObject(Parent, &UMetaSoundEditorBuilderListener::OnGraphInputDataTypeChanged);
-		Parent->OnOutputDataTypeChangedHandle = InterfaceDelegates.OnOutputDataTypeChanged.AddUObject(Parent, &UMetaSoundEditorBuilderListener::OnGraphOutputDataTypeChanged);
-		Parent->OnInputNameChangedHandle = InterfaceDelegates.OnInputNameChanged.AddUObject(Parent, &UMetaSoundEditorBuilderListener::OnGraphInputNameChanged);
 		Parent->OnOutputNameChangedHandle = InterfaceDelegates.OnOutputNameChanged.AddUObject(Parent, &UMetaSoundEditorBuilderListener::OnGraphOutputNameChanged);
+		Parent->OnOutputDisplayNameChangedHandle = InterfaceDelegates.OnOutputDisplayNameChanged.AddUObject(Parent, &UMetaSoundEditorBuilderListener::OnGraphOutputDisplayNameChanged);
+		Parent->OnOutputDataTypeChangedHandle = InterfaceDelegates.OnOutputDataTypeChanged.AddUObject(Parent, &UMetaSoundEditorBuilderListener::OnGraphOutputDataTypeChanged);
+		Parent->OnOutputDescriptionChangedHandle = InterfaceDelegates.OnOutputDescriptionChanged.AddUObject(Parent, &UMetaSoundEditorBuilderListener::OnGraphOutputDescriptionChanged);
+		Parent->OnOutputSortOrderIndexChangedHandle = InterfaceDelegates.OnOutputSortOrderIndexChanged.AddUObject(Parent, &UMetaSoundEditorBuilderListener::OnGraphOutputSortOrderIndexChanged);
+		Parent->OnOutputIsConstructorPinChangedHandle = InterfaceDelegates.OnOutputIsConstructorPinChanged.AddUObject(Parent, &UMetaSoundEditorBuilderListener::OnGraphOutputIsConstructorPinChanged);
+		Parent->OnOutputIsAdvancedDisplayChangedHandle = InterfaceDelegates.OnOutputIsAdvancedDisplayChanged.AddUObject(Parent, &UMetaSoundEditorBuilderListener::OnGraphOutputIsAdvancedDisplayChanged);
 	}
 }
diff --git a/Engine/Plugins/Runtime/Metasound/Source/MetasoundEditor/Public/MetasoundEditorBuilderListener.h b/Engine/Plugins/Runtime/Metasound/Source/MetasoundEditor/Public/MetasoundEditorBuilderListener.h
--- a/Engine/Plugins/Runtime/Metasound/Source/MetasoundEditor/Public/MetasoundEditorBuilderListener.h
+++ b/Engine/Plugins/Runtime/Metasound/Source/MetasoundEditor/Public/MetasoundEditorBuilderListener.h
@@ -11,77 +11,168 @@
 struct FMetaSoundNodeHandle;
 struct FMetasoundFrontendLiteral;
 
-// BP Delegates for builder changes 
+// BP Delegates for builder changes
+DECLARE_DYNAMIC_MULTICAST_DELEGATE_OneParam(FOnMetaSoundBuilderDocumentMetadataTextChanged, FText, NewText);
+DECLARE_DYNAMIC_MULTICAST_DELEGATE_OneParam(FOnMetaSoundBuilderDocumentMetadataTextArrayChanged, TArray<FText>, NewTextArray);
+DECLARE_DYNAMIC_MULTICAST_DELEGATE_OneParam(FOnMetaSoundBuilderDocumentMetadataStringChanged, FString, NewString);
+DECLARE_DYNAMIC_MULTICAST_DELEGATE_OneParam(FOnMetaSoundBuilderDocumentMetadataBoolChanged, bool, NewBool);
 DECLARE_DYNAMIC_MULTICAST_DELEGATE_TwoParams(FOnMetaSoundBuilderGraphInterfaceMutate, FName, VertexName, FName, DataType);
 DECLARE_DYNAMIC_MULTICAST_DELEGATE_ThreeParams(FOnMetaSoundBuilderGraphLiteralMutate, FName, VertexName, FMetasoundFrontendLiteral, LiteralValue, FName, PageName);
 DECLARE_DYNAMIC_MULTICAST_DELEGATE_TwoParams(FOnMetaSoundBuilderGraphVertexRename, FName, OldName, FName, NewName);
+DECLARE_DYNAMIC_MULTICAST_DELEGATE_TwoParams(FOnMetaSoundBuilderGraphVertexBoolChanged, FName, VertexName, bool, bNewBool);
+DECLARE_DYNAMIC_MULTICAST_DELEGATE_TwoParams(FOnMetaSoundBuilderGraphVertexIntChanged, FName, VertexName, int32, NewInt);
+DECLARE_DYNAMIC_MULTICAST_DELEGATE_TwoParams(FOnMetaSoundBuilderGraphVertexTextChanged, FName, VertexName, FText, NewText);
 
 UCLASS(MinimalAPI, BlueprintType)
 class UMetaSoundEditorBuilderListener : public UObject
 {
-public: 
+public:
 	GENERATED_BODY()
 
 	virtual ~UMetaSoundEditorBuilderListener() {}
 	
 	void Init(const TWeakObjectPtr<UMetaSoundBuilderBase> InBuilder);
 
+	void OnDocumentDisplayNameChanged();
+	void OnDocumentDescriptionChanged();
+	void OnDocumentAuthorChanged();
+	void OnDocumentKeywordsChanged();
+	void OnDocumentCategoryHierarchyChanged();
+	void OnDocumentIsDeprecatedChanged();
+
 	void OnGraphInputAdded(int32 Index);
-	void OnGraphInputDefaultChanged(int32 Index);
 	void OnRemovingGraphInput(int32 Index);
+	void OnGraphInputNameChanged(FName OldName, FName NewName);
+	void OnGraphInputDisplayNameChanged(int32 Index);
+	void OnGraphInputDataTypeChanged(int32 Index);
+	void OnGraphInputDescriptionChanged(int32 Index);
+	void OnGraphInputSortOrderIndexChanged(int32 Index);
+	void OnGraphInputIsConstructorPinChanged(int32 Index);
+	void OnGraphInputIsAdvancedDisplayChanged(int32 Index);
+	void OnGraphInputDefaultChanged(int32 Index);
+	void OnGraphInputInheritsDefaultChanged(int32 Index);
 
 	void OnGraphOutputAdded(int32 Index);
 	void OnRemovingGraphOutput(int32 Index);
-
-	void OnGraphInputDataTypeChanged(int32 Index);
-	void OnGraphOutputDataTypeChanged(int32 Index);
-	void OnGraphInputNameChanged(FName OldName, FName NewName);
 	void OnGraphOutputNameChanged(FName OldName, FName NewName);
+	void OnGraphOutputDisplayNameChanged(int32 Index);
+	void OnGraphOutputDataTypeChanged(int32 Index);
+	void OnGraphOutputDescriptionChanged(int32 Index);
+	void OnGraphOutputSortOrderIndexChanged(int32 Index);
+	void OnGraphOutputIsConstructorPinChanged(int32 Index);
+	void OnGraphOutputIsAdvancedDisplayChanged(int32 Index);
 
-
-	UFUNCTION(BlueprintCallable, Category = "Audio|MetaSound|Builder|Editor|Delegates")
+	UFUNCTION(BlueprintCallable, Category = "Audio|MetaSound|BuilderListener")
 	METASOUNDEDITOR_API void RemoveAllDelegates();
 
-	UPROPERTY(BlueprintAssignable, BlueprintCallable, Category = "Audio|MetaSound|Builder|Editor|Delegates")
+	UPROPERTY(BlueprintAssignable, BlueprintCallable, Category = "Audio|MetaSound|BuilderListener", DisplayName = "On Document Display Name Changed")
+	FOnMetaSoundBuilderDocumentMetadataTextChanged OnDocumentDisplayNameChangedDelegate;
+
+	UPROPERTY(BlueprintAssignable, BlueprintCallable, Category = "Audio|MetaSound|BuilderListener", DisplayName = "On Document Description Changed")
+	FOnMetaSoundBuilderDocumentMetadataTextChanged OnDocumentDescriptionChangedDelegate;
+
+	UPROPERTY(BlueprintAssignable, BlueprintCallable, Category = "Audio|MetaSound|BuilderListener", DisplayName = "On Document Author Changed")
+	FOnMetaSoundBuilderDocumentMetadataStringChanged OnDocumentAuthorChangedDelegate;
+
+	UPROPERTY(BlueprintAssignable, BlueprintCallable, Category = "Audio|MetaSound|BuilderListener", DisplayName = "On Document Keywords Changed")
+	FOnMetaSoundBuilderDocumentMetadataTextArrayChanged OnDocumentKeywordsChangedDelegate;
+
+	UPROPERTY(BlueprintAssignable, BlueprintCallable, Category = "Audio|MetaSound|BuilderListener", DisplayName = "On Document Category Hierarchy Changed")
+	FOnMetaSoundBuilderDocumentMetadataTextArrayChanged OnDocumentCategoryHierarchyChangedDelegate;
+
+	UPROPERTY(BlueprintAssignable, BlueprintCallable, Category = "Audio|MetaSound|BuilderListener", DisplayName = "On Document Is Deprecated Changed")
+	FOnMetaSoundBuilderDocumentMetadataBoolChanged OnDocumentIsDeprecatedChangedDelegate;
+
+	UPROPERTY(BlueprintAssignable, BlueprintCallable, Category = "Audio|MetaSound|BuilderListener|InputDelegates", DisplayName = "On Graph Input Added")
 	FOnMetaSoundBuilderGraphInterfaceMutate OnGraphInputAddedDelegate;
-	
-	UPROPERTY(BlueprintAssignable, BlueprintCallable, Category = "Audio|MetaSound|Builder|Editor|Delegates")
-	FOnMetaSoundBuilderGraphLiteralMutate OnGraphInputDefaultChangedDelegate;
 
-	UPROPERTY(BlueprintAssignable, BlueprintCallable, Category = "Audio|MetaSound|Builder|Editor|Delegates")
+	UPROPERTY(BlueprintAssignable, BlueprintCallable, Category = "Audio|MetaSound|BuilderListener|InputDelegates", DisplayName = "On Graph Input Removed")
 	FOnMetaSoundBuilderGraphInterfaceMutate OnRemovingGraphInputDelegate;
-	
-	UPROPERTY(BlueprintAssignable, BlueprintCallable, Category= "Audio|MetaSound|Builder|Editor|Delegates")
+
+	UPROPERTY(BlueprintAssignable, BlueprintCallable, Category = "Audio|MetaSound|BuilderListener|InputDelegates", DisplayName = "On Graph Input Name Changed")
+	FOnMetaSoundBuilderGraphVertexRename OnGraphInputNameChangedDelegate;
+
+	UPROPERTY(BlueprintAssignable, BlueprintCallable, Category = "Audio|MetaSound|BuilderListener|InputDelegates", DisplayName = "On Graph Input Display Name Changed")
+	FOnMetaSoundBuilderGraphVertexTextChanged OnGraphInputDisplayNameChangedDelegate;
+
+	UPROPERTY(BlueprintAssignable, BlueprintCallable, Category = "Audio|MetaSound|BuilderListener|InputDelegates", DisplayName = "On Graph Input Data Type Changed")
+	FOnMetaSoundBuilderGraphInterfaceMutate OnGraphInputDataTypeChangedDelegate;
+
+	UPROPERTY(BlueprintAssignable, BlueprintCallable, Category = "Audio|MetaSound|BuilderListener|InputDelegates", DisplayName = "On Graph Input Description Changed")
+	FOnMetaSoundBuilderGraphVertexTextChanged OnGraphInputDescriptionChangedDelegate;
+
+	UPROPERTY(BlueprintAssignable, BlueprintCallable, Category = "Audio|MetaSound|BuilderListener|InputDelegates", DisplayName = "On Graph Input Sort Order Index Changed")
+	FOnMetaSoundBuilderGraphVertexIntChanged OnGraphInputSortOrderIndexChangedDelegate;
+
+	UPROPERTY(BlueprintAssignable, BlueprintCallable, Category = "Audio|MetaSound|BuilderListener|InputDelegates", DisplayName = "On Graph Input Is Constructor Pin Changed")
+	FOnMetaSoundBuilderGraphVertexBoolChanged OnGraphInputIsConstructorPinChangedDelegate;
+
+	UPROPERTY(BlueprintAssignable, BlueprintCallable, Category = "Audio|MetaSound|BuilderListener|InputDelegates", DisplayName = "On Graph Input Is Advanced Display Changed")
+	FOnMetaSoundBuilderGraphVertexBoolChanged OnGraphInputIsAdvancedDisplayChangedDelegate;
+
+	UPROPERTY(BlueprintAssignable, BlueprintCallable, Category = "Audio|MetaSound|BuilderListener|InputDelegates", DisplayName = "On Graph Input Default Changed")
+	FOnMetaSoundBuilderGraphLiteralMutate OnGraphInputDefaultChangedDelegate;
+
+	UPROPERTY(BlueprintAssignable, BlueprintCallable, Category = "Audio|MetaSound|BuilderListener|InputDelegates", DisplayName = "On Graph Input Inherits Default Changed")
+	FOnMetaSoundBuilderGraphVertexBoolChanged OnGraphInputInheritsDefaultChangedDelegate;
+
+	UPROPERTY(BlueprintAssignable, BlueprintCallable, Category = "Audio|MetaSound|BuilderListener|OutputDelegates", DisplayName = "On Graph Output Added")
 	FOnMetaSoundBuilderGraphInterfaceMutate OnGraphOutputAddedDelegate;
 
-	UPROPERTY(BlueprintAssignable, BlueprintCallable, Category = "Audio|MetaSound|Builder|Editor|Delegates")
+	UPROPERTY(BlueprintAssignable, BlueprintCallable, Category = "Audio|MetaSound|BuilderListener|OutputDelegates", DisplayName = "On Graph Output Removed")
 	FOnMetaSoundBuilderGraphInterfaceMutate OnRemovingGraphOutputDelegate;
 
-	UPROPERTY(BlueprintAssignable, BlueprintCallable, Category = "Audio|MetaSound|Builder|Editor|Delegates")
-	FOnMetaSoundBuilderGraphInterfaceMutate OnGraphInputDataTypeChangedDelegate;
+	UPROPERTY(BlueprintAssignable, BlueprintCallable, Category = "Audio|MetaSound|BuilderListener|OutputDelegates", DisplayName = "On Graph Output Name Changed")
+	FOnMetaSoundBuilderGraphVertexRename OnGraphOutputNameChangedDelegate;
 
-	UPROPERTY(BlueprintAssignable, BlueprintCallable, Category = "Audio|MetaSound|Builder|Editor|Delegates")
+	UPROPERTY(BlueprintAssignable, BlueprintCallable, Category = "Audio|MetaSound|BuilderListener|OutputDelegates", DisplayName = "On Graph Output Display Name Changed")
+	FOnMetaSoundBuilderGraphVertexTextChanged OnGraphOutputDisplayNameChangedDelegate;
+
+	UPROPERTY(BlueprintAssignable, BlueprintCallable, Category = "Audio|MetaSound|BuilderListener|OutputDelegates", DisplayName = "On Graph Output Data Type Changed")
 	FOnMetaSoundBuilderGraphInterfaceMutate OnGraphOutputDataTypeChangedDelegate;
 
-	UPROPERTY(BlueprintAssignable, BlueprintCallable, Category = "Audio|MetaSound|Builder|Editor|Delegates")
-	FOnMetaSoundBuilderGraphVertexRename OnGraphInputNameChangedDelegate;
+	UPROPERTY(BlueprintAssignable, BlueprintCallable, Category = "Audio|MetaSound|BuilderListener|OutputDelegates", DisplayName = "On Graph Output Description Changed")
+	FOnMetaSoundBuilderGraphVertexTextChanged OnGraphOutputDescriptionChangedDelegate;
 
-	UPROPERTY(BlueprintAssignable, BlueprintCallable, Category = "Audio|MetaSound|Builder|Editor|Delegates")
-	FOnMetaSoundBuilderGraphVertexRename OnGraphOutputNameChangedDelegate;
+	UPROPERTY(BlueprintAssignable, BlueprintCallable, Category = "Audio|MetaSound|BuilderListener|OutputDelegates", DisplayName = "On Graph Output Sort Order Index Changed")
+	FOnMetaSoundBuilderGraphVertexIntChanged OnGraphOutputSortOrderIndexChangedDelegate;
+
+	UPROPERTY(BlueprintAssignable, BlueprintCallable, Category = "Audio|MetaSound|BuilderListener|OutputDelegates", DisplayName = "On Graph Output Is Constructor Pin Changed")
+	FOnMetaSoundBuilderGraphVertexBoolChanged OnGraphOutputIsConstructorPinChangedDelegate;
+
+	UPROPERTY(BlueprintAssignable, BlueprintCallable, Category = "Audio|MetaSound|BuilderListener|OutputDelegates", DisplayName = "On Graph Output Is Advanced Display Changed")
+	FOnMetaSoundBuilderGraphVertexBoolChanged OnGraphOutputIsAdvancedDisplayChangedDelegate;
 
 private:
 	// Handles for document delegates
+	FDelegateHandle OnDocumentDisplayNameChangedHandle;
+	FDelegateHandle OnDocumentDescriptionChangedHandle;
+	FDelegateHandle OnDocumentAuthorChangedHandle;
+	FDelegateHandle OnDocumentKeywordsChangedHandle;
+	FDelegateHandle OnDocumentCategoryHierarchyChangedHandle;
+	FDelegateHandle OnDocumentIsDeprecatedChangedHandle;
+
 	FDelegateHandle OnInputAddedHandle;
-	FDelegateHandle OnInputDefaultChangedHandle;
 	FDelegateHandle OnRemovingInputHandle;
+	FDelegateHandle OnInputNameChangedHandle;
+	FDelegateHandle OnInputDisplayNameChangedHandle;
+	FDelegateHandle OnInputDataTypeChangedHandle;
+	FDelegateHandle OnInputDescriptionChangedHandle;
+	FDelegateHandle OnInputSortOrderIndexChangedHandle;
+	FDelegateHandle OnInputIsConstructorPinChangedHandle;
+	FDelegateHandle OnInputIsAdvancedDisplayChangedHandle;
+	FDelegateHandle OnInputDefaultChangedHandle;
+	FDelegateHandle OnInputInheritsDefaultChangedHandle;
 
 	FDelegateHandle OnOutputAddedHandle;
 	FDelegateHandle OnRemovingOutputHandle;
-
-	FDelegateHandle OnInputDataTypeChangedHandle;
-	FDelegateHandle OnOutputDataTypeChangedHandle;
-	FDelegateHandle OnInputNameChangedHandle;
 	FDelegateHandle OnOutputNameChangedHandle;
+	FDelegateHandle OnOutputDisplayNameChangedHandle;
+	FDelegateHandle OnOutputDataTypeChangedHandle;
+	FDelegateHandle OnOutputDescriptionChangedHandle;
+	FDelegateHandle OnOutputSortOrderIndexChangedHandle;
+	FDelegateHandle OnOutputIsConstructorPinChangedHandle;
+	FDelegateHandle OnOutputIsAdvancedDisplayChangedHandle;
 
 	TWeakObjectPtr<UMetaSoundBuilderBase> Builder;
 
diff --git a/Engine/Plugins/Runtime/Metasound/Source/MetasoundFrontend/Private/MetasoundFrontendDocumentBuilder.cpp b/Engine/Plugins/Runtime/Metasound/Source/MetasoundFrontend/Private/MetasoundFrontendDocumentBuilder.cpp
--- a/Engine/Plugins/Runtime/Metasound/Source/MetasoundFrontend/Private/MetasoundFrontendDocumentBuilder.cpp
+++ b/Engine/Plugins/Runtime/Metasound/Source/MetasoundFrontend/Private/MetasoundFrontendDocumentBuilder.cpp
@@ -5164,6 +5164,7 @@ void FMetaSoundFrontendDocumentBuilder::SetAuthor(const FString& InAuthor)
 {
 	FMetasoundFrontendClassMetadata& ClassMetadata = GetDocumentChecked().RootGraph.Metadata;
 	ClassMetadata.SetAuthor(InAuthor);
+	DocumentDelegates->OnDocumentMetadataChanged.Broadcast();
 }
 #endif // WITH_EDITOR
 
@@ -5204,6 +5205,7 @@ bool FMetaSoundFrontendDocumentBuilder::SetGraphInputAdvancedDisplay(const FName
 		{
 			GraphInput.Metadata.SetIsAdvancedDisplay(InAdvancedDisplay);
 			Document.Metadata.ModifyContext.AddMemberIDModified(GraphInput.VertexID);
+			DocumentDelegates->InterfaceDelegates.OnInputIsAdvancedDisplayChanged.Broadcast(*Index);
 			return true;
 		}
 	}
@@ -5294,6 +5296,8 @@ bool FMetaSoundFrontendDocumentBuilder::SetGraphInputAccessType(FName InputName,
 			return false;
 		}
 
+		DocumentDelegates->InterfaceDelegates.OnInputIsConstructorPinChanged.Broadcast(*Index);
+
 #if WITH_EDITORONLY_DATA
 		RootGraph.GetDefaultInterface().UpdateChangeID();
 		Document.Metadata.ModifyContext.AddMemberIDModified(GraphInput.NodeID);
@@ -5359,9 +5363,8 @@ bool FMetaSoundFrontendDocumentBuilder::SetGraphInputDataType(FName InputName, F
 				return false;
 			}
 
-#if WITH_EDITOR
 			DocumentDelegates->InterfaceDelegates.OnInputDataTypeChanged.Broadcast(*Index);
-#endif // WITH_EDITOR
+			DocumentDelegates->InterfaceDelegates.OnInputDefaultChanged.Broadcast(*Index);
 
 			RemoveUnusedDependencies();
 
@@ -5465,6 +5468,7 @@ bool FMetaSoundFrontendDocumentBuilder::SetGraphInputDescription(FName InputName
 		FMetasoundFrontendClassInput& GraphInput = Document.RootGraph.GetDefaultInterface().Inputs[*Index];
 		GraphInput.Metadata.SetDescription(Description);
 		Document.Metadata.ModifyContext.AddMemberIDModified(GraphInput.NodeID);
+		DocumentDelegates->InterfaceDelegates.OnInputDescriptionChanged.Broadcast(*Index);
 		return true;
 	}
 
@@ -5481,6 +5485,7 @@ bool FMetaSoundFrontendDocumentBuilder::SetGraphInputDisplayName(FName InputName
 		FMetasoundFrontendClassInput& GraphInput = Document.RootGraph.GetDefaultInterface().Inputs[*Index];
 		GraphInput.Metadata.SetDisplayName(DisplayName);
 		Document.Metadata.ModifyContext.AddMemberIDModified(GraphInput.NodeID);
+		DocumentDelegates->InterfaceDelegates.OnInputDisplayNameChanged.Broadcast(*Index);
 		return true;
 	}
 
@@ -5491,22 +5496,30 @@ bool FMetaSoundFrontendDocumentBuilder::SetGraphInputDisplayName(FName InputName
 bool FMetaSoundFrontendDocumentBuilder::SetGraphInputInheritsDefault(FName InName, bool bInputInheritsDefault, bool bForceUpdate)
 {
 	FMetasoundFrontendGraphClassPresetOptions& PresetOptions = GetDocumentChecked().RootGraph.PresetOptions;
+	if (!PresetOptions.bIsPreset && !bForceUpdate)
+	{
+		return false;
+	}
+
+	bool bValueChanged;
 	if (bInputInheritsDefault)
 	{
-		if (PresetOptions.bIsPreset || bForceUpdate)
-		{
-			return PresetOptions.InputsInheritingDefault.Add(InName).IsValidId();
-		}
+		bValueChanged = PresetOptions.InputsInheritingDefault.Add(InName).IsValidId();
 	}
 	else
 	{
-		if (PresetOptions.bIsPreset || bForceUpdate)
+		bValueChanged = PresetOptions.InputsInheritingDefault.Remove(InName) > 0;
+	}
+
+	if (bValueChanged)
+	{
+		if (const int32* Index = DocumentCache->GetInterfaceCache().FindInputIndex(InName))
 		{
-			return PresetOptions.InputsInheritingDefault.Remove(InName) > 0;
+			DocumentDelegates->InterfaceDelegates.OnInputInheritsDefaultChanged.Broadcast(*Index);
 		}
 	}
 
-	return false;
+	return bValueChanged;
 }
 
 bool FMetaSoundFrontendDocumentBuilder::SetGraphInputName(FName InputName, FName NewName)
@@ -5569,6 +5582,7 @@ bool FMetaSoundFrontendDocumentBuilder::SetGraphInputSortOrderIndex(const FName
 		FMetasoundFrontendClassInput& GraphInput = Document.RootGraph.GetDefaultInterface().Inputs[*Index];
 		GraphInput.Metadata.SortOrderIndex = InSortOrderIndex;
 		Document.Metadata.ModifyContext.AddMemberIDModified(GraphInput.NodeID);
+		DocumentDelegates->InterfaceDelegates.OnInputSortOrderIndexChanged.Broadcast(*Index);
 		return true;
 	}
 
@@ -5585,6 +5599,7 @@ bool FMetaSoundFrontendDocumentBuilder::SetGraphOutputSortOrderIndex(const FName
 		FMetasoundFrontendClassOutput& GraphOutput = Document.RootGraph.GetDefaultInterface().Outputs[*Index];
 		GraphOutput.Metadata.SortOrderIndex = InSortOrderIndex;
 		Document.Metadata.ModifyContext.AddMemberIDModified(GraphOutput.NodeID);
+		DocumentDelegates->InterfaceDelegates.OnOutputSortOrderIndexChanged.Broadcast(*Index);
 		return true;
 	}
 
@@ -5603,6 +5618,7 @@ bool FMetaSoundFrontendDocumentBuilder::SetGraphOutputAdvancedDisplay(const FNam
 		{
 			GraphOutput.Metadata.SetIsAdvancedDisplay(InAdvancedDisplay);
 			Document.Metadata.ModifyContext.AddMemberIDModified(GraphOutput.VertexID);
+			DocumentDelegates->InterfaceDelegates.OnOutputIsAdvancedDisplayChanged.Broadcast(*Index);
 			return true;
 		}
 	}
@@ -5663,6 +5679,8 @@ bool FMetaSoundFrontendDocumentBuilder::SetGraphOutputAccessType(FName OutputNam
 			return false;
 		}
 
+		DocumentDelegates->InterfaceDelegates.OnOutputIsConstructorPinChanged.Broadcast(*Index);
+
 #if WITH_EDITORONLY_DATA
 		RootGraph.GetDefaultInterface().UpdateChangeID();
 		Document.Metadata.ModifyContext.AddMemberIDModified(GraphOutput.NodeID);
@@ -5714,9 +5732,8 @@ bool FMetaSoundFrontendDocumentBuilder::SetGraphOutputDataType(FName OutputName,
 		{
 			return false;
 		}
-#if WITH_EDITOR
+
 		DocumentDelegates->InterfaceDelegates.OnOutputDataTypeChanged.Broadcast(*Index);
-#endif // WITH_EDITOR
 
 #if WITH_EDITORONLY_DATA
 		RootGraph.GetDefaultInterface().UpdateChangeID();
@@ -5739,6 +5756,7 @@ bool FMetaSoundFrontendDocumentBuilder::SetGraphOutputDescription(FName OutputNa
 		FMetasoundFrontendClassOutput& GraphOutput = Document.RootGraph.GetDefaultInterface().Outputs[*Index];
 		GraphOutput.Metadata.SetDescription(Description);
 		Document.Metadata.ModifyContext.AddMemberIDModified(GraphOutput.NodeID);
+		DocumentDelegates->InterfaceDelegates.OnOutputDescriptionChanged.Broadcast(*Index);
 		return true;
 	}
 
@@ -5755,6 +5773,7 @@ bool FMetaSoundFrontendDocumentBuilder::SetGraphOutputDisplayName(FName OutputNa
 		FMetasoundFrontendClassOutput& GraphOutput = Document.RootGraph.GetDefaultInterface().Outputs[*Index];
 		GraphOutput.Metadata.SetDisplayName(DisplayName);
 		Document.Metadata.ModifyContext.AddMemberIDModified(GraphOutput.NodeID);
+		DocumentDelegates->InterfaceDelegates.OnOutputDisplayNameChanged.Broadcast(*Index);
 		return true;
 	}
 
@@ -5881,26 +5900,31 @@ bool FMetaSoundFrontendDocumentBuilder::SetGraphVariableName(FName VariableName,
 void FMetaSoundFrontendDocumentBuilder::SetDisplayName(const FText& InDisplayName)
 {
 	DocumentInterface->GetDocument().RootGraph.Metadata.SetDisplayName(InDisplayName);
+	DocumentDelegates->OnDocumentMetadataChanged.Broadcast();
 }
 
 void FMetaSoundFrontendDocumentBuilder::SetDescription(const FText& InDescription)
 {
 	DocumentInterface->GetDocument().RootGraph.Metadata.SetDescription(InDescription);
+	DocumentDelegates->OnDocumentMetadataChanged.Broadcast();
 }
 
 void FMetaSoundFrontendDocumentBuilder::SetKeywords(const TArray<FText>& InKeywords)
 {
 	DocumentInterface->GetDocument().RootGraph.Metadata.SetKeywords(InKeywords);
+	DocumentDelegates->OnDocumentMetadataChanged.Broadcast();
 }
 
 void FMetaSoundFrontendDocumentBuilder::SetCategoryHierarchy(const TArray<FText>& InCategoryHierarchy)
 {
 	DocumentInterface->GetDocument().RootGraph.Metadata.SetCategoryHierarchy(InCategoryHierarchy);
+	DocumentDelegates->OnDocumentMetadataChanged.Broadcast();
 }
 
 void FMetaSoundFrontendDocumentBuilder::SetIsDeprecated(const bool bInIsDeprecated)
 {
 	DocumentInterface->GetDocument().RootGraph.Metadata.SetIsDeprecated(bInIsDeprecated);
+	DocumentDelegates->OnDocumentMetadataChanged.Broadcast();
 }
 
 void FMetaSoundFrontendDocumentBuilder::SetMemberMetadata(UMetaSoundFrontendMemberMetadata& NewMetadata)
diff --git a/Engine/Plugins/Runtime/Metasound/Source/MetasoundFrontend/Private/MetasoundFrontendDocumentCache.cpp b/Engine/Plugins/Runtime/Metasound/Source/MetasoundFrontend/Private/MetasoundFrontendDocumentCache.cpp
--- a/Engine/Plugins/Runtime/Metasound/Source/MetasoundFrontend/Private/MetasoundFrontendDocumentCache.cpp
+++ b/Engine/Plugins/Runtime/Metasound/Source/MetasoundFrontend/Private/MetasoundFrontendDocumentCache.cpp
@@ -403,17 +403,29 @@
 	void FDocumentGraphInterfaceCache::Init(FInterfaceModifyDelegates& OutDelegates)
 	{
 		OutDelegates.OnInputAdded.AddSP(this, &FDocumentGraphInterfaceCache::OnInputAdded);
+		OutDelegates.OnRemovingInput.AddSP(this, &FDocumentGraphInterfaceCache::OnRemovingInput);
+		OutDelegates.OnInputNameChanged.AddSP(this, &FDocumentGraphInterfaceCache::OnInputNameChanged);
+		OutDelegates.OnInputDataTypeChanged.AddSP(this, &FDocumentGraphInterfaceCache::OnInputDataTypeChanged);
 		OutDelegates.OnInputDefaultChanged.AddSP(this, &FDocumentGraphInterfaceCache::OnInputDefaultChanged);
+		OutDelegates.OnInputInheritsDefaultChanged.AddSP(this, &FDocumentGraphInterfaceCache::OnInputInheritsDefaultChanged);
+
+#if WITH_EDITOR
+		OutDelegates.OnInputDisplayNameChanged.AddSP(this, &FDocumentGraphInterfaceCache::OnInputDisplayNameChanged);
+		OutDelegates.OnInputDescriptionChanged.AddSP(this, &FDocumentGraphInterfaceCache::OnInputDescriptionChanged);
+		OutDelegates.OnInputSortOrderIndexChanged.AddSP(this, &FDocumentGraphInterfaceCache::OnInputSortOrderIndexChanged);
+		OutDelegates.OnInputIsAdvancedDisplayChanged.AddSP(this, &FDocumentGraphInterfaceCache::OnInputIsAdvancedDisplayChanged);
+#endif // WITH_EDITOR
+
 		OutDelegates.OnOutputAdded.AddSP(this, &FDocumentGraphInterfaceCache::OnOutputAdded);
-
-		OutDelegates.OnRemovingInput.AddSP(this, &FDocumentGraphInterfaceCache::OnRemovingInput);
 		OutDelegates.OnRemovingOutput.AddSP(this, &FDocumentGraphInterfaceCache::OnRemovingOutput);
-
-		OutDelegates.OnInputNameChanged.AddSP(this, &FDocumentGraphInterfaceCache::OnInputNameChanged);
 		OutDelegates.OnOutputNameChanged.AddSP(this, &FDocumentGraphInterfaceCache::OnOutputNameChanged);
-#if WITH_EDITOR
-		OutDelegates.OnInputDataTypeChanged.AddSP(this, &FDocumentGraphInterfaceCache::OnInputDataTypeChanged);
 		OutDelegates.OnOutputDataTypeChanged.AddSP(this, &FDocumentGraphInterfaceCache::OnOutputDataTypeChanged);
+
+#if WITH_EDITOR
+		OutDelegates.OnOutputDisplayNameChanged.AddSP(this, &FDocumentGraphInterfaceCache::OnOutputDisplayNameChanged);
+		OutDelegates.OnOutputDescriptionChanged.AddSP(this, &FDocumentGraphInterfaceCache::OnOutputDescriptionChanged);
+		OutDelegates.OnOutputSortOrderIndexChanged.AddSP(this, &FDocumentGraphInterfaceCache::OnOutputSortOrderIndexChanged);
+		OutDelegates.OnOutputIsAdvancedDisplayChanged.AddSP(this, &FDocumentGraphInterfaceCache::OnOutputIsAdvancedDisplayChanged);
 #endif // WITH_EDITOR
 	}
 
@@ -427,15 +439,13 @@
 		++TransactionCount;
 	}
 
-#if WITH_EDITOR
-	void FDocumentGraphInterfaceCache::OnInputDataTypeChanged(int32 Index)
+	void FDocumentGraphInterfaceCache::OnRemovingInput(int32 IndexBeingRemoved)
 	{
-		++TransactionCount;
-	}
-#endif // WITH_EDITOR
+		const FMetasoundFrontendDocument& Document = Parent->GetDocument();
+		const FMetasoundFrontendGraphClass& GraphClass = Document.RootGraph;
+		const FMetasoundFrontendClassInput& Input = GraphClass.GetDefaultInterface().Inputs[IndexBeingRemoved];
+		InputNameToIndex.Remove(Input.Name);
 
-	void FDocumentGraphInterfaceCache::OnInputDefaultChanged(int32 NewIndex)
-	{
 		++TransactionCount;
 	}
 
@@ -449,6 +459,43 @@
 		++TransactionCount;
 	}
 
+	void FDocumentGraphInterfaceCache::OnInputDataTypeChanged(int32 Index)
+	{
+		++TransactionCount;
+	}
+
+	void FDocumentGraphInterfaceCache::OnInputDefaultChanged(int32 NewIndex)
+	{
+		++TransactionCount;
+	}
+
+	void FDocumentGraphInterfaceCache::OnInputInheritsDefaultChanged(int32 NewIndex)
+	{
+		++TransactionCount;
+	}
+
+#if WITH_EDITOR
+	void FDocumentGraphInterfaceCache::OnInputDescriptionChanged(int32 Index)
+	{
+		++TransactionCount;
+	}
+
+	void FDocumentGraphInterfaceCache::OnInputSortOrderIndexChanged(int32 Index)
+	{
+		++TransactionCount;
+	}
+
+	void FDocumentGraphInterfaceCache::OnInputIsAdvancedDisplayChanged(int32 Index)
+	{
+		++TransactionCount;
+	}
+
+	void FDocumentGraphInterfaceCache::OnOutputDisplayNameChanged(int32 Index)
+	{
+		++TransactionCount;
+	}
+#endif // WITH_EDITOR
+
 	void FDocumentGraphInterfaceCache::OnOutputAdded(int32 NewIndex)
 	{
 		const FMetasoundFrontendDocument& Document = Parent->GetDocument();
@@ -459,12 +506,15 @@
 		++TransactionCount;
 	}
 
-#if WITH_EDITOR
-	void FDocumentGraphInterfaceCache::OnOutputDataTypeChanged(int32 Index)
+	void FDocumentGraphInterfaceCache::OnRemovingOutput(int32 IndexBeingRemoved)
 	{
+		const FMetasoundFrontendDocument& Document = Parent->GetDocument();
+		const FMetasoundFrontendGraphClass& GraphClass = Document.RootGraph;
+		const FMetasoundFrontendClassOutput& Output = GraphClass.GetDefaultInterface().Outputs[IndexBeingRemoved];
+		OutputNameToIndex.Remove(Output.Name);
+
 		++TransactionCount;
 	}
-#endif // WITH_EDITOR
 
 	void FDocumentGraphInterfaceCache::OnOutputNameChanged(FName OldName, FName NewName)
 	{
@@ -476,26 +526,33 @@
 		++TransactionCount;
 	}
 
-	void FDocumentGraphInterfaceCache::OnRemovingInput(int32 IndexBeingRemoved)
+	void FDocumentGraphInterfaceCache::OnOutputDataTypeChanged(int32 Index)
 	{
-		const FMetasoundFrontendDocument& Document = Parent->GetDocument();
-		const FMetasoundFrontendGraphClass& GraphClass = Document.RootGraph;
-		const FMetasoundFrontendClassInput& Input = GraphClass.GetDefaultInterface().Inputs[IndexBeingRemoved];
-		InputNameToIndex.Remove(Input.Name);
-
 		++TransactionCount;
 	}
 
-	void FDocumentGraphInterfaceCache::OnRemovingOutput(int32 IndexBeingRemoved)
+#if WITH_EDITOR
+	void FDocumentGraphInterfaceCache::OnInputDisplayNameChanged(int32 Index)
 	{
-		const FMetasoundFrontendDocument& Document = Parent->GetDocument();
-		const FMetasoundFrontendGraphClass& GraphClass = Document.RootGraph;
-		const FMetasoundFrontendClassOutput& Output = GraphClass.GetDefaultInterface().Outputs[IndexBeingRemoved];
-		OutputNameToIndex.Remove(Output.Name);
-
 		++TransactionCount;
 	}
 
+	void FDocumentGraphInterfaceCache::OnOutputDescriptionChanged(int32 Index)
+	{
+		++TransactionCount;
+	}
+
+	void FDocumentGraphInterfaceCache::OnOutputSortOrderIndexChanged(int32 Index)
+	{
+		++TransactionCount;
+	}
+
+	void FDocumentGraphInterfaceCache::OnOutputIsAdvancedDisplayChanged(int32 Index)
+	{
+		++TransactionCount;
+	}
+#endif // WITH_EDITOR
+
 	FDocumentGraphNodeCache::FDocumentGraphNodeCache()
 		: Parent(MakeShared<FDocumentCache>())
 	{
diff --git a/Engine/Plugins/Runtime/Metasound/Source/MetasoundFrontend/Private/MetasoundFrontendDocumentCache.h b/Engine/Plugins/Runtime/Metasound/Source/MetasoundFrontend/Private/MetasoundFrontendDocumentCache.h
--- a/Engine/Plugins/Runtime/Metasound/Source/MetasoundFrontend/Private/MetasoundFrontendDocumentCache.h
+++ b/Engine/Plugins/Runtime/Metasound/Source/MetasoundFrontend/Private/MetasoundFrontendDocumentCache.h
@@ -136,17 +136,27 @@ namespace Metasound::Frontend
 		void Init(FInterfaceModifyDelegates& OutDelegates);
 
 		void OnInputAdded(int32 NewIndex);
-		void OnInputDefaultChanged(int32 NewIndex);
+		void OnRemovingInput(int32 IndexBeingRemoved);
 		void OnInputNameChanged(FName OldName, FName NewName);
+		void OnInputDataTypeChanged(int32 Index);
+		void OnInputDefaultChanged(int32 NewIndex);
+		void OnInputInheritsDefaultChanged(int32 NewIndex);
+#if WITH_EDITOR
+		void OnInputDisplayNameChanged(int32 Index);
+		void OnInputDescriptionChanged(int32 Index);
+		void OnInputSortOrderIndexChanged(int32 Index);
+		void OnInputIsAdvancedDisplayChanged(int32 Index);
+#endif // WITH_EDITOR
 
 		void OnOutputAdded(int32 NewIndex);
-		void OnOutputNameChanged(FName OldName, FName NewName);
-		void OnRemovingInput(int32 IndexBeingRemoved);
 		void OnRemovingOutput(int32 IndexBeingRemoved);
-
-#if WITH_EDITOR
-		void OnInputDataTypeChanged(int32 Index);
+		void OnOutputNameChanged(FName OldName, FName NewName);
 		void OnOutputDataTypeChanged(int32 Index);
+#if WITH_EDITOR
+		void OnOutputDisplayNameChanged(int32 Index);
+		void OnOutputDescriptionChanged(int32 Index);
+		void OnOutputSortOrderIndexChanged(int32 Index);
+		void OnOutputIsAdvancedDisplayChanged(int32 Index);
 #endif // WITH_EDITOR
 
 		// Cache of Input name to array index of input
diff --git a/Engine/Plugins/Runtime/Metasound/Source/MetasoundFrontend/Private/MetasoundFrontendGraphController.cpp b/Engine/Plugins/Runtime/Metasound/Source/MetasoundFrontend/Private/MetasoundFrontendGraphController.cpp
--- a/Engine/Plugins/Runtime/Metasound/Source/MetasoundFrontend/Private/MetasoundFrontendGraphController.cpp
+++ b/Engine/Plugins/Runtime/Metasound/Source/MetasoundFrontend/Private/MetasoundFrontendGraphController.cpp
@@ -1380,6 +1380,24 @@ PRAGMA_ENABLE_DEPRECATION_WARNINGS
 			if (FMetasoundFrontendClassInput* Desc = FindInputDescriptionWithName(InName))
 			{
 				Desc->Metadata.SetDisplayName(InDisplayName);
+				if (const FMetasoundFrontendGraphClass* GraphClass = GraphClassPtr.Get())
+				{
+					for (const FMetasoundFrontendClassInput& Input : GraphClass->GetDefaultInterface().Inputs)
+					{
+						if (Input.Name == InName)
+						{
+							if (const IDocumentBuilderRegistry* BuilderRegistry = IDocumentBuilderRegistry::Get())
+							{
+								if (FMetaSoundFrontendDocumentBuilder* Builder = BuilderRegistry->FindBuilder(GraphClass->Metadata.GetClassName(), { }))
+								{
+									const int32 Index = Builder->GetConstDocumentChecked().RootGraph.GetDefaultInterface().Inputs.Find(Input);
+									FDocumentModifyDelegates& DocumentDelegates = Builder->GetDocumentDelegates();
+									DocumentDelegates.InterfaceDelegates.OnInputDisplayNameChanged.Broadcast(Index);
+								}
+							}
+						}
+					}
+				}
 			}
 		}
 
@@ -1388,6 +1406,24 @@ PRAGMA_ENABLE_DEPRECATION_WARNINGS
 			if (FMetasoundFrontendClassOutput* Desc = FindOutputDescriptionWithName(InName))
 			{
 				Desc->Metadata.SetDisplayName(InDisplayName);
+				if (const FMetasoundFrontendGraphClass* GraphClass = GraphClassPtr.Get())
+				{
+					for (const FMetasoundFrontendClassOutput& Output : GraphClass->GetDefaultInterface().Outputs)
+					{
+						if (Output.Name == InName)
+						{
+							if (const IDocumentBuilderRegistry* BuilderRegistry = IDocumentBuilderRegistry::Get())
+							{
+								if (FMetaSoundFrontendDocumentBuilder* Builder = BuilderRegistry->FindBuilder(GraphClass->Metadata.GetClassName(), { }))
+								{
+									const int32 Index = Builder->GetConstDocumentChecked().RootGraph.GetDefaultInterface().Outputs.Find(Output);
+									FDocumentModifyDelegates& DocumentDelegates = Builder->GetDocumentDelegates();
+									DocumentDelegates.InterfaceDelegates.OnOutputDisplayNameChanged.Broadcast(Index);
+								}
+							}
+						}
+					}
+				}
 			}
 		}
 
@@ -1396,6 +1432,24 @@ PRAGMA_ENABLE_DEPRECATION_WARNINGS
 			if (FMetasoundFrontendClassInput* Desc = FindInputDescriptionWithName(InName))
 			{
 				Desc->Metadata.SetDescription(InDescription);
+				if (const FMetasoundFrontendGraphClass* GraphClass = GraphClassPtr.Get())
+				{
+					for (const FMetasoundFrontendClassInput& Input : GraphClass->GetDefaultInterface().Inputs)
+					{
+						if (Input.Name == InName)
+						{
+							if (const IDocumentBuilderRegistry* BuilderRegistry = IDocumentBuilderRegistry::Get())
+							{
+								if (FMetaSoundFrontendDocumentBuilder* Builder = BuilderRegistry->FindBuilder(GraphClass->Metadata.GetClassName(), { }))
+								{
+									const int32 Index = Builder->GetConstDocumentChecked().RootGraph.GetDefaultInterface().Inputs.Find(Input);
+									FDocumentModifyDelegates& DocumentDelegates = Builder->GetDocumentDelegates();
+									DocumentDelegates.InterfaceDelegates.OnInputDescriptionChanged.Broadcast(Index);
+								}
+							}
+						}
+					}
+				}
 			}
 		}
 
@@ -1404,6 +1458,24 @@ PRAGMA_ENABLE_DEPRECATION_WARNINGS
 			if (FMetasoundFrontendClassOutput* Desc = FindOutputDescriptionWithName(InName))
 			{
 				Desc->Metadata.SetDescription(InDescription);
+				if (const FMetasoundFrontendGraphClass* GraphClass = GraphClassPtr.Get())
+				{
+					for (const FMetasoundFrontendClassOutput& Output : GraphClass->GetDefaultInterface().Outputs)
+					{
+						if (Output.Name == InName)
+						{
+							if (const IDocumentBuilderRegistry* BuilderRegistry = IDocumentBuilderRegistry::Get())
+							{
+								if (FMetaSoundFrontendDocumentBuilder* Builder = BuilderRegistry->FindBuilder(GraphClass->Metadata.GetClassName(), { }))
+								{
+									const int32 Index = Builder->GetConstDocumentChecked().RootGraph.GetDefaultInterface().Outputs.Find(Output);
+									FDocumentModifyDelegates& DocumentDelegates = Builder->GetDocumentDelegates();
+									DocumentDelegates.InterfaceDelegates.OnOutputDescriptionChanged.Broadcast(Index);
+								}
+							}
+						}
+					}
+				}
 			}
 		}
 
@@ -1418,6 +1490,15 @@ PRAGMA_ENABLE_DEPRECATION_WARNINGS
 					if (Input.Name == InName)
 					{
 						Input.Metadata.SortOrderIndex = InSortOrderIndex;
+						if (const IDocumentBuilderRegistry* BuilderRegistry = IDocumentBuilderRegistry::Get())
+						{
+							if (FMetaSoundFrontendDocumentBuilder* Builder = BuilderRegistry->FindBuilder(GraphClass->Metadata.GetClassName(), { }))
+							{
+								const int32 Index = Builder->GetConstDocumentChecked().RootGraph.GetDefaultInterface().Inputs.Find(Input);
+								FDocumentModifyDelegates& DocumentDelegates = Builder->GetDocumentDelegates();
+								DocumentDelegates.InterfaceDelegates.OnInputSortOrderIndexChanged.Broadcast(Index);
+							}
+						}
 					}
 					Style.DefaultSortOrder.Add(Input.Metadata.SortOrderIndex);
 				}
@@ -1436,6 +1517,15 @@ PRAGMA_ENABLE_DEPRECATION_WARNINGS
 					if (Output.Name == InName)
 					{
 						Output.Metadata.SortOrderIndex = InSortOrderIndex;
+						if (const IDocumentBuilderRegistry* BuilderRegistry = IDocumentBuilderRegistry::Get())
+						{
+							if (FMetaSoundFrontendDocumentBuilder* Builder = BuilderRegistry->FindBuilder(GraphClass->Metadata.GetClassName(), { }))
+							{
+								const int32 Index = Builder->GetConstDocumentChecked().RootGraph.GetDefaultInterface().Outputs.Find(Output);
+								FDocumentModifyDelegates& DocumentDelegates = Builder->GetDocumentDelegates();
+								DocumentDelegates.InterfaceDelegates.OnOutputSortOrderIndexChanged.Broadcast(Index);
+							}
+						}
 					}
 					Style.DefaultSortOrder.Add(Output.Metadata.SortOrderIndex);
 				}
diff --git a/Engine/Plugins/Runtime/Metasound/Source/MetasoundFrontend/Public/MetasoundFrontendDocumentModifyDelegates.h b/Engine/Plugins/Runtime/Metasound/Source/MetasoundFrontend/Public/MetasoundFrontendDocumentModifyDelegates.h
--- a/Engine/Plugins/Runtime/Metasound/Source/MetasoundFrontend/Public/MetasoundFrontendDocumentModifyDelegates.h
+++ b/Engine/Plugins/Runtime/Metasound/Source/MetasoundFrontend/Public/MetasoundFrontendDocumentModifyDelegates.h
@@ -9,6 +9,7 @@
 
 
 // TODO: Move these to namespace
+DECLARE_MULTICAST_DELEGATE(FOnMetaSoundFrontendDocumentMutateMetadata);
 DECLARE_MULTICAST_DELEGATE_OneParam(FOnMetaSoundFrontendDocumentMutateArray, int32 /* Index */);
 DECLARE_MULTICAST_DELEGATE_OneParam(FOnMetaSoundFrontendDocumentMutateInterfaceArray, const FMetasoundFrontendInterface& /* Interface */);
 DECLARE_MULTICAST_DELEGATE_TwoParams(FOnMetaSoundFrontendDocumentRemoveSwappingArray, int32 /* Index */, int32 /* LastIndex */);
@@ -40,24 +41,31 @@ namespace Metasound::Frontend
 		FOnMetaSoundFrontendDocumentMutateInterfaceArray OnRemovingInterface;
 
 		FOnMetaSoundFrontendDocumentMutateArray OnInputAdded;
-
-// Currently only used in editor contexts  
-// so only enabled there to avoid unnecessary delegate overhead
-// but may be changed in the future 
-#if WITH_EDITOR
+		FOnMetaSoundFrontendDocumentMutateArray OnRemovingInput;
+		FOnMetaSoundFrontendDocumentRenameVertex OnInputNameChanged;
 		FOnMetaSoundFrontendDocumentMutateArray OnInputDataTypeChanged;
-#endif // WITH_EDITOR
 		FOnMetaSoundFrontendDocumentMutateArray OnInputDefaultChanged;
+		FOnMetaSoundFrontendDocumentMutateArray OnInputIsConstructorPinChanged;
+		FOnMetaSoundFrontendDocumentMutateArray OnInputInheritsDefaultChanged;
 
-		FOnMetaSoundFrontendDocumentRenameVertex OnInputNameChanged;
-		FOnMetaSoundFrontendDocumentMutateArray OnRemovingInput;
+#if WITH_EDITOR
+		FOnMetaSoundFrontendDocumentMutateArray OnInputDisplayNameChanged;
+		FOnMetaSoundFrontendDocumentMutateArray OnInputDescriptionChanged;
+		FOnMetaSoundFrontendDocumentMutateArray OnInputSortOrderIndexChanged;
+		FOnMetaSoundFrontendDocumentMutateArray OnInputIsAdvancedDisplayChanged;
+#endif // WITH_EDITOR
 
 		FOnMetaSoundFrontendDocumentMutateArray OnOutputAdded;
-#if WITH_EDITOR
+		FOnMetaSoundFrontendDocumentMutateArray OnRemovingOutput;
+		FOnMetaSoundFrontendDocumentRenameVertex OnOutputNameChanged;
 		FOnMetaSoundFrontendDocumentMutateArray OnOutputDataTypeChanged;
+		FOnMetaSoundFrontendDocumentMutateArray OnOutputIsConstructorPinChanged;
+#if WITH_EDITOR
+		FOnMetaSoundFrontendDocumentMutateArray OnOutputDisplayNameChanged;
+		FOnMetaSoundFrontendDocumentMutateArray OnOutputDescriptionChanged;
+		FOnMetaSoundFrontendDocumentMutateArray OnOutputSortOrderIndexChanged;
+		FOnMetaSoundFrontendDocumentMutateArray OnOutputIsAdvancedDisplayChanged;
 #endif // WITH_EDITOR
-		FOnMetaSoundFrontendDocumentRenameVertex OnOutputNameChanged;
-		FOnMetaSoundFrontendDocumentMutateArray OnRemovingOutput;
 	};
 
 	struct FNodeModifyDelegates
@@ -84,7 +92,7 @@ namespace Metasound::Frontend
 		UE_API FDocumentModifyDelegates& operator=(const FDocumentModifyDelegates&);
 		UE_API FDocumentModifyDelegates& operator=(FDocumentModifyDelegates&&);
 
-
+		FOnMetaSoundFrontendDocumentMutateMetadata OnDocumentMetadataChanged;
 		FOnMetaSoundFrontendDocumentMutateArray OnDependencyAdded;
 		FOnMetaSoundFrontendDocumentRemoveSwappingArray OnRemoveSwappingDependency;
 		FOnMetaSoundFrontendDocumentRenameClass OnRenamingDependencyClass;
