Fix newly created MetaSound presets not being presets on editor restart if not saved

diff --git a/Engine/Plugins/Runtime/Metasound/Source/MetasoundEditor/Private/MetasoundAssetDefinitions.cpp b/Engine/Plugins/Runtime/Metasound/Source/MetasoundEditor/Private/MetasoundAssetDefinitions.cpp
--- a/Engine/Plugins/Runtime/Metasound/Source/MetasoundEditor/Private/MetasoundAssetDefinitions.cpp
+++ b/Engine/Plugins/Runtime/Metasound/Source/MetasoundEditor/Private/MetasoundAssetDefinitions.cpp
@@ -9,8 +9,10 @@
 #include "ContentBrowserMenuContexts.h"
 #include "ContentBrowserModule.h"
 #include "Editor/EditorPerProjectUserSettings.h"
+#include "FileHelpers.h"
 #include "IAssetTools.h"
 #include "IContentBrowserSingleton.h"
+#include "ISourceControlModule.h"
 #include "Metasound.h"
 #include "MetasoundAssetBase.h"
 #include "MetasoundAssetManager.h"
@@ -165,6 +167,20 @@ namespace Metasound::Editor
 				{
 					FContentBrowserModule& ContentBrowserModule = FModuleManager::LoadModuleChecked<FContentBrowserModule>("ContentBrowser");
 					ContentBrowserModule.Get().SyncBrowserToAssets(ObjectsToSync);
+					
+					// Save happened in Create/Duplicate asset calls above but
+					// need to save again after register graph with frontend (which applies preset transform)
+					for (UObject* InObject : ObjectsToSync)
+					{
+						if (ISourceControlModule::Get().IsEnabled())
+						{
+							constexpr bool bCheckDirty = false;
+							constexpr bool bPromptToSave = false;
+							TArray<UPackage*> OutermostPackagesToSave;
+							OutermostPackagesToSave.Add(InObject->GetOutermost());
+							FEditorFileUtils::PromptForCheckoutAndSave(OutermostPackagesToSave, bCheckDirty, bPromptToSave);
+						}
+					};
 				}
 			}
 		}
