Metasound change BasicNode ctor arg types and NodeFacade plumbing

Take args by const reference; this then avoids copies along the forwarding chain into the BasicNode ctor.

diff --git a/Engine/Plugins/Runtime/Metasound/Source/MetasoundGraphCore/Private/MetasoundBasicNode.cpp b/Engine/Plugins/Runtime/Metasound/Source/MetasoundGraphCore/Private/MetasoundBasicNode.cpp
--- a/Engine/Plugins/Runtime/Metasound/Source/MetasoundGraphCore/Private/MetasoundBasicNode.cpp
+++ b/Engine/Plugins/Runtime/Metasound/Source/MetasoundGraphCore/Private/MetasoundBasicNode.cpp
@@ -7,9 +7,9 @@
 
 namespace Metasound
 {
-	FBasicNode::FBasicNode(FNodeData InNodeData, TSharedRef<const FNodeClassMetadata> InClassMetadata)
-	: NodeData(MoveTemp(InNodeData))
-	, ClassMetadata(MoveTemp(InClassMetadata))
+	FBasicNode::FBasicNode(const FNodeData& InNodeData, const TSharedRef<const FNodeClassMetadata>& InClassMetadata)
+	: NodeData(InNodeData)
+	, ClassMetadata(InClassMetadata)
 	{
 	}
 
diff --git a/Engine/Plugins/Runtime/Metasound/Source/MetasoundGraphCore/Public/MetasoundBasicNode.h b/Engine/Plugins/Runtime/Metasound/Source/MetasoundGraphCore/Public/MetasoundBasicNode.h
--- a/Engine/Plugins/Runtime/Metasound/Source/MetasoundGraphCore/Public/MetasoundBasicNode.h
+++ b/Engine/Plugins/Runtime/Metasound/Source/MetasoundGraphCore/Public/MetasoundBasicNode.h
@@ -17,7 +17,7 @@ namespace Metasound
 			 * @param InNodeData - Information unique to this node instance. 
 			 * @param InClassMetadata - Shared object representing the node class.
 			 */
-			UE_API FBasicNode(FNodeData InNodeData, TSharedRef<const FNodeClassMetadata> InClassMetadata);
+			UE_API FBasicNode(const FNodeData& InNodeData, const TSharedRef<const FNodeClassMetadata>& InClassMetadata);
 
 			virtual ~FBasicNode() = default;
 
diff --git a/Engine/Plugins/Runtime/Metasound/Source/MetasoundGraphCore/Public/MetasoundFacade.h b/Engine/Plugins/Runtime/Metasound/Source/MetasoundGraphCore/Public/MetasoundFacade.h
--- a/Engine/Plugins/Runtime/Metasound/Source/MetasoundGraphCore/Public/MetasoundFacade.h
+++ b/Engine/Plugins/Runtime/Metasound/Source/MetasoundGraphCore/Public/MetasoundFacade.h
@@ -95,6 +95,7 @@ namespace Metasound
 		class FFactory : public IOperatorFactory
 		{
 			using FCreateOperatorFunction = TFunction<TUniquePtr<IOperator>(const FBuildOperatorParams& InParams, FBuildResults& OutResults)>;
+			friend class FNodeFacade;
 
 			public:
 				UE_API FFactory(FCreateOperatorFunction InCreateFunc);
@@ -120,20 +121,19 @@ namespace Metasound
 		 */
 		template<typename OperatorType>
 		FNodeFacade(const FVertexName& InInstanceName, const FGuid& InInstanceID, TFacadeOperatorClass<OperatorType> OperatorClass)
-		:	FNodeFacade(FNodeData(InInstanceName, InInstanceID, OperatorType::GetNodeInfo().DefaultInterface), MakeShared<FNodeClassMetadata>(OperatorType::GetNodeInfo()), OperatorClass)
+		: FNodeFacade(InInstanceName, InInstanceID, OperatorType::GetNodeInfo(), &OperatorType::CreateOperator)
 		{
 		}
 
 		template<typename OperatorType>
-		FNodeFacade(FNodeInitData InInitData, TFacadeOperatorClass<OperatorType> OperatorClass)
-		:	FNodeFacade(FNodeData(InInitData.InstanceName, InInitData.InstanceID, OperatorType::GetNodeInfo().DefaultInterface), MakeShared<FNodeClassMetadata>(OperatorType::GetNodeInfo()), OperatorClass)
+		FNodeFacade(const FNodeInitData& InInitData, TFacadeOperatorClass<OperatorType> OperatorClass)
+		: FNodeFacade(InInitData.InstanceName, InInitData.InstanceID, OperatorType::GetNodeInfo(), &OperatorType::CreateOperator)
 		{
 		}
 
 		template<typename OperatorType>
-		FNodeFacade(FNodeData InNodeData, TSharedRef<const FNodeClassMetadata> InClassMetadata, TFacadeOperatorClass<OperatorType> OperatorClass)
-		: FBasicNode(MoveTemp(InNodeData), MoveTemp(InClassMetadata))
-		, Factory(MakeShared<FFactory, ESPMode::ThreadSafe>(&OperatorType::CreateOperator))
+		FNodeFacade(const FNodeData& InNodeData, const TSharedRef<const FNodeClassMetadata>& InClassMetadata, TFacadeOperatorClass<OperatorType> OperatorClass)
+		: FNodeFacade(InNodeData, InClassMetadata, &OperatorType::CreateOperator)
 		{
 		}
 
@@ -143,6 +143,17 @@ namespace Metasound
 		UE_API virtual FOperatorFactorySharedRef GetDefaultOperatorFactory() const override;
 
 	private:
+		FORCENOINLINE FNodeFacade(const FVertexName& InInstanceName, const FGuid& InInstanceID, const FNodeClassMetadata& InNodeInfo, FFactory::FCreateOperatorFunction InCreate)
+		: FBasicNode(FNodeData(InInstanceName, InInstanceID, InNodeInfo.DefaultInterface), MakeShared<FNodeClassMetadata>(InNodeInfo))
+		, Factory(MakeShared<FFactory, ESPMode::ThreadSafe>(InCreate))
+		{
+		}
+
+		FORCENOINLINE FNodeFacade(const FNodeData& InNodeData, const TSharedRef<const FNodeClassMetadata>& InClassMetadata, FFactory::FCreateOperatorFunction InCreate)
+		: FBasicNode(InNodeData, InClassMetadata)
+		, Factory(MakeShared<FFactory, ESPMode::ThreadSafe>(InCreate))
+		{
+		}
 
 		TSharedRef<IOperatorFactory, ESPMode::ThreadSafe> Factory;
 	};
@@ -161,12 +172,12 @@ namespace Metasound
 		}
 
 		TNodeFacade(const FNodeInitData& InInitData)
-			: FNodeFacade(FNodeData(InInitData.InstanceName, InInitData.InstanceID, OperatorType::GetNodeInfo().DefaultInterface), MakeShared<FNodeClassMetadata>(OperatorType::GetNodeInfo()), TFacadeOperatorClass<OperatorType>())
+			: FNodeFacade(InInitData, TFacadeOperatorClass<OperatorType>())
 		{
 		}
 
-		TNodeFacade(FNodeData InNodeData, TSharedRef<const FNodeClassMetadata> InClassMetadata)
-			: FNodeFacade(MoveTemp(InNodeData), MoveTemp(InClassMetadata), TFacadeOperatorClass<OperatorType>())
+		TNodeFacade(const FNodeData& InNodeData, const TSharedRef<const FNodeClassMetadata>& InClassMetadata)
+			: FNodeFacade(InNodeData, InClassMetadata, TFacadeOperatorClass<OperatorType>())
 		{
 		}
 
