Fixes for referenced MetaSound data type updating and:
1. Default metadata not synchronizing appropriately on preset.
2. input data type updating and preset not resetting 'Override Interhited Default' accordingly.

#rnx

diff --git a/Engine/Plugins/Runtime/Metasound/Source/MetasoundEditor/Private/MetasoundEditorSubsystem.cpp b/Engine/Plugins/Runtime/Metasound/Source/MetasoundEditor/Private/MetasoundEditorSubsystem.cpp
--- a/Engine/Plugins/Runtime/Metasound/Source/MetasoundEditor/Private/MetasoundEditorSubsystem.cpp
+++ b/Engine/Plugins/Runtime/Metasound/Source/MetasoundEditor/Private/MetasoundEditorSubsystem.cpp
@@ -176,8 +176,13 @@ bool UMetaSoundEditorSubsystem::BindMemberMetadata(
 	{
 		if (UMetaSoundFrontendMemberMetadata* Literal = Builder.FindMemberMetadata(MemberID))
 		{
-			InMember.Literal = CastChecked<UMetasoundEditorGraphMemberDefaultLiteral>(Literal);
-			return false;
+			const UClass* ExistingClass = Literal->GetClass();
+			check(ExistingClass);
+			if (ExistingClass->IsChildOf(LiteralClass))
+			{
+				InMember.Literal = CastChecked<UMetasoundEditorGraphMemberDefaultLiteral>(Literal);
+				return false;
+			}
 		}
 
 		NewLiteral = CreateMemberMetadata(Builder, InMember.GetMemberName(), LiteralClass);
diff --git a/Engine/Plugins/Runtime/Metasound/Source/MetasoundFrontend/Private/MetasoundFrontendTransform.cpp b/Engine/Plugins/Runtime/Metasound/Source/MetasoundFrontend/Private/MetasoundFrontendTransform.cpp
--- a/Engine/Plugins/Runtime/Metasound/Source/MetasoundFrontend/Private/MetasoundFrontendTransform.cpp
+++ b/Engine/Plugins/Runtime/Metasound/Source/MetasoundFrontend/Private/MetasoundFrontendTransform.cpp
@@ -891,6 +891,7 @@ PRAGMA_ENABLE_DEPRECATION_WARNINGS
 			TArray<FMetasoundFrontendClassInput> ClassInputs;
 
 			FConstGraphHandle ReferencedGraph = ReferencedDocument->GetRootGraph();
+			TSet<FName> InputsToInherit = InPresetGraph->GetInputsInheritingDefault();
 
 			// Iterate through all input nodes of referenced graph
 			ReferencedGraph->IterateConstNodes([&](FConstNodeHandle InputNode)
@@ -952,6 +953,10 @@ PRAGMA_ENABLE_DEPRECATION_WARNINGS
 					else
 					{
 						InheritDefaultsFromGraph(ReferencedGraph);
+
+						// Add this to the list of inheriting as the
+						// type no longer matches or it no longer exists
+						InputsToInherit.Add(NodeName);
 					}
 
 					ClassInputs.Add(MoveTemp(ClassInput));
@@ -961,15 +966,15 @@ PRAGMA_ENABLE_DEPRECATION_WARNINGS
 			// Fill new managed inputs set with names of all class inputs & if the old input was explicitly not
 			// marked as a managed input, then remove it from the new managed inputs if found.
 			OutInputsInheritingDefault.Reset();
-			Algo::Transform(ClassInputs, OutInputsInheritingDefault, [](const FMetasoundFrontendClassInput& Input) { return Input.Name; });
-			const TSet<FName>& InputsInheritingDefault = InPresetGraph->GetInputsInheritingDefault();
-			InPresetGraph->IterateConstNodes([&InputsInheritingDefault, &OutInputsInheritingDefault](FConstNodeHandle Input)
-			{
-				if (!InputsInheritingDefault.Contains(Input->GetNodeName()))
+			Algo::TransformIf(ClassInputs, OutInputsInheritingDefault,
+				[&InputsToInherit](const FMetasoundFrontendClassInput& Input)
 				{
-					OutInputsInheritingDefault.Remove(Input->GetNodeName());
-				}
-			}, EMetasoundFrontendClassType::Input);
+					return InputsToInherit.Contains(Input.Name);
+				},
+				[](const FMetasoundFrontendClassInput& Input)
+				{
+					return Input.Name;
+				});
 
 			return ClassInputs;
 		}
