Add BP equality operator for MetaSound literals

diff --git a/Engine/Plugins/Runtime/Metasound/Source/MetasoundEngine/Private/MetasoundFrontendLiteralBlueprintAccess.cpp b/Engine/Plugins/Runtime/Metasound/Source/MetasoundEngine/Private/MetasoundFrontendLiteralBlueprintAccess.cpp
--- a/Engine/Plugins/Runtime/Metasound/Source/MetasoundEngine/Private/MetasoundFrontendLiteralBlueprintAccess.cpp
+++ b/Engine/Plugins/Runtime/Metasound/Source/MetasoundEngine/Private/MetasoundFrontendLiteralBlueprintAccess.cpp
@@ -59,6 +59,11 @@ EMetasoundFrontendLiteralType UMetasoundFrontendLiteralBlueprintAccess::GetType(
 	return Literal.GetType();
 }
 
+bool UMetasoundFrontendLiteralBlueprintAccess::EqualEqual_MetaSoundLiteral(const FMetasoundFrontendLiteral& LiteralA, const FMetasoundFrontendLiteral& LiteralB)
+{
+	return LiteralA.IsEqual(LiteralB);
+}
+
 FMetasoundFrontendLiteral UMetasoundFrontendLiteralBlueprintAccess::CreateBoolMetaSoundLiteral(bool Value)
 {
 	return MetasoundFrontendLiteralBlueprintAccessPrivate::CreatePODMetaSoundLiteral(Value);
diff --git a/Engine/Plugins/Runtime/Metasound/Source/MetasoundEngine/Public/MetasoundFrontendLiteralBlueprintAccess.h b/Engine/Plugins/Runtime/Metasound/Source/MetasoundEngine/Public/MetasoundFrontendLiteralBlueprintAccess.h
--- a/Engine/Plugins/Runtime/Metasound/Source/MetasoundEngine/Public/MetasoundFrontendLiteralBlueprintAccess.h
+++ b/Engine/Plugins/Runtime/Metasound/Source/MetasoundEngine/Public/MetasoundFrontendLiteralBlueprintAccess.h
@@ -21,6 +21,9 @@ class UMetasoundFrontendLiteralBlueprintAccess : public UBlueprintFunctionLibrar
 
 	UFUNCTION(BlueprintPure, Category = "Audio|MetaSound", meta = (DisplayName = "Get MetaSound Literal Type"))
 	static UPARAM(DisplayName = "Type") EMetasoundFrontendLiteralType GetType(const FMetasoundFrontendLiteral& Literal);
+	
+	UFUNCTION(BlueprintPure, meta=(DisplayName = "Equal (MetaSound Literal)", CompactNodeTitle = "==", Keywords = "== equal"), Category="Audio|MetaSound")
+	static bool EqualEqual_MetaSoundLiteral(const FMetasoundFrontendLiteral& LiteralA, const FMetasoundFrontendLiteral& LiteralB);
 
 	// Literal creation
 
