Remove forward declare of unused, deprecated struct
#rnx

diff --git a/Engine/Plugins/Runtime/Metasound/Source/MetasoundEngine/Public/MetasoundSource.h b/Engine/Plugins/Runtime/Metasound/Source/MetasoundEngine/Public/MetasoundSource.h
--- a/Engine/Plugins/Runtime/Metasound/Source/MetasoundEngine/Public/MetasoundSource.h
+++ b/Engine/Plugins/Runtime/Metasound/Source/MetasoundEngine/Public/MetasoundSource.h
@@ -41,7 +41,6 @@ namespace Audio
 
 namespace Metasound
 {
-	struct FMetasoundGeneratorInitParams;
 	class FMetasoundGenerator;
 
 	namespace DynamicGraph
