Fix MetaSound presets created with the builder BP functions being invalid on restart if not saved

diff --git a/Engine/Plugins/Runtime/Metasound/Source/MetasoundEditor/Private/MetasoundEditorSubsystem.cpp b/Engine/Plugins/Runtime/Metasound/Source/MetasoundEditor/Private/MetasoundEditorSubsystem.cpp
--- a/Engine/Plugins/Runtime/Metasound/Source/MetasoundEditor/Private/MetasoundEditorSubsystem.cpp
+++ b/Engine/Plugins/Runtime/Metasound/Source/MetasoundEditor/Private/MetasoundEditorSubsystem.cpp
@@ -2,7 +2,9 @@
 #include "MetasoundEditorSubsystem.h"
 
 #include "AudioPropertiesSheetAssetBase.h"
+#include "FileHelpers.h"
 #include "IAssetTools.h"
+#include "ISourceControlModule.h"
 #include "Kismet/KismetSystemLibrary.h"
 #include "MetasoundDocumentBuilderRegistry.h"
 #include "MetasoundDocumentInterface.h"
@@ -106,6 +108,17 @@ TScriptInterface<IMetaSoundDocumentInterface> UMetaSoundEditorSubsystem::BuildTo
 				}
 			}
 
+			// Save happened in Create/Duplicate asset calls above but
+			// need to save again after building (which can apply preset transform)
+			if (ISourceControlModule::Get().IsEnabled())
+			{
+				constexpr bool bCheckDirty = false;
+				constexpr bool bPromptToSave = false;
+				TArray<UPackage*> OutermostPackagesToSave;
+				OutermostPackagesToSave.Add(NewMetaSound->GetOutermost());
+				FEditorFileUtils::PromptForCheckoutAndSave(OutermostPackagesToSave, bCheckDirty, bPromptToSave);
+			}
+
 			if (!bWasRooted)
 			{
 				InBuilder->RemoveFromRoot();
