Added interface check

#rnx

diff --git a/Engine/Plugins/Runtime/Metasound/Source/MetasoundEngine/Private/MetasoundBuilderSubsystem.cpp b/Engine/Plugins/Runtime/Metasound/Source/MetasoundEngine/Private/MetasoundBuilderSubsystem.cpp
--- a/Engine/Plugins/Runtime/Metasound/Source/MetasoundEngine/Private/MetasoundBuilderSubsystem.cpp
+++ b/Engine/Plugins/Runtime/Metasound/Source/MetasoundEngine/Private/MetasoundBuilderSubsystem.cpp
@@ -1202,6 +1202,12 @@ UMetaSoundBuilderBase* UMetaSoundBuilderSubsystem::FindParentBuilderOfPreset(con
 	using namespace Metasound::Engine;
 	using namespace Metasound::Frontend;
 
+	if (!InMetaSoundPreset.GetObject() || !InMetaSoundPreset.GetObject()->Implements<UMetaSoundDocumentInterface>())
+	{
+		UE_LOG(LogMetaSound, Error, TEXT("FindParentBuilderOfPreset called with object that is not a MetaSound asset"));
+		return nullptr;
+	}
+
 	// Find the builder for the given asset to check if it is a preset.
 	const FDocumentBuilderRegistry& Registry = FDocumentBuilderRegistry::GetChecked();
 	const UMetaSoundBuilderBase* Builder = Registry.FindBuilderObject(InMetaSoundPreset);
