NodeConfig Live Audition support

diff --git a/Engine/Plugins/Runtime/Metasound/Source/MetasoundEditor/Private/MetasoundEditorGraph.cpp b/Engine/Plugins/Runtime/Metasound/Source/MetasoundEditor/Private/MetasoundEditorGraph.cpp
--- a/Engine/Plugins/Runtime/Metasound/Source/MetasoundEditor/Private/MetasoundEditorGraph.cpp
+++ b/Engine/Plugins/Runtime/Metasound/Source/MetasoundEditor/Private/MetasoundEditorGraph.cpp
@@ -2149,13 +2149,15 @@ UObject& UMetasoundEditorGraph::GetMetasoundChecked() const
 	return *ParentMetasound;
 }
 
-void UMetasoundEditorGraph::RegisterGraphWithFrontend()
+void UMetasoundEditorGraph::RegisterGraphWithFrontend(Metasound::Frontend::FMetaSoundAssetRegistrationOptions* RegOptions)
 {
 	using namespace Metasound::Editor;
 
 	if (UObject* ParentMetasound = GetOutermostObject())
 	{
-		FGraphBuilder::RegisterGraphWithFrontend(*ParentMetasound);
+		FGraphBuilder::RegisterGraphWithFrontend(*ParentMetasound, RegOptions
+			? *RegOptions
+			: Metasound::Editor::FGraphBuilder::GetDefaultRegistrationOptions());
 	}
 }
 
diff --git a/Engine/Plugins/Runtime/Metasound/Source/MetasoundEditor/Private/MetasoundEditorGraphSchema.cpp b/Engine/Plugins/Runtime/Metasound/Source/MetasoundEditor/Private/MetasoundEditorGraphSchema.cpp
--- a/Engine/Plugins/Runtime/Metasound/Source/MetasoundEditor/Private/MetasoundEditorGraphSchema.cpp
+++ b/Engine/Plugins/Runtime/Metasound/Source/MetasoundEditor/Private/MetasoundEditorGraphSchema.cpp
@@ -1428,14 +1428,6 @@ UEdGraphNode* FMetasoundGraphSchemaAction_NewReroute::PerformAction(UEdGraph* Pa
 	ParentMetasound.Modify();
 	ParentGraph->Modify();
 
-	FMetasoundAssetBase* MetaSoundAsset = Metasound::IMetasoundUObjectRegistry::Get().GetObjectAsAssetBase(&ParentMetasound);
-	check(MetaSoundAsset);
-
-	const INodeTemplate* RerouteTemplate = INodeTemplateRegistry::Get().FindTemplate(FRerouteNodeTemplate::ClassName);
-	check(RerouteTemplate);
-
-	const FMetaSoundFrontendDocumentBuilder& DocBuilder = IDocumentBuilderRegistry::GetChecked().FindOrBeginBuilding(&ParentMetasound);
-
 	// Provided 'FromPin' is what to connect to, so if its an input, its the output of the generated node needs to
 	// match the from pin and vice versa.
 	const FName FromVertexDataType = FGraphBuilder::GetPinDataType(FromPin);
@@ -1449,24 +1441,33 @@ UEdGraphNode* FMetasoundGraphSchemaAction_NewReroute::PerformAction(UEdGraph* Pa
 		Params.InputsToConnect.Add(FromVertexDataType);
 	}
 
-	FConstNodeHandle NodeHandle = MetaSoundAsset->GetRootGraphHandle()->AddTemplateNode(*RerouteTemplate, MoveTemp(Params));
+	FMetaSoundFrontendDocumentBuilder& DocBuilder = MetaSoundGraph->GetBuilderChecked().GetBuilder();
+	const INodeTemplate* RerouteTemplate = INodeTemplateRegistry::Get().FindTemplate(FRerouteNodeTemplate::ClassName);
+	check(RerouteTemplate);
 
-	if (UMetasoundEditorGraphExternalNode* NewGraphNode = FGraphBuilder::AddTemplateNode(ParentMetasound, NodeHandle->GetID(), NodeHandle->GetClassMetadata(), bSelectNewNode))
+	if (const FMetasoundFrontendNode* TemplateNode = DocBuilder.AddNodeByTemplate(*RerouteTemplate, MoveTemp(Params)))
 	{
-		NewGraphNode->Modify();
-		NewGraphNode->UpdateFrontendNodeLocation(FDeprecateSlateVector2D(Location));
-		NewGraphNode->SyncLocationFromFrontendNode();
+		if (UMetasoundEditorGraphExternalNode* NewGraphNode = FGraphBuilder::AddTemplateNode(
+			ParentMetasound,
+			TemplateNode->GetID(),
+			RerouteTemplate->GetFrontendClass().Metadata,
+			bSelectNewNode))
+		{
+			NewGraphNode->Modify();
+			NewGraphNode->UpdateFrontendNodeLocation(FDeprecateSlateVector2D(Location));
+			NewGraphNode->SyncLocationFromFrontendNode();
 
-		SchemaPrivate::TryConnectNewNodeToMatchingDataTypePin(*NewGraphNode, FromPin);
-		MetaSoundAsset->GetModifyContext().AddNodeIDsModified({ NewGraphNode->GetNodeID()});
+			SchemaPrivate::TryConnectNewNodeToMatchingDataTypePin(*NewGraphNode, FromPin);
+			DocBuilder.GetConstDocumentChecked().Metadata.ModifyContext.AddNodeIDsModified({ NewGraphNode->GetNodeID() });
 
-		TSharedPtr<FEditor> ParentEditor = FGraphBuilder::GetEditorForMetasound(ParentMetasound);
-		if (ParentEditor.IsValid() && bSelectNewNode)
-		{
-			ParentEditor->ClearSelectionAndSelectNode(NewGraphNode);
-		}
+			TSharedPtr<FEditor> ParentEditor = FGraphBuilder::GetEditorForMetasound(ParentMetasound);
+			if (ParentEditor.IsValid() && bSelectNewNode)
+			{
+				ParentEditor->ClearSelectionAndSelectNode(NewGraphNode);
+			}
 
-		return NewGraphNode;
+			return NewGraphNode;
+		}
 	}
 
 	return nullptr;
diff --git a/Engine/Plugins/Runtime/Metasound/Source/MetasoundEditor/Public/MetasoundEditorGraph.h b/Engine/Plugins/Runtime/Metasound/Source/MetasoundEditor/Public/MetasoundEditorGraph.h
--- a/Engine/Plugins/Runtime/Metasound/Source/MetasoundEditor/Public/MetasoundEditorGraph.h
+++ b/Engine/Plugins/Runtime/Metasound/Source/MetasoundEditor/Public/MetasoundEditorGraph.h
@@ -523,7 +523,7 @@ public:
 	UE_API bool IsEditable() const;
 
 	// UMetasoundEditorGraphBase Implementation
-	UE_API virtual void RegisterGraphWithFrontend() override;
+	UE_API virtual void RegisterGraphWithFrontend(Metasound::Frontend::FMetaSoundAssetRegistrationOptions* RegOptions = nullptr) override;
 	UE_API virtual void MigrateEditorDocumentData(FMetaSoundFrontendDocumentBuilder& OutBuilder) override;
 
 private:
diff --git a/Engine/Plugins/Runtime/Metasound/Source/MetasoundEngine/Private/MetasoundBuilderSubsystem.cpp b/Engine/Plugins/Runtime/Metasound/Source/MetasoundEngine/Private/MetasoundBuilderSubsystem.cpp
--- a/Engine/Plugins/Runtime/Metasound/Source/MetasoundEngine/Private/MetasoundBuilderSubsystem.cpp
+++ b/Engine/Plugins/Runtime/Metasound/Source/MetasoundEngine/Private/MetasoundBuilderSubsystem.cpp
@@ -69,22 +69,21 @@ namespace Metasound::Engine
 			return Literal;
 		}
 
-		TOptional<FMetasoundFrontendLiteral> TryResolveNodeInputDefault(const FMetaSoundFrontendDocumentBuilder& Builder, const FGuid& InNodeID, FName VertexName)
+		const FMetasoundFrontendLiteral* TryResolveNodeInputDefault(const FMetaSoundFrontendDocumentBuilder& Builder, const FGuid& InNodeID, FName VertexName)
 		{
-			TOptional<FMetasoundFrontendLiteral> DefaultLiteral;
 			if (const FMetasoundFrontendVertexLiteral* InputDefault = Builder.FindNodeInputDefault(InNodeID, VertexName))
 			{
-				DefaultLiteral = InputDefault->Value;
+				return &InputDefault->Value;
 			}
 			else if (const TArray<FMetasoundFrontendClassInputDefault>* ClassDefaults = Builder.FindNodeClassInputDefaults(InNodeID, VertexName))
 			{
 				if (const FMetasoundFrontendClassInputDefault* ClassDefault = FindPreferredPage_ThreadSafe(*ClassDefaults))
 				{
-					DefaultLiteral = ClassDefault->Literal;
+					return &ClassDefault->Literal;
 				}
 			}
 
-			return DefaultLiteral;
+			return nullptr;
 		}
 	} // namespace BuilderSubsystemPrivate
 } // namespace Metasound::Engine
@@ -238,9 +237,88 @@ void UMetaSoundSourceBuilder::OnLiveComponentFinished(UAudioComponent* AudioComp
 	if (LiveComponentIDs.IsEmpty())
 	{
 		AudioComponent->OnAudioFinishedNative.Remove(LiveComponentHandle);
+		if (UMetaSoundSource* Source = Cast<UMetaSoundSource>(AudioComponent->Sound))
+		{
+			// This sometimes is called during teardown/resave in the editor so have to
+			// make sure the builder is in a valid state before disabling the dynamic generator.
+			if (Builder.IsValid() && &Builder.CastDocumentObjectChecked<UObject>() == Source)
+			{
+				Source->SetDynamicGeneratorEnabled(false);
+			}
+		}
+	}
+}
+
+bool UMetaSoundSourceBuilder::AddEdgeToTransactor(const FMetasoundFrontendEdge& NewEdge, Metasound::DynamicGraph::FDynamicOperatorTransactor& Transactor) const
+{
+	using namespace Metasound::Frontend;
+
+	const FMetaSoundFrontendDocumentBuilder& DocBuilder = GetConstBuilder();
+
+	const FMetasoundFrontendVertex* FromNodeOutput = DocBuilder.FindNodeOutput(NewEdge.FromNodeID, NewEdge.FromVertexID);
+	const FMetasoundFrontendVertex* ToNodeInput = DocBuilder.FindNodeInput(NewEdge.ToNodeID, NewEdge.ToVertexID);
+	if (FromNodeOutput && ToNodeInput)
+	{
+	#if WITH_EDITOR
+		const FMetasoundFrontendNode* ToNode = DocBuilder.FindNode(NewEdge.ToNodeID);
+		const FMetasoundFrontendNode* FromNode = DocBuilder.FindNode(NewEdge.FromNodeID);
+		if (ToNode && FromNode)
+		{
+			const FMetasoundFrontendClass* ToNodeClass = DocBuilder.FindDependency(ToNode->ClassID);
+			const FMetasoundFrontendClass* FromNodeClass = DocBuilder.FindDependency(FromNode->ClassID);
+
+			// If template node is connecting via an output, then the given edge may be routing data accordingly.
+			if (FromNodeClass->Metadata.GetType() == EMetasoundFrontendClassType::Template)
+			{
+				bool bHandled = false;
+				const FTemplateNodeEdgeTransactionContext TransactionContext
+				{
+					.Builder = DocBuilder,
+					.PageID = DocBuilder.GetBuildPageID(),
+					.FromNode = FromNode,
+					.FromOutput = FromNodeOutput,
+					.ToNode = ToNode,
+					.ToInput = ToNodeInput
+				};
+
+				const FMetasoundFrontendClassName& FromClassName = FromNodeClass->Metadata.GetClassName();
+				if (const INodeTemplate* Template = INodeTemplateRegistry::Get().FindTemplate(FromClassName))
+				{
+					bHandled = Template->OnEdgeAdded(TransactionContext, Transactor);
+				}
+
+				if (bHandled)
+				{
+					return true;
+				}
+			}
 
-		GetMetaSoundSource().SetDynamicGeneratorEnabled(false);
+			// Adding edge to a template node input however can be safely ignored as template nodes are processed out.
+			if (ToNodeClass->Metadata.GetType() == EMetasoundFrontendClassType::Template)
+			{
+				return true;
+			}
+		}
+	#endif // WITH_EDITOR
+
+		Transactor.AddDataEdge(NewEdge.FromNodeID, FromNodeOutput->Name, NewEdge.ToNodeID, ToNodeInput->Name);
+		return true;
 	}
+
+	return false;
+}
+
+bool UMetaSoundSourceBuilder::AddNodeInputLiteralToTransactor(
+	const FMetasoundFrontendNode& Node,
+	const FMetasoundFrontendVertex& Input,
+	const FMetasoundFrontendLiteral& InputDefault,
+	Metasound::DynamicGraph::FDynamicOperatorTransactor& Transactor) const
+{
+	using namespace Metasound;
+	using namespace Metasound::Engine;
+
+	Transactor.SetValue(Node.GetID(), Input.Name, InputDefault.ToLiteral(Input.TypeName), &FMetaSoundFrontendDocumentBuilder::CreateDataReference);
+	return true;
 }
 
 bool UMetaSoundSourceBuilder::ExecuteAuditionableTransaction(FAuditionableTransaction Transaction) const
@@ -369,6 +447,7 @@ void UMetaSoundSourceBuilder::InitTargetPageDelegates(Metasound::Frontend::FDocu
 	OutDocumentDelegates.IterateGraphNodeDelegates([this](FNodeModifyDelegates& NodeDelegates)
 	{
 		NodeDelegates.OnNodeAdded.RemoveAll(this);
+		NodeDelegates.OnNodeConfigurationUpdated.RemoveAll(this);
 		NodeDelegates.OnNodeInputLiteralSet.RemoveAll(this);
 		NodeDelegates.OnRemoveSwappingNode.RemoveAll(this);
 		NodeDelegates.OnRemovingNodeInputLiteral.RemoveAll(this);
@@ -395,6 +474,7 @@ void UMetaSoundSourceBuilder::InitTargetPageDelegates(Metasound::Frontend::FDocu
 
 	FNodeModifyDelegates& NodeDelegates = OutDocumentDelegates.FindNodeDelegatesChecked(TargetPageID);
 	NodeDelegates.OnNodeAdded.AddUObject(this, &UMetaSoundSourceBuilder::OnNodeAdded);
+	NodeDelegates.OnNodeConfigurationUpdated.AddUObject(this, &UMetaSoundSourceBuilder::OnNodeConfigurationUpdated);
 	NodeDelegates.OnNodeInputLiteralSet.AddUObject(this, &UMetaSoundSourceBuilder::OnNodeInputLiteralSet);
 	NodeDelegates.OnRemoveSwappingNode.AddUObject(this, &UMetaSoundSourceBuilder::OnRemoveSwappingNode);
 	NodeDelegates.OnRemovingNodeInputLiteral.AddUObject(this, &UMetaSoundSourceBuilder::OnRemovingNodeInputLiteral);
@@ -419,61 +499,7 @@ void UMetaSoundSourceBuilder::OnEdgeAdded(int32 EdgeIndex) const
 	const FMetasoundFrontendEdge& NewEdge = GetConstTargetPageGraphChecked().Edges[EdgeIndex];
 	ExecuteAuditionableTransaction([this, &NewEdge](Metasound::DynamicGraph::FDynamicOperatorTransactor& Transactor)
 	{
-		using namespace Metasound::Frontend;
-
-		const FMetaSoundFrontendDocumentBuilder& DocBuilder = GetConstBuilder();
-
-		const FMetasoundFrontendVertex* FromNodeOutput = DocBuilder.FindNodeOutput(NewEdge.FromNodeID, NewEdge.FromVertexID);
-		const FMetasoundFrontendVertex* ToNodeInput = DocBuilder.FindNodeInput(NewEdge.ToNodeID, NewEdge.ToVertexID);
-		if (FromNodeOutput && ToNodeInput)
-		{
-#if WITH_EDITOR
-			const FMetasoundFrontendNode* ToNode = DocBuilder.FindNode(NewEdge.ToNodeID);
-			const FMetasoundFrontendNode* FromNode = DocBuilder.FindNode(NewEdge.FromNodeID);
-			if (ToNode && FromNode)
-			{
-				const FMetasoundFrontendClass* ToNodeClass = DocBuilder.FindDependency(ToNode->ClassID);
-				const FMetasoundFrontendClass* FromNodeClass = DocBuilder.FindDependency(FromNode->ClassID);
-
-				// Adding edge to a template node can be safely ignored as template nodes are processed out.
-				// Once template node is connected as a source however, then the given edge can be processed accordingly.
-				if (ToNodeClass->Metadata.GetType() == EMetasoundFrontendClassType::Template)
-				{
-					return true;
-				}
-
-				if (FromNodeClass->Metadata.GetType() == EMetasoundFrontendClassType::Template)
-				{
-					bool bHandled = false;
-					const FTemplateNodeEdgeTransactionContext TransactionContext
-					{
-						.Builder = DocBuilder,
-						.PageID = DocBuilder.GetBuildPageID(),
-						.FromNode = FromNode,
-						.FromOutput = FromNodeOutput,
-						.ToNode = ToNode,
-						.ToInput = ToNodeInput
-					};
-
-					const FMetasoundFrontendClassName& FromClassName = FromNodeClass->Metadata.GetClassName();
-					if (const INodeTemplate* Template = INodeTemplateRegistry::Get().FindTemplate(FromClassName))
-					{
-						bHandled = Template->OnEdgeAdded(TransactionContext, Transactor);
-					}
-
-					if (bHandled)
-					{
-						return true;
-					}
-				}
-			}
-#endif // WITH_EDITOR
-
-			Transactor.AddDataEdge(NewEdge.FromNodeID, FromNodeOutput->Name, NewEdge.ToNodeID, ToNodeInput->Name);
-			return true;
-		}
-
-		return false;
+		return AddEdgeToTransactor(NewEdge, Transactor);
 	});
 }
 
@@ -558,7 +584,7 @@ void UMetaSoundSourceBuilder::OnNodeAdded(int32 NodeIndex) const
 
 		if (!NewNode.IsValid())
 		{
-			UE_LOG(LogMetaSound, Error, TEXT("Builder '%s' failed to create and forward added node '%s' to live update transactor."), *GetName(), *AddedNode.Name.ToString());
+			UE_LOG(LogMetaSound, Error, TEXT("Builder '%s' failed to create and forward added node '%s' to live update transactor."), *Builder.GetDebugName(), *AddedNode.Name.ToString());
 			return false;
 		}
 
@@ -567,6 +593,108 @@ void UMetaSoundSourceBuilder::OnNodeAdded(int32 NodeIndex) const
 	});
 }
 
+void UMetaSoundSourceBuilder::OnNodeConfigurationUpdated(int32 NodeIndex) const
+{
+	using namespace Metasound::DynamicGraph;
+
+	ExecuteAuditionableTransaction([this, NodeIndex](FDynamicOperatorTransactor& Transactor)
+	{
+		using namespace Metasound;
+		using namespace Metasound::Frontend;
+
+		const IMetaSoundDocumentInterface& DocInterface = Builder.GetConstDocumentInterfaceChecked();
+		const FMetasoundFrontendGraphClass& OwningGraphClass = DocInterface.GetConstDocument().RootGraph;
+		const FMetasoundFrontendGraph& OwningGraph = GetConstTargetPageGraphChecked();
+		const FMetasoundFrontendNode& FrontendConfigNode = OwningGraph.Nodes[NodeIndex];
+		const FMetasoundFrontendClass* NodeClass = Builder.FindDependency(FrontendConfigNode.ClassID);
+		checkf(NodeClass, TEXT("Node successfully configured on graph but document is missing associated dependency"));
+
+#if WITH_EDITOR
+		if (NodeClass->Metadata.GetType() == EMetasoundFrontendClassType::Template)
+		{
+			const FMetasoundFrontendClassName& AddedClassName = NodeClass->Metadata.GetClassName();
+			if (const INodeTemplate* Template = INodeTemplateRegistry::Get().FindTemplate(AddedClassName))
+			{
+				const FTemplateNodeTransactionContext TransactionContext { Builder, Builder.GetBuildPageID(), &FrontendConfigNode };
+				const bool bHandled = Template->OnNodeConfigurationUpdated(TransactionContext, Transactor);
+				if (bHandled)
+				{
+					return true;
+				}
+			}
+		}
+#endif // WITH_EDITOR
+
+		FGraphBuilder::FCreateNodeParams ReplacementNodeParams
+		{
+			.FrontendNode = FrontendConfigNode,
+			.FrontendNodeClass = *NodeClass,
+			.OwningFrontendGraph = OwningGraph,
+			.OwningFrontendGraphClass = OwningGraphClass,
+			.OwningAssetPath = DocInterface.GetAssetPathChecked()
+		};
+
+		TUniquePtr<INode> ConfiguredNode = FGraphBuilder::CreateNode(ReplacementNodeParams);
+		if (!ConfiguredNode.IsValid())
+		{
+			UE_LOG(LogMetaSound, Error,
+				TEXT("Builder '%s' failed to create and forward reconfigured node '%s' to live update transactor."),
+				*Builder.GetDebugName(),
+				*FrontendConfigNode.Name.ToString());
+			return false;
+		}
+
+		const bool bRemovedNode = Transactor.RemoveNode(FrontendConfigNode.GetID());
+		if (!bRemovedNode)
+		{
+			return false;
+		}
+
+		const bool bAddedNode = Transactor.AddNode(FrontendConfigNode.GetID(), MoveTemp(ConfiguredNode));
+		if (!bAddedNode)
+		{
+			return false;
+		}
+
+		// Removal and then readding the given node via the transactor indirectly:
+		// a. Removes edges to the node, so they must be readded
+		bool bEdgesAdded = true;
+		{
+			auto AddConnectedEdges = [&](const TArray<FMetasoundFrontendVertex>& InVertices)
+				{
+					for (const FMetasoundFrontendVertex& Input : InVertices)
+					{
+						TArray<const FMetasoundFrontendEdge*> Edges = Builder.FindEdges(FrontendConfigNode.GetID(), Input.VertexID);
+						for (const FMetasoundFrontendEdge* Edge : Edges)
+						{
+							check(Edge);
+							bEdgesAdded &= AddEdgeToTransactor(*Edge, Transactor);
+						}
+					}
+				};
+
+			AddConnectedEdges(FrontendConfigNode.Interface.Inputs);
+			AddConnectedEdges(FrontendConfigNode.Interface.Outputs);
+		}
+
+		// b. Removes literals, so they must be readded
+		bool bLiteralsSet = true;
+		for (const FMetasoundFrontendVertexLiteral& VertexLiteral : FrontendConfigNode.InputLiterals)
+		{
+			if (const FMetasoundFrontendVertex* Input = Builder.FindNodeInput(FrontendConfigNode.GetID(), VertexLiteral.VertexID))
+			{
+				bLiteralsSet &= AddNodeInputLiteralToTransactor(FrontendConfigNode, *Input, VertexLiteral.Value, Transactor);
+			}
+			else
+			{
+				bLiteralsSet = false;
+			}
+		}
+
+		return bEdgesAdded && bLiteralsSet;
+	});
+}
+
 void UMetaSoundSourceBuilder::OnNodeInputLiteralSet(int32 NodeIndex, int32 VertexIndex, int32 LiteralIndex) const
 {
 	using namespace Metasound::DynamicGraph;
@@ -580,13 +708,8 @@ void UMetaSoundSourceBuilder::OnNodeInputLiteralSet(int32 NodeIndex, int32 Verte
 	{
 		ExecuteAuditionableTransaction([this, &Node, &Input, &LiteralIndex](FDynamicOperatorTransactor& Transactor)
 		{
-			using namespace Metasound;
-			using namespace Metasound::Engine;
-
 			const FMetasoundFrontendLiteral& InputDefault = Node.InputLiterals[LiteralIndex].Value;
-			
-			Transactor.SetValue(Node.GetID(), Input.Name, InputDefault.ToLiteral(Input.TypeName), &FMetaSoundFrontendDocumentBuilder::CreateDataReference);
-			return true;
+			return AddNodeInputLiteralToTransactor(Node, Input, InputDefault, Transactor);
 		});
 	}
 }
@@ -640,27 +763,31 @@ void UMetaSoundSourceBuilder::OnRemoveSwappingEdge(int32 SwapIndex, int32 LastIn
 		const FMetasoundFrontendVertex* ToNodeInput = DocBuilder.FindNodeInput(EdgeBeingRemoved.ToNodeID, EdgeBeingRemoved.ToVertexID);
 		if (FromNodeOutput && ToNodeInput)
 		{
-			TOptional<FMetasoundFrontendLiteral> InputDefault = BuilderSubsystemPrivate::TryResolveNodeInputDefault(Builder, EdgeBeingRemoved.ToNodeID, ToNodeInput->Name);
-			if (ensureAlwaysMsgf(InputDefault.IsSet(), TEXT("Could not dynamically assign default literal upon removing edge: literal should be assigned by either the frontend document's input or the class definition")))
-			{
 #if WITH_EDITOR
-				const FMetasoundFrontendNode* ToNode = DocBuilder.FindNode(EdgeBeingRemoved.ToNodeID);
-				const FMetasoundFrontendNode* FromNode = DocBuilder.FindNode(EdgeBeingRemoved.FromNodeID);
-				if (ToNode && FromNode)
+			const FMetasoundFrontendNode* ToNode = DocBuilder.FindNode(EdgeBeingRemoved.ToNodeID);
+			const FMetasoundFrontendNode* FromNode = DocBuilder.FindNode(EdgeBeingRemoved.FromNodeID);
+			if (ToNode && FromNode)
+			{
+				const FMetasoundFrontendClass* ToNodeClass = DocBuilder.FindDependency(ToNode->ClassID);
+				const FMetasoundFrontendClass* FromNodeClass = DocBuilder.FindDependency(FromNode->ClassID);
+
+				// Adding edge to a template node can be safely ignored as template nodes are processed out.
+				// Once template node is connected as a source however, then the given edge can be processed accordingly.
+				if (ToNodeClass->Metadata.GetType() == EMetasoundFrontendClassType::Template)
 				{
-					const FMetasoundFrontendClass* ToNodeClass = DocBuilder.FindDependency(ToNode->ClassID);
-					const FMetasoundFrontendClass* FromNodeClass = DocBuilder.FindDependency(FromNode->ClassID);
+					return true;
+				}
 
-					// Adding edge to a template node can be safely ignored as template nodes are processed out.
-					// Once template node is connected as a source however, then the given edge can be processed accordingly.
-					if (ToNodeClass->Metadata.GetType() == EMetasoundFrontendClassType::Template)
+				if (FromNodeClass->Metadata.GetType() == EMetasoundFrontendClassType::Template)
+				{
+					auto FindNodeInputDefault = [this](const FGuid& InToNodeID, const FName InputName)
 					{
-						return true;
-					}
+						return BuilderSubsystemPrivate::TryResolveNodeInputDefault(Builder, InToNodeID, InputName);
+					};
 
-					if (FromNodeClass->Metadata.GetType() == EMetasoundFrontendClassType::Template)
+					const FTemplateNodeEdgeRemoveTransactionContext TransactionContext
 					{
-						const FTemplateNodeEdgeTransactionContext TransactionContext
+						FTemplateNodeEdgeTransactionContext
 						{
 							.Builder = DocBuilder,
 							.PageID = DocBuilder.GetBuildPageID(),
@@ -668,24 +795,26 @@ void UMetaSoundSourceBuilder::OnRemoveSwappingEdge(int32 SwapIndex, int32 LastIn
 							.FromOutput = FromNodeOutput,
 							.ToNode = ToNode,
 							.ToInput = ToNodeInput,
-							.ToInputDefault = &InputDefault.GetValue()
-						};
-
-						bool bHandled = false;
-						const FMetasoundFrontendClassName& FromClassName = FromNodeClass->Metadata.GetClassName();
-						if (const INodeTemplate* Template = INodeTemplateRegistry::Get().FindTemplate(FromClassName))
-						{
-							bHandled = Template->OnRemoveSwappingEdge(TransactionContext, Transactor);
-						}
+						},
+						FindNodeInputDefault
+					};
 
+					const FMetasoundFrontendClassName& FromClassName = FromNodeClass->Metadata.GetClassName();
+					if (const INodeTemplate* Template = INodeTemplateRegistry::Get().FindTemplate(FromClassName))
+					{
+						const bool bHandled = Template->OnRemoveSwappingEdge(TransactionContext, Transactor);
 						if (bHandled)
 						{
 							return true;
 						}
 					}
 				}
+			}
 #endif // WITH_EDITOR
 
+			const FMetasoundFrontendLiteral* InputDefault = BuilderSubsystemPrivate::TryResolveNodeInputDefault(Builder, EdgeBeingRemoved.ToNodeID, ToNodeInput->Name);
+			if (ensureAlwaysMsgf(InputDefault, TEXT("Could not dynamically assign default literal upon removing edge: literal should be assigned by either the frontend document's input or the class definition")))
+			{
 				Transactor.RemoveDataEdge(
 					EdgeBeingRemoved.FromNodeID,
 					FromNodeOutput->Name,
@@ -804,8 +933,11 @@ void UMetaSoundSourceBuilder::OnRemovingNodeInputLiteral(int32 NodeIndex, int32
 			using namespace Metasound::Engine;
 			using namespace Metasound::Frontend;
 
-			TOptional<FMetasoundFrontendLiteral> InputDefault = BuilderSubsystemPrivate::TryResolveNodeInputDefault(Builder, Node.GetID(), Input.Name);
-			if (ensureAlwaysMsgf(InputDefault.IsSet(), TEXT("Could not dynamically assign default literal from class definition upon removing input '%s' literal: document's dependency entry invalid and has no default assigned"), *Input.Name.ToString()))
+			const FMetasoundFrontendLiteral* InputDefault = BuilderSubsystemPrivate::TryResolveNodeInputDefault(Builder, Node.GetID(), Input.Name);
+			if (ensureAlwaysMsgf(InputDefault,
+				TEXT("Could not dynamically assign default literal from class definition upon removing input '%s' literal: "
+					"document's dependency entry invalid and has no default assigned"),
+					*Input.Name.ToString()))
 			{
 				Transactor.SetValue(Node.GetID(), Input.Name, InputDefault->ToLiteral(Input.TypeName), &FMetaSoundFrontendDocumentBuilder::CreateDataReference);
 				return true;
diff --git a/Engine/Plugins/Runtime/Metasound/Source/MetasoundEngine/Private/MetasoundEngineAsset.h b/Engine/Plugins/Runtime/Metasound/Source/MetasoundEngine/Private/MetasoundEngineAsset.h
--- a/Engine/Plugins/Runtime/Metasound/Source/MetasoundEngine/Private/MetasoundEngineAsset.h
+++ b/Engine/Plugins/Runtime/Metasound/Source/MetasoundEngine/Private/MetasoundEngineAsset.h
@@ -283,7 +283,20 @@ namespace Metasound::Engine
 				if (UMetasoundEditorGraphBase* MetaSoundGraph = Cast<UMetasoundEditorGraphBase>(InMetaSound.GetGraph()))
 				{
 					// Uses graph flavor of register with frontend to update editor systems/asset editors in case editor is enabled.
-					MetaSoundGraph->RegisterGraphWithFrontend();
+					// Ignore if live auditioning to avoid stopping playback when saving in background.
+					{
+						FMetaSoundAssetRegistrationOptions RegOptions
+						{
+#if WITH_EDITOR
+							.bForceViewSynchronization = false,
+							.bIgnoreIfLiveAuditioning = true,
+#endif // WITH_EDITOR
+							.PageOrder = UMetaSoundSettings::GetPageOrder()
+						};
+
+						MetaSoundGraph->RegisterGraphWithFrontend(&RegOptions);
+					}
+
 					InMetaSound.GetModifyContext().SetForceRefreshViews();
 				}
 			}
diff --git a/Engine/Plugins/Runtime/Metasound/Source/MetasoundEngine/Public/Metasound.h b/Engine/Plugins/Runtime/Metasound/Source/MetasoundEngine/Public/Metasound.h
--- a/Engine/Plugins/Runtime/Metasound/Source/MetasoundEngine/Public/Metasound.h
+++ b/Engine/Plugins/Runtime/Metasound/Source/MetasoundEngine/Public/Metasound.h
@@ -1,10 +1,9 @@
 // Copyright Epic Games, Inc. All Rights Reserved.
 #pragma once
 
-#include "CoreMinimal.h"
-
 #include "EdGraph/EdGraph.h"
 #include "MetasoundAssetBase.h"
+#include "MetasoundAssetManager.h"
 #include "MetasoundDocumentInterface.h"
 #include "MetasoundFrontend.h"
 #include "MetasoundFrontendDocument.h"
@@ -41,7 +40,7 @@ public:
 	virtual bool IsEditorOnly() const override { return true; }
 	virtual bool NeedsLoadForEditorGame() const override { return false; }
 
-	virtual void RegisterGraphWithFrontend() PURE_VIRTUAL(UMetasoundEditorGraphBase::RegisterGraphWithFrontend(), )
+	virtual void RegisterGraphWithFrontend(Metasound::Frontend::FMetaSoundAssetRegistrationOptions* RegOptions = nullptr) PURE_VIRTUAL(UMetasoundEditorGraphBase::RegisterGraphWithFrontend(), )
 
 #if WITH_EDITORONLY_DATA
 	UE_DEPRECATED(5.5, "ModifyContext is to be replaced by builder API delegates providing context when items changed and it will be up to the caller to track modification deltas.")
diff --git a/Engine/Plugins/Runtime/Metasound/Source/MetasoundEngine/Public/MetasoundBuilderBase.h b/Engine/Plugins/Runtime/Metasound/Source/MetasoundEngine/Public/MetasoundBuilderBase.h
--- a/Engine/Plugins/Runtime/Metasound/Source/MetasoundEngine/Public/MetasoundBuilderBase.h
+++ b/Engine/Plugins/Runtime/Metasound/Source/MetasoundEngine/Public/MetasoundBuilderBase.h
@@ -511,9 +511,6 @@ public:
 
 	UE_API void SetMemberMetadata(UMetaSoundFrontendMemberMetadata& NewMetadata);
 
-	UE_DEPRECATED(5.5, "Call directly on Frontend builder using 'GetBuilder'")
-	UE_API void UpdateDependencyClassNames(const TMap<FMetasoundFrontendClassName, FMetasoundFrontendClassName>& OldToNewReferencedClassNames);
-
 	UFUNCTION(BlueprintCallable, Category = "Audio|MetaSound|Builder", meta = (WorldContext = "Parent", DisplayName = "Build (Soft Deprecated. Parent no longer supported and field is ignored.)"))
 	virtual UPARAM(DisplayName = "MetaSound") TScriptInterface<IMetaSoundDocumentInterface> Build(UObject* Parent, const FMetaSoundBuilderOptions& Options) const { return BuildNewMetaSound(Options.Name); }
 
diff --git a/Engine/Plugins/Runtime/Metasound/Source/MetasoundEngine/Public/MetasoundBuilderSubsystem.h b/Engine/Plugins/Runtime/Metasound/Source/MetasoundEngine/Public/MetasoundBuilderSubsystem.h
--- a/Engine/Plugins/Runtime/Metasound/Source/MetasoundEngine/Public/MetasoundBuilderSubsystem.h
+++ b/Engine/Plugins/Runtime/Metasound/Source/MetasoundEngine/Public/MetasoundBuilderSubsystem.h
@@ -127,6 +127,7 @@ private:
 	UE_API void OnInputAdded(int32 InputIndex);
 	UE_API void OnLiveComponentFinished(UAudioComponent* AudioComponent);
 	UE_API void OnNodeAdded(int32 NodeIndex) const;
+	UE_API void OnNodeConfigurationUpdated(int32 NodeIndex) const;
 	UE_API void OnNodeInputLiteralSet(int32 NodeIndex, int32 VertexIndex, int32 LiteralIndex) const;
 	UE_API void OnOutputAdded(int32 OutputIndex) const;
 	UE_API void OnPageAdded(const Metasound::Frontend::FDocumentMutatePageArgs& InArgs);
@@ -137,6 +138,13 @@ private:
 	UE_API void OnRemovingOutput(int32 OutputIndex) const;
 	UE_API void OnRemovingPage(const Metasound::Frontend::FDocumentMutatePageArgs& InArgs);
 
+	UE_API bool AddEdgeToTransactor(const FMetasoundFrontendEdge& NewEdge, Metasound::DynamicGraph::FDynamicOperatorTransactor& Transactor) const;
+	UE_API bool AddNodeInputLiteralToTransactor(
+		const FMetasoundFrontendNode& Node,
+		const FMetasoundFrontendVertex& Input,
+		const FMetasoundFrontendLiteral& InputDefault,
+		Metasound::DynamicGraph::FDynamicOperatorTransactor& Transactor) const;
+
 	using FAuditionableTransaction = TFunctionRef<bool(Metasound::DynamicGraph::FDynamicOperatorTransactor&)>;
 	UE_API bool ExecuteAuditionableTransaction(FAuditionableTransaction Transaction) const;
 
diff --git a/Engine/Plugins/Runtime/Metasound/Source/MetasoundFrontend/Private/MetasoundFrontendDocumentBuilder.cpp b/Engine/Plugins/Runtime/Metasound/Source/MetasoundFrontend/Private/MetasoundFrontendDocumentBuilder.cpp
--- a/Engine/Plugins/Runtime/Metasound/Source/MetasoundFrontend/Private/MetasoundFrontendDocumentBuilder.cpp
+++ b/Engine/Plugins/Runtime/Metasound/Source/MetasoundFrontend/Private/MetasoundFrontendDocumentBuilder.cpp
@@ -5970,8 +5970,7 @@ bool FMetaSoundFrontendDocumentBuilder::SetNodeConfiguration(const FGuid& InNode
 	{
 		FMetasoundFrontendNode& Node = Graph.Nodes[*NodeIndex];
 		Node.Configuration = MoveTemp(InNodeConfiguration);
-		UpdateNodeInterfaceFromConfiguration(InNodeID, InPageID);
-		return true;
+		return UpdateNodeInterfaceFromConfiguration(InNodeID, InPageID);
 	}
 	return false;
 }
@@ -6404,47 +6403,57 @@ bool FMetaSoundFrontendDocumentBuilder::SwapGraphOutput(const FMetasoundFrontend
 
 bool FMetaSoundFrontendDocumentBuilder::UpdateNodeInterfaceFromConfiguration(const FGuid& InNodeID, const FGuid* InPageID)
 {
-	if (FMetasoundFrontendNode* Node = FindNodeInternal(InNodeID, InPageID))
+	using namespace Metasound::Frontend;
+
+	const FGuid& PageID = InPageID ? *InPageID : BuildPageID;
+	const IDocumentGraphNodeCache& NodeCache = DocumentCache->GetNodeCache(PageID);
+	const int32* NodeIndex = NodeCache.FindNodeIndex(InNodeID);
+	if (!NodeIndex)
 	{
-		const FMetasoundFrontendClass* Class = FindDependency(Node->ClassID);
-		check(Class);
+		return false;
+	}
 
-		// Update class interface override
-		if (const FMetaSoundFrontendNodeConfiguration* ConfigurationPtr = Node->Configuration.GetPtr())
-		{
-			Node->ClassInterfaceOverride = ConfigurationPtr->OverrideDefaultInterface(*Class);
-		}
-		else
-		{
-			// Set class interface override back to default if no node configuration
-			Node->ClassInterfaceOverride = TInstancedStruct<FMetasoundFrontendClassInterface>();
-		}
+	FMetasoundFrontendGraph& Graph = GetDocumentChecked().RootGraph.FindGraphChecked(PageID);
+	FMetasoundFrontendNode& Node = Graph.Nodes[*NodeIndex];
 
-		// Update node interface
-		const FMetasoundFrontendClassInterface& ClassInterface = Class->GetInterfaceForNode(*Node);
+	const FMetasoundFrontendClass* Class = FindDependency(Node.ClassID);
+	check(Class);
 
-		auto DisconnectInput = [this, &InNodeID, &InPageID](const FMetasoundFrontendVertex& InNodeInput)
-		{
-			this->RemoveEdgeToNodeInput(InNodeID, InNodeInput.VertexID, InPageID);
-		};
+	// Update class interface override
+	if (const FMetaSoundFrontendNodeConfiguration* ConfigurationPtr = Node.Configuration.GetPtr())
+	{
+		Node.ClassInterfaceOverride = ConfigurationPtr->OverrideDefaultInterface(*Class);
+	}
+	else
+	{
+		// Set class interface override back to default if no node configuration
+		Node.ClassInterfaceOverride = TInstancedStruct<FMetasoundFrontendClassInterface>();
+	}
+
+	// Update node interface
+	const FMetasoundFrontendClassInterface& ClassInterface = Class->GetInterfaceForNode(Node);
+
+	auto DisconnectInput = [this, &InNodeID, &InPageID](const FMetasoundFrontendVertex& InNodeInput)
+	{
+		RemoveEdgeToNodeInput(InNodeID, InNodeInput.VertexID, InPageID);
+	};
+
+	auto DisconnectOutput = [this, &InNodeID, &InPageID](const FMetasoundFrontendVertex& InNodeOutput)
+	{
+		RemoveEdgesFromNodeOutput(InNodeID, InNodeOutput.VertexID, InPageID);
+	};
+
+	const bool bInterfaceUpdated = Node.Interface.Update(ClassInterface, DisconnectInput, DisconnectOutput);
+	DocumentDelegates->FindNodeDelegatesChecked(PageID).OnNodeConfigurationUpdated.Broadcast(*NodeIndex);
 
-		auto DisconnectOutput = [this, &InNodeID, &InPageID](const FMetasoundFrontendVertex& InNodeOutput)
-		{
-			this->RemoveEdgesFromNodeOutput(InNodeID, InNodeOutput.VertexID, InPageID);
-		};
 #if WITH_EDITORONLY_DATA
-		const bool bInterfaceUpdated = Node->Interface.Update(ClassInterface, DisconnectInput, DisconnectOutput);
-		if (bInterfaceUpdated)
-		{
-			GetDocumentChecked().Metadata.ModifyContext.AddNodeIDModified(InNodeID);
-		}
-#else
-		Node->Interface.Update(ClassInterface, DisconnectInput, DisconnectOutput);
-#endif
-		return true;
+	if (bInterfaceUpdated)
+	{
+		GetDocumentChecked().Metadata.ModifyContext.AddNodeIDModified(InNodeID);
 	}
+#endif // WITH_EDITORONLY_DATA
 
-	return false;
+	return true;
 }
 
 bool FMetaSoundFrontendDocumentBuilder::UnlinkVariableNode(const FGuid& InNodeID, const FGuid& InPageID)
@@ -6525,21 +6534,6 @@ bool FMetaSoundFrontendDocumentBuilder::UpdateDependencyRegistryData(const TMap<
 
 	return bUpdated;
 }
-
-bool FMetaSoundFrontendDocumentBuilder::UpdateDependencyClassNames(const TMap<FMetasoundFrontendClassName, FMetasoundFrontendClassName>& OldToNewReferencedClassNames)
-{
-	using namespace Metasound::Frontend;
-
-	TMap<FNodeRegistryKey, FNodeRegistryKey> OldToNewKeys;
-	Algo::Transform(OldToNewReferencedClassNames, OldToNewKeys, [](const TPair<FMetasoundFrontendClassName, FMetasoundFrontendClassName>& ClassNamePair)
-	{
-		return TPair<FNodeRegistryKey, FNodeRegistryKey>(
-			FNodeRegistryKey(EMetasoundFrontendClassType::External, ClassNamePair.Key, FMetasoundFrontendVersionNumber()),
-			FNodeRegistryKey(EMetasoundFrontendClassType::External, ClassNamePair.Value, FMetasoundFrontendVersionNumber())
-		);
-	});
-	return UpdateDependencyRegistryData(OldToNewKeys);
-}
 #endif // WITH_EDITOR
 
 #if WITH_EDITORONLY_DATA
diff --git a/Engine/Plugins/Runtime/Metasound/Source/MetasoundFrontend/Private/MetasoundFrontendNodeTemplateRegistry.cpp b/Engine/Plugins/Runtime/Metasound/Source/MetasoundFrontend/Private/MetasoundFrontendNodeTemplateRegistry.cpp
--- a/Engine/Plugins/Runtime/Metasound/Source/MetasoundFrontend/Private/MetasoundFrontendNodeTemplateRegistry.cpp
+++ b/Engine/Plugins/Runtime/Metasound/Source/MetasoundFrontend/Private/MetasoundFrontendNodeTemplateRegistry.cpp
@@ -160,7 +160,12 @@ namespace Metasound::Frontend
 		return false;
 	}
 
-	bool FNodeTemplateBase::OnRemoveSwappingEdge(const FTemplateNodeEdgeTransactionContext& TransactionContext, DynamicGraph::FDynamicOperatorTransactor& OutTransactor) const
+	bool FNodeTemplateBase::OnNodeConfigurationUpdated(const FTemplateNodeTransactionContext& TransactionContext, DynamicGraph::FDynamicOperatorTransactor& OutTransactor) const
+	{
+		return false;
+	}
+
+	bool FNodeTemplateBase::OnRemoveSwappingEdge(const FTemplateNodeEdgeRemoveTransactionContext& TransactionContext, DynamicGraph::FDynamicOperatorTransactor& OutTransactor) const
 	{
 		return false;
 	}
diff --git a/Engine/Plugins/Runtime/Metasound/Source/MetasoundFrontend/Private/NodeTemplates/MetasoundFrontendNodeTemplateReroute.cpp b/Engine/Plugins/Runtime/Metasound/Source/MetasoundFrontend/Private/NodeTemplates/MetasoundFrontendNodeTemplateReroute.cpp
--- a/Engine/Plugins/Runtime/Metasound/Source/MetasoundFrontend/Private/NodeTemplates/MetasoundFrontendNodeTemplateReroute.cpp
+++ b/Engine/Plugins/Runtime/Metasound/Source/MetasoundFrontend/Private/NodeTemplates/MetasoundFrontendNodeTemplateReroute.cpp
@@ -130,7 +130,7 @@ namespace Metasound::Frontend
 				}
 			}
 
-			return FRerouteNodeVertexPair { ReroutedNode, ReroutedVertex};
+			return FRerouteNodeVertexPair { ReroutedNode, ReroutedVertex };
 		}
 
 		void FindReroutedInputVertices(
@@ -400,25 +400,70 @@ namespace Metasound::Frontend
 		if (RerouteOutputPair.IsValid())
 		{
 			check(RerouteOutputPair.Node && RerouteOutputPair.Vertex);
-			if (const FMetasoundFrontendClass* RerouteClass = TransactionContext.Builder.FindDependency(RerouteOutputPair.Node->ClassID))
+
+			if (RerouteOutputPair.Node != TransactionContext.FromNode)
 			{
-				if (RerouteClass->Metadata.GetType() == EMetasoundFrontendClassType::Template)
+				if (const FMetasoundFrontendClass* RerouteClass = TransactionContext.Builder.FindDependency(RerouteOutputPair.Node->ClassID))
 				{
-					if (const INodeTemplate* Template = INodeTemplateRegistry::Get().FindTemplate(RerouteClass->Metadata.GetClassName()))
+					if (RerouteClass->Metadata.GetType() == EMetasoundFrontendClassType::Template)
 					{
-						FTemplateNodeEdgeTransactionContext RerouteContext = TransactionContext;
-						RerouteContext.FromNode = RerouteOutputPair.Node;
-						RerouteContext.FromOutput = RerouteOutputPair.Vertex;
+						if (const INodeTemplate* Template = INodeTemplateRegistry::Get().FindTemplate(RerouteClass->Metadata.GetClassName()))
+						{
+							FTemplateNodeEdgeTransactionContext RerouteContext = TransactionContext;
+							RerouteContext.FromNode = RerouteOutputPair.Node;
+							RerouteContext.FromOutput = RerouteOutputPair.Vertex;
+							return Template->OnEdgeAdded(RerouteContext, OutTransactor);
+						}
+					}
+				}
+			}
+
+			TArray<const FMetasoundFrontendNode*> ReroutedInputOwningNodes;
+			TArray<const FMetasoundFrontendVertex*> ReroutedInputVertices;
 
-						return Template->OnEdgeAdded(TransactionContext, OutTransactor);
+			FindReroutedInputVertices(
+				ClassName,
+				TransactionContext.Builder,
+				TransactionContext.PageID,
+				*TransactionContext.ToNode,
+				*TransactionContext.ToInput,
+				ReroutedInputOwningNodes,
+				ReroutedInputVertices);
+
+			for (int32 Index = 0; Index < ReroutedInputOwningNodes.Num(); ++Index)
+			{
+				bool bHandled = false;
+				if (const FMetasoundFrontendClass* RerouteClass = TransactionContext.Builder.FindDependency(ReroutedInputOwningNodes[Index]->ClassID))
+				{
+					if (RerouteClass->Metadata.GetType() == EMetasoundFrontendClassType::Template)
+					{
+						if (const INodeTemplate* Template = INodeTemplateRegistry::Get().FindTemplate(RerouteClass->Metadata.GetClassName()))
+						{
+							FTemplateNodeEdgeTransactionContext RerouteContext = TransactionContext;
+							RerouteContext.FromNode = RerouteOutputPair.Node;
+							RerouteContext.FromOutput = RerouteOutputPair.Vertex;
+							RerouteContext.ToNode = ReroutedInputOwningNodes[Index];
+							RerouteContext.ToInput = ReroutedInputVertices[Index];
+							bHandled = Template->OnEdgeAdded(RerouteContext, OutTransactor);
+						}
 					}
 				}
+
+				if (!bHandled)
+				{
+					OutTransactor.AddDataEdge(
+						RerouteOutputPair.Node->GetID(),
+						RerouteOutputPair.Vertex->Name,
+						ReroutedInputOwningNodes[Index]->GetID(),
+						ReroutedInputVertices[Index]->Name);
+				}
 			}
-			OutTransactor.AddDataEdge(RerouteOutputPair.Node->GetID(), RerouteOutputPair.Vertex->Name, TransactionContext.ToNode->GetID(), TransactionContext.ToInput->Name);
-			return true;
 		}
 
-		return false;
+		// Return handled even if no rerouted edge was found. In this context, this means
+		// the reroute was not routed to anything of consequence relative to the actively
+		// executing graph, so ignore the edge removal in the frontend graph.
+		return true;
 	}
 
 	bool FRerouteNodeTemplate::OnNodeAdded(const FTemplateNodeTransactionContext& TransactionContext, DynamicGraph::FDynamicOperatorTransactor& OutTransactor) const
@@ -427,39 +472,84 @@ namespace Metasound::Frontend
 		return true;
 	}
 
-	bool FRerouteNodeTemplate::OnRemoveSwappingEdge(const FTemplateNodeEdgeTransactionContext& TransactionContext, DynamicGraph::FDynamicOperatorTransactor& OutTransactor) const
+	bool FRerouteNodeTemplate::OnRemoveSwappingEdge(const FTemplateNodeEdgeRemoveTransactionContext& TransactionContext, DynamicGraph::FDynamicOperatorTransactor& OutTransactor) const
 	{
 		using namespace ReroutePrivate;
 
 		const FRerouteNodeVertexPair RerouteOutputPair = FindReroutedOutputVertex(GetClassName(), TransactionContext.Builder, TransactionContext.PageID, *TransactionContext.FromNode, *TransactionContext.FromOutput);
 		if (RerouteOutputPair.IsValid())
 		{
-			check(RerouteOutputPair.Node && RerouteOutputPair.Vertex);
-			if (const FMetasoundFrontendClass* RerouteClass = TransactionContext.Builder.FindDependency(RerouteOutputPair.Node->ClassID))
+			if (RerouteOutputPair.Node != TransactionContext.FromNode)
 			{
-				if (RerouteClass->Metadata.GetType() == EMetasoundFrontendClassType::Template)
+				check(RerouteOutputPair.Node && RerouteOutputPair.Vertex);
+				if (const FMetasoundFrontendClass* RerouteClass = TransactionContext.Builder.FindDependency(RerouteOutputPair.Node->ClassID))
 				{
-					if (const INodeTemplate* Template = INodeTemplateRegistry::Get().FindTemplate(RerouteClass->Metadata.GetClassName()))
+					if (RerouteClass->Metadata.GetType() == EMetasoundFrontendClassType::Template)
 					{
-						FTemplateNodeEdgeTransactionContext RerouteContext = TransactionContext;
-						RerouteContext.FromNode = RerouteOutputPair.Node;
-						RerouteContext.FromOutput = RerouteOutputPair.Vertex;
-						return Template->OnRemoveSwappingEdge(TransactionContext, OutTransactor);
+						if (const INodeTemplate* Template = INodeTemplateRegistry::Get().FindTemplate(RerouteClass->Metadata.GetClassName()))
+						{
+							FTemplateNodeEdgeRemoveTransactionContext RerouteContext = TransactionContext;
+							RerouteContext.FromNode = RerouteOutputPair.Node;
+							RerouteContext.FromOutput = RerouteOutputPair.Vertex;
+							return Template->OnRemoveSwappingEdge(RerouteContext, OutTransactor);
+						}
 					}
 				}
 			}
 
-			OutTransactor.RemoveDataEdge(
-				RerouteOutputPair.Node->GetID(),
-				RerouteOutputPair.Vertex->Name,
-				TransactionContext.ToNode->GetID(),
-				TransactionContext.ToInput->Name,
-				TransactionContext.ToInputDefault->ToLiteral(TransactionContext.ToInput->TypeName),
-				&FMetaSoundFrontendDocumentBuilder::CreateDataReference);
-			return true;
+			TArray<const FMetasoundFrontendNode*> ReroutedInputOwningNodes;
+			TArray<const FMetasoundFrontendVertex*> ReroutedInputVertices;
+
+			FindReroutedInputVertices(
+				// We use the direct reroute ClassName because, for now, no inheriting classes of reroute are supported.
+				// Recursion should probably 
+				ClassName,
+				TransactionContext.Builder,
+				TransactionContext.PageID,
+				*TransactionContext.ToNode,
+				*TransactionContext.ToInput,
+				ReroutedInputOwningNodes,
+				ReroutedInputVertices);
+
+			for (int32 Index = 0; Index < ReroutedInputOwningNodes.Num(); ++Index)
+			{
+				bool bHandled = false;
+				if (const FMetasoundFrontendClass* RerouteClass = TransactionContext.Builder.FindDependency(ReroutedInputOwningNodes[Index]->ClassID))
+				{
+					if (RerouteClass->Metadata.GetType() == EMetasoundFrontendClassType::Template)
+					{
+						if (const INodeTemplate* Template = INodeTemplateRegistry::Get().FindTemplate(RerouteClass->Metadata.GetClassName()))
+						{
+							FTemplateNodeEdgeRemoveTransactionContext ReroutedContext = TransactionContext;
+							ReroutedContext.FromNode = RerouteOutputPair.Node;
+							ReroutedContext.FromOutput = RerouteOutputPair.Vertex;
+							ReroutedContext.ToNode = ReroutedInputOwningNodes[Index];
+							ReroutedContext.ToInput = ReroutedInputVertices[Index];
+							bHandled = Template->OnRemoveSwappingEdge(ReroutedContext, OutTransactor);
+						}
+					}
+				}
+
+				if (!bHandled)
+				{
+					if (const FMetasoundFrontendLiteral* DefaultLiteral = TransactionContext.ResolveActiveNodeInputDefault(RerouteOutputPair.Node->GetID(), RerouteOutputPair.Vertex->Name))
+					{
+						OutTransactor.RemoveDataEdge(
+							RerouteOutputPair.Node->GetID(),
+							RerouteOutputPair.Vertex->Name,
+							ReroutedInputOwningNodes[Index]->GetID(),
+							ReroutedInputVertices[Index]->Name,
+							DefaultLiteral->ToLiteral(ReroutedInputVertices[Index]->TypeName),
+							&FMetaSoundFrontendDocumentBuilder::CreateDataReference);
+					}
+				}
+			}
 		}
 
-		return false;
+		// Return handled even if no rerouted edge was found. In this context, this means
+		// the reroute was not routed to anything of consequence relative to the actively
+		// executing graph, so ignore the edge removal in the frontend graph.
+		return true;
 	}
 
 	bool FRerouteNodeTemplate::OnRemoveSwappingNode(const FTemplateNodeTransactionContext& TransactionContext, DynamicGraph::FDynamicOperatorTransactor& OutTransactor) const
diff --git a/Engine/Plugins/Runtime/Metasound/Source/MetasoundFrontend/Public/MetasoundAssetManager.h b/Engine/Plugins/Runtime/Metasound/Source/MetasoundFrontend/Public/MetasoundAssetManager.h
--- a/Engine/Plugins/Runtime/Metasound/Source/MetasoundFrontend/Public/MetasoundAssetManager.h
+++ b/Engine/Plugins/Runtime/Metasound/Source/MetasoundFrontend/Public/MetasoundAssetManager.h
@@ -94,8 +94,10 @@ namespace Metasound::Frontend
 		// if the following option 'bRegisterDependencies' is enabled).
 		bool bForceReregister = true;
 
+#if WITH_EDITOR
 		// If true, forces flag to resync all view (editor) data pertaining to the given asset(s) being registered.
 		bool bForceViewSynchronization = true;
+#endif // WITH_EDITOR
 
 		// If true, recursively attempts to register dependencies. (TODO: Determine if this option should be removed.
 		// Must validate that failed dependency updates due to auto-update for ex. being disabled is handled gracefully
diff --git a/Engine/Plugins/Runtime/Metasound/Source/MetasoundFrontend/Public/MetasoundFrontendDocumentBuilder.h b/Engine/Plugins/Runtime/Metasound/Source/MetasoundFrontend/Public/MetasoundFrontendDocumentBuilder.h
--- a/Engine/Plugins/Runtime/Metasound/Source/MetasoundFrontend/Public/MetasoundFrontendDocumentBuilder.h
+++ b/Engine/Plugins/Runtime/Metasound/Source/MetasoundFrontend/Public/MetasoundFrontendDocumentBuilder.h
@@ -641,9 +641,6 @@ public:
 	UE_API bool SwapGraphOutput(const FMetasoundFrontendClassVertex& InExistingOutputVertex, const FMetasoundFrontendClassVertex& NewOutputVertex);
 
 #if WITH_EDITOR
-	UE_DEPRECATED(5.5, "Use 'UpdateDependencyRegistryData' instead and supply keys (comprised of name, version & node class type)")
-	UE_API bool UpdateDependencyClassNames(const TMap<FMetasoundFrontendClassName, FMetasoundFrontendClassName>& OldToNewReferencedClassNames);
-
 	UE_API bool UpdateDependencyRegistryData(const TMap<Metasound::Frontend::FNodeRegistryKey, Metasound::Frontend::FNodeRegistryKey>& OldToNewClassKeys);
 #endif // WITH_EDITOR
 
diff --git a/Engine/Plugins/Runtime/Metasound/Source/MetasoundFrontend/Public/MetasoundFrontendDocumentModifyDelegates.h b/Engine/Plugins/Runtime/Metasound/Source/MetasoundFrontend/Public/MetasoundFrontendDocumentModifyDelegates.h
--- a/Engine/Plugins/Runtime/Metasound/Source/MetasoundFrontend/Public/MetasoundFrontendDocumentModifyDelegates.h
+++ b/Engine/Plugins/Runtime/Metasound/Source/MetasoundFrontend/Public/MetasoundFrontendDocumentModifyDelegates.h
@@ -71,6 +71,7 @@ namespace Metasound::Frontend
 	struct FNodeModifyDelegates
 	{
 		FOnMetaSoundFrontendDocumentMutateArray OnNodeAdded;
+		FOnMetaSoundFrontendDocumentMutateArray OnNodeConfigurationUpdated;
 		FOnMetaSoundFrontendDocumentRemoveSwappingArray OnRemoveSwappingNode;
 
 		FOnMetaSoundFrontendDocumentMutateNodeInputLiteralArray OnNodeInputLiteralSet;
diff --git a/Engine/Plugins/Runtime/Metasound/Source/MetasoundFrontend/Public/MetasoundFrontendNodeTemplateRegistry.h b/Engine/Plugins/Runtime/Metasound/Source/MetasoundFrontend/Public/MetasoundFrontendNodeTemplateRegistry.h
--- a/Engine/Plugins/Runtime/Metasound/Source/MetasoundFrontend/Public/MetasoundFrontendNodeTemplateRegistry.h
+++ b/Engine/Plugins/Runtime/Metasound/Source/MetasoundFrontend/Public/MetasoundFrontendNodeTemplateRegistry.h
@@ -84,9 +84,24 @@ namespace Metasound::Frontend
 
 		// Node Input being transacted on
 		const FMetasoundFrontendVertex* ToInput = nullptr;
+	};
 
-		// Resolved default literal of input (value may be from class definition or set on node)
-		const FMetasoundFrontendLiteral* ToInputDefault = nullptr;
+	// Removal requires re-instating an input default, which must be selected based
+	// on page context.  Selection of page context is engine-level behavior.  Thus,
+	// a function ref is provided for default selection should the node template be
+	// providing a graph topology reduction that requires look-up of a separate node's
+	// input on intercepting the given edge removal transaction.
+	struct FTemplateNodeEdgeRemoveTransactionContext : public FTemplateNodeEdgeTransactionContext
+	{
+		using FResolveActiveInputDefaultFunctionRef = TFunctionRef<const FMetasoundFrontendLiteral*(
+			const FGuid& /* NodeID */,
+			const FName /* VertexName */
+			)>;
+
+		// Reference to function used to determine preferred page input default for a given
+		// node's vertex. If default is set directly on paged node, uses that value.
+		// Otherwise, uses the default provided by the node's class interface.
+		const FResolveActiveInputDefaultFunctionRef ResolveActiveNodeInputDefault;
 	};
 #endif // WITH_EDITOR
 
@@ -176,9 +191,13 @@ namespace Metasound::Frontend
 		// Returns whether or not Frontend operation was handled and applied to the given transactor.
 		virtual bool OnNodeAdded(const FTemplateNodeTransactionContext& TransactionContext, DynamicGraph::FDynamicOperatorTransactor& OutTransactor) const = 0;
 
-		// Provided a Frontend transaction context, apply adding edge transaction to the given dynamic core transactor.
+		// Provided a Frontend transaction context, apply updating node config transaction to the given dynamic core transactor.
+		// Returns whether or not Frontend operation was handled and applied to the given transactor.
+		virtual bool OnNodeConfigurationUpdated(const FTemplateNodeTransactionContext& TransactionContext, DynamicGraph::FDynamicOperatorTransactor& OutTransactor) const = 0;
+
+		// Provided a Frontend transaction context, apply removing edge transaction to the given dynamic core transactor.
 		// Returns whether or not Frontend operation was handled and applied to the given transactor.
-		virtual bool OnRemoveSwappingEdge(const FTemplateNodeEdgeTransactionContext& TransactionContext, DynamicGraph::FDynamicOperatorTransactor& OutTransactor) const = 0;
+		virtual bool OnRemoveSwappingEdge(const FTemplateNodeEdgeRemoveTransactionContext& TransactionContext, DynamicGraph::FDynamicOperatorTransactor& OutTransactor) const = 0;
 
 		// Provided a Frontend transaction context, apply removing node transaction to the given dynamic core transactor.
 		// Returns whether or not Frontend operation was handled and applied to the given transactor.
@@ -214,7 +233,8 @@ namespace Metasound::Frontend
 
 		UE_API virtual bool OnEdgeAdded(const FTemplateNodeEdgeTransactionContext& TransactionContext, DynamicGraph::FDynamicOperatorTransactor& OutTransactor) const override;
 		UE_API virtual bool OnNodeAdded(const FTemplateNodeTransactionContext& TransactionContext, DynamicGraph::FDynamicOperatorTransactor& OutTransactor) const override;
-		UE_API virtual bool OnRemoveSwappingEdge(const FTemplateNodeEdgeTransactionContext& TransactionContext, DynamicGraph::FDynamicOperatorTransactor& OutTransactor) const override;
+		UE_API virtual bool OnNodeConfigurationUpdated(const FTemplateNodeTransactionContext& TransactionContext, DynamicGraph::FDynamicOperatorTransactor& OutTransactor) const override;
+		UE_API virtual bool OnRemoveSwappingEdge(const FTemplateNodeEdgeRemoveTransactionContext& TransactionContext, DynamicGraph::FDynamicOperatorTransactor& OutTransactor) const override;
 		UE_API virtual bool OnRemoveSwappingNode(const FTemplateNodeTransactionContext& TransactionContext, DynamicGraph::FDynamicOperatorTransactor& OutTransactor) const override;
 #endif // WITH_EDITOR
 	};
diff --git a/Engine/Plugins/Runtime/Metasound/Source/MetasoundFrontend/Public/NodeTemplates/MetasoundFrontendNodeTemplateReroute.h b/Engine/Plugins/Runtime/Metasound/Source/MetasoundFrontend/Public/NodeTemplates/MetasoundFrontendNodeTemplateReroute.h
--- a/Engine/Plugins/Runtime/Metasound/Source/MetasoundFrontend/Public/NodeTemplates/MetasoundFrontendNodeTemplateReroute.h
+++ b/Engine/Plugins/Runtime/Metasound/Source/MetasoundFrontend/Public/NodeTemplates/MetasoundFrontendNodeTemplateReroute.h
@@ -52,7 +52,7 @@ namespace Metasound::Frontend
 
 		UE_API virtual bool OnEdgeAdded(const FTemplateNodeEdgeTransactionContext& TransactionContext, DynamicGraph::FDynamicOperatorTransactor& OutTransactor) const override;
 		UE_API virtual bool OnNodeAdded(const FTemplateNodeTransactionContext& TransactionContext, DynamicGraph::FDynamicOperatorTransactor& OutTransactor) const override;
-		UE_API virtual bool OnRemoveSwappingEdge(const FTemplateNodeEdgeTransactionContext& TransactionContext, DynamicGraph::FDynamicOperatorTransactor& OutTransactor) const override;
+		UE_API virtual bool OnRemoveSwappingEdge(const FTemplateNodeEdgeRemoveTransactionContext& TransactionContext, DynamicGraph::FDynamicOperatorTransactor& OutTransactor) const override;
 		UE_API virtual bool OnRemoveSwappingNode(const FTemplateNodeTransactionContext& TransactionContext, DynamicGraph::FDynamicOperatorTransactor& OutTransactor) const override;
 #endif // WITH_EDITOR
 	};
diff --git a/Engine/Plugins/Runtime/Metasound/Source/MetasoundGraphCore/Private/MetasoundDynamicOperatorTransactor.cpp b/Engine/Plugins/Runtime/Metasound/Source/MetasoundGraphCore/Private/MetasoundDynamicOperatorTransactor.cpp
--- a/Engine/Plugins/Runtime/Metasound/Source/MetasoundGraphCore/Private/MetasoundDynamicOperatorTransactor.cpp
+++ b/Engine/Plugins/Runtime/Metasound/Source/MetasoundGraphCore/Private/MetasoundDynamicOperatorTransactor.cpp
@@ -641,7 +641,7 @@ namespace Metasound
 			return Queue;
 		}
 
-		void FDynamicOperatorTransactor::AddNode(const FGuid& InNodeID, TUniquePtr<INode> InNode)
+		bool FDynamicOperatorTransactor::AddNode(const FGuid& InNodeID, TUniquePtr<INode> InNode)
 		{
 			using namespace DirectedGraphAlgo;
 
@@ -649,7 +649,7 @@ namespace Metasound
 
 			if (!InNode)
 			{
-				return;
+				return false;
 			}
 
 			// Cache pointer because InNode gets moved. 
@@ -659,9 +659,10 @@ namespace Metasound
 			int32 OperatorOrdinal = GraphSorter.InsertOperator(GetOperatorID(NodePtr), FDynamicGraphIncrementalSorter::EInsertLocation::Last);
 
 			EnqueueInsertOperatorTransform(*NodePtr, OperatorOrdinal);
+			return true;
 		}
 
-		void FDynamicOperatorTransactor::RemoveNode(const FGuid& InNodeID)
+		bool FDynamicOperatorTransactor::RemoveNode(const FGuid& InNodeID)
 		{
 			using namespace DirectedGraphAlgo;
 
@@ -688,11 +689,12 @@ namespace Metasound
 				{
 					RemoveNodeInternal(*Node, bRemoveDataEdgesWithNode);
 				}
+
+				return true;
 			}
-			else
-			{
-				UE_LOG(LogMetaSound, Error, TEXT("No node found in dynamic transactor graph with ID %s"), *InNodeID.ToString());
-			}
+
+			UE_LOG(LogMetaSound, Error, TEXT("No node found in dynamic transactor graph with ID %s"), *InNodeID.ToString());
+			return false;
 		}
 
 		/** Add an edge to the graph. */
diff --git a/Engine/Plugins/Runtime/Metasound/Source/MetasoundGraphCore/Public/MetasoundDynamicOperatorTransactor.h b/Engine/Plugins/Runtime/Metasound/Source/MetasoundGraphCore/Public/MetasoundDynamicOperatorTransactor.h
--- a/Engine/Plugins/Runtime/Metasound/Source/MetasoundGraphCore/Public/MetasoundDynamicOperatorTransactor.h
+++ b/Engine/Plugins/Runtime/Metasound/Source/MetasoundGraphCore/Public/MetasoundDynamicOperatorTransactor.h
@@ -195,10 +195,10 @@ namespace Metasound
 			UE_API TSharedRef<TSpscQueue<TUniquePtr<IDynamicOperatorTransform>>> CreateTransformQueue(const FOperatorSettings& InOperatorSettings, const FMetasoundEnvironment& InEnvironment, const TSharedPtr<FGraphRenderCost>& InRenderCost);
 
 			/** Add a node to the graph. */
-			UE_API void AddNode(const FGuid& InNodeID, TUniquePtr<INode> InNode);
+			UE_API bool AddNode(const FGuid& InNodeID, TUniquePtr<INode> InNode);
 
 			/** Remove a node from the graph. */
-			UE_API void RemoveNode(const FGuid& InNodeID);
+			UE_API bool RemoveNode(const FGuid& InNodeID);
 
 			/** Add an edge to the graph, connecting two vertices from two 
 			 * nodes. 
