Add MetaSound document version transform for deduplicating dependencies

diff --git a/Engine/Plugins/Runtime/Metasound/Source/MetasoundFrontend/Private/MetasoundFrontendDocumentCache.cpp b/Engine/Plugins/Runtime/Metasound/Source/MetasoundFrontend/Private/MetasoundFrontendDocumentCache.cpp
--- a/Engine/Plugins/Runtime/Metasound/Source/MetasoundFrontend/Private/MetasoundFrontendDocumentCache.cpp
+++ b/Engine/Plugins/Runtime/Metasound/Source/MetasoundFrontend/Private/MetasoundFrontendDocumentCache.cpp
@@ -1082,6 +1082,12 @@ namespace Metasound::Frontend
 			const FMetasoundFrontendClassMetadata& Metadata = Class.Metadata;
 			FNodeRegistryKey Key = FNodeRegistryKey(Metadata);
 			KeyToIndex.Add(MoveTemp(Key), Index);
+#if !NO_LOGGING
+			if (IDToIndex.Contains(Class.ID))
+			{
+				UE_LOG(LogMetaSound, Warning, TEXT("MetaSound document with class name %s found containing multiple dependencies with the same ID while initializing builder cache. Document may be currently versioning or require versioning past version 1.15"), *InDocument.RootGraph.Metadata.GetClassName().ToString());
+			}
+#endif // !NO_LOGGING
 			IDToIndex.Add(Class.ID, Index);
 		}
 	}
diff --git a/Engine/Plugins/Runtime/Metasound/Source/MetasoundFrontend/Private/MetasoundFrontendDocumentVersioning.cpp b/Engine/Plugins/Runtime/Metasound/Source/MetasoundFrontend/Private/MetasoundFrontendDocumentVersioning.cpp
--- a/Engine/Plugins/Runtime/Metasound/Source/MetasoundFrontend/Private/MetasoundFrontendDocumentVersioning.cpp
+++ b/Engine/Plugins/Runtime/Metasound/Source/MetasoundFrontend/Private/MetasoundFrontendDocumentVersioning.cpp
@@ -3,6 +3,7 @@
 
 #if WITH_EDITORONLY_DATA
 #include "Algo/Transform.h"
+#include "Algo/Unique.h"
 #include "CoreGlobals.h"
 #include "HAL/IConsoleManager.h"
 #include "Interfaces/MetasoundFrontendInterface.h"
@@ -762,6 +763,71 @@ PRAGMA_ENABLE_DEPRECATION_WARNINGS
 			}
 		};
 
+		/** Versions document from 1.14 to 1.15. */
+		class FVersionDocument_1_15 : public FVersionDocumentTransform
+		{
+		public:
+			virtual ~FVersionDocument_1_15() = default;
+
+			FMetasoundFrontendVersionNumber GetTargetVersion() const override
+			{
+				return { 1, 15 };
+			}
+
+			void TransformInternal(FMetaSoundFrontendDocumentBuilder& OutBuilder) const override
+			{
+				// Some old assets have multiple copies of the same dependency 
+				// which causes issues with the builder's cache which relies on 
+				// a 1 to 1 relationship between IDs and dependencies.
+				// This removes duplicate entries, keeping the highest version dependencies. 
+				FMetasoundFrontendDocument& Document = const_cast<FMetasoundFrontendDocument&>(OutBuilder.GetConstDocumentChecked());
+
+				auto CompareForUnique = [&](const FMetasoundFrontendClass& InLHS, const FMetasoundFrontendClass& InRHS)
+				{
+					return InLHS.ID == InRHS.ID;
+				};
+
+				auto CompareForSort = [&](const FMetasoundFrontendClass& InLHS, const FMetasoundFrontendClass& InRHS)
+				{
+					// Sort by ID
+					if (InLHS.ID < InRHS.ID)
+					{
+						return true;
+					}
+					else if (InRHS.ID < InLHS.ID)
+					{
+						return false;
+					}
+					// If IDs are equal, sort by version number descending
+					else if (InLHS.Metadata.GetVersion() > InRHS.Metadata.GetVersion())
+					{
+						return true;
+					}
+					else if (InLHS.Metadata.GetVersion() < InRHS.Metadata.GetVersion())
+					{
+						return false;
+					}
+					else
+					{
+						// if IDs and version numbers are equal, sort by number of inputs & outputs descending
+						const int32 NumLHSVertices = InLHS.GetDefaultInterface().Inputs.Num() + InLHS.GetDefaultInterface().Outputs.Num();
+						const int32 NumRHSVertices = InRHS.GetDefaultInterface().Inputs.Num() + InRHS.GetDefaultInterface().Outputs.Num();
+
+						return NumLHSVertices > NumRHSVertices;
+					}
+				};
+				
+				Algo::Sort(Document.Dependencies, CompareForSort);
+				Document.Dependencies.SetNum(Algo::Unique(Document.Dependencies, CompareForUnique));
+
+				// Safeguards against prior fix-up corrupting any cached data
+				if (IDocumentBuilderRegistry* BuilderRegistry = IDocumentBuilderRegistry::Get())
+				{
+					BuilderRegistry->ReloadBuilder(Document.RootGraph.Metadata.GetClassName());
+				}
+			}
+		};
+
 		bool VersionBuilderDocument(FMetaSoundFrontendDocumentBuilder& Builder)
 		{
 			UObject& DocObject = Builder.CastDocumentObjectChecked<UObject>();
@@ -772,6 +838,7 @@ PRAGMA_ENABLE_DEPRECATION_WARNINGS
 			bWasUpdated |= FVersionDocument_1_12(Name, Path).Transform(Builder);
 			bWasUpdated |= FVersionDocument_1_13().Transform(Builder);
 			bWasUpdated |= FVersionDocument_1_14().Transform(Builder);
+			bWasUpdated |= FVersionDocument_1_15().Transform(Builder);
 
 			return bWasUpdated;
 		}
diff --git a/Engine/Plugins/Runtime/Metasound/Source/MetasoundFrontend/Private/MetasoundFrontendDocumentVersioning.h b/Engine/Plugins/Runtime/Metasound/Source/MetasoundFrontend/Private/MetasoundFrontendDocumentVersioning.h
--- a/Engine/Plugins/Runtime/Metasound/Source/MetasoundFrontend/Private/MetasoundFrontendDocumentVersioning.h
+++ b/Engine/Plugins/Runtime/Metasound/Source/MetasoundFrontend/Private/MetasoundFrontendDocumentVersioning.h
@@ -16,7 +16,7 @@ namespace Metasound::Frontend
 {
 	static FMetasoundFrontendVersionNumber GetMaxDocumentVersion()
 	{
-		return FMetasoundFrontendVersionNumber { 1, 14 };
+		return FMetasoundFrontendVersionNumber { 1, 15 };
 	}
 
 	bool ChangeIDComparisonEnabledInAutoUpdate();
