- Fix display name in node tooltip showing transient node name incorrectly
- Remove remainder of handler logic in favor of Builder API in slate MetaSound graph pins

#rnx

diff --git a/Engine/Plugins/Runtime/Metasound/Source/MetasoundEditor/Private/SMetasoundGraphNode.cpp b/Engine/Plugins/Runtime/Metasound/Source/MetasoundEditor/Private/SMetasoundGraphNode.cpp
--- a/Engine/Plugins/Runtime/Metasound/Source/MetasoundEditor/Private/SMetasoundGraphNode.cpp
+++ b/Engine/Plugins/Runtime/Metasound/Source/MetasoundEditor/Private/SMetasoundGraphNode.cpp
@@ -688,6 +688,32 @@ namespace Metasound
 			return FAppStyle::GetBrush("Graph.Node.Body");
 		}
 
+		FText SMetaSoundGraphNode::GetNodeTooltip() const
+		{
+			FText TooltipText = FText::GetEmpty();
+			if (GraphNode)
+			{
+				// Don't show if pin already displaying interactive tooltip
+				for (UEdGraphPin* Pin : GraphNode->GetAllPins())
+				{
+					TSharedPtr<SGraphPin> PinWidget = FindWidgetForPin(Pin);
+
+					if (PinWidget.IsValid() && PinWidget->HasInteractiveTooltip())
+					{
+						return TooltipText;
+					}
+				}
+
+				TooltipText = GraphNode->GetTooltipText();
+				if (TooltipText.IsEmpty())
+				{
+					TooltipText = GraphNode->GetNodeTitle(ENodeTitleType::FullTitle);
+				}
+			}
+
+			return TooltipText;
+		}
+
 		EVisibility SMetaSoundGraphNode::IsAddPinButtonVisible() const
 		{
 			EVisibility DefaultVisibility = SGraphNode::IsAddPinButtonVisible();
diff --git a/Engine/Plugins/Runtime/Metasound/Source/MetasoundEditor/Private/SMetasoundGraphNode.h b/Engine/Plugins/Runtime/Metasound/Source/MetasoundEditor/Private/SMetasoundGraphNode.h
--- a/Engine/Plugins/Runtime/Metasound/Source/MetasoundEditor/Private/SMetasoundGraphNode.h
+++ b/Engine/Plugins/Runtime/Metasound/Source/MetasoundEditor/Private/SMetasoundGraphNode.h
@@ -51,6 +51,7 @@ namespace Metasound::Editor
 		virtual TSharedRef<SWidget> CreateNodeContentArea() override;
 		virtual TSharedRef<SWidget> CreateTitleWidget(TSharedPtr<SNodeTitle> NodeTitle) override;
 		virtual const FSlateBrush* GetNodeBodyBrush() const override;
+		virtual FText GetNodeTooltip() const override;
 		virtual TSharedRef<SWidget> CreateTitleRightWidget() override;
 		virtual EVisibility IsAddPinButtonVisible() const override;
 		virtual FReply OnAddPin() override;
diff --git a/Engine/Plugins/Runtime/Metasound/Source/MetasoundEditor/Private/SMetasoundGraphPin.h b/Engine/Plugins/Runtime/Metasound/Source/MetasoundEditor/Private/SMetasoundGraphPin.h
--- a/Engine/Plugins/Runtime/Metasound/Source/MetasoundEditor/Private/SMetasoundGraphPin.h
+++ b/Engine/Plugins/Runtime/Metasound/Source/MetasoundEditor/Private/SMetasoundGraphPin.h
@@ -232,48 +232,6 @@ namespace Metasound
 
 			virtual ~TMetasoundGraphPin() = default;
 
-			Frontend::FConstInputHandle GetConstInputHandle() const
-			{
-				using namespace Frontend;
-
-				const bool bIsInput = (ParentPinType::GetDirection() == EGPD_Input);
-				if (bIsInput)
-				{
-					if (const UEdGraphPin* Pin = ParentPinType::GetPinObj())
-					{
-						if (const UMetasoundEditorGraphNode* MetasoundEditorNode = Cast<UMetasoundEditorGraphNode>(Pin->GetOwningNode()))
-						{
-							FConstNodeHandle NodeHandle = MetasoundEditorNode->GetConstNodeHandle();
-							return NodeHandle->GetConstInputWithVertexName(Pin->GetFName());
-						}
-					}
-				}
-
-				return IInputController::GetInvalidHandle();
-			}
-
-			Frontend::FInputHandle GetInputHandle()
-			{
-				using namespace Frontend;
-
-				const bool bIsInput = (ParentPinType::GetDirection() == EGPD_Input);
-				if (bIsInput)
-				{
-					if (UEdGraphPin* Pin = ParentPinType::GetPinObj())
-					{
-						if (UMetasoundEditorGraphNode* MetasoundEditorNode = Cast<UMetasoundEditorGraphNode>(Pin->GetOwningNode()))
-						{
-PRAGMA_DISABLE_DEPRECATION_WARNINGS
-							FNodeHandle NodeHandle = MetasoundEditorNode->GetNodeHandle();
-PRAGMA_ENABLE_DEPRECATION_WARNINGS
-							return NodeHandle->GetInputWithVertexName(Pin->GetFName());
-						}
-					}
-				}
-
-				return IInputController::GetInvalidHandle();
-			}
-
 			const UMetasoundEditorGraphNode* GetOwningMetaSoundNode() const
 			{
 				if (const UEdGraphPin* Pin = ParentPinType::GetPinObj())
@@ -387,26 +345,35 @@ PRAGMA_ENABLE_DEPRECATION_WARNINGS
 						using namespace Frontend;
 						if (!ParentPinType::IsConnected())
 						{
-							FConstInputHandle InputHandle = GetConstInputHandle();
-							if (const FMetasoundFrontendLiteral* Literal = InputHandle->GetLiteral())
+							UMetaSoundBuilderBase* Builder = GetBuilder();
+							UEdGraphPin* GraphPin = ParentPinType::GetPinObj();
+							if (!Builder || !GraphPin)
+							{
+								return EVisibility::Collapsed;
+							}
+
+							const FMetaSoundFrontendDocumentBuilder& DocBuilder = Builder->GetConstBuilder();
+							const FMetasoundFrontendNode* FrontendNode = nullptr;
+							if (const FMetasoundFrontendVertex* Vertex = FGraphBuilder::GetPinVertex(DocBuilder, GraphPin, &FrontendNode))
 							{
-								const bool bIsDefaultConstructed = Literal->GetType() == EMetasoundFrontendLiteralType::None;
-								const bool bIsTriggerDataType = InputHandle->GetDataType() == GetMetasoundDataTypeName<FTrigger>();
-								bool bIsRerouteNode = false;
-								if (UEdGraphPin* Pin = ParentPinType::GetPinObj())
+								check(FrontendNode);
+								if (const FMetasoundFrontendVertexLiteral* Literal = DocBuilder.FindNodeInputDefault(FrontendNode->GetID(), Vertex->Name))
 								{
-									if (UMetasoundEditorGraphExternalNode* Node = Cast<UMetasoundEditorGraphExternalNode>(Pin->GetOwningNode()))
+									const bool bIsDefaultConstructed = Literal->Value.GetType() == EMetasoundFrontendLiteralType::None;
+									const bool bIsTriggerDataType = Vertex->TypeName == GetMetasoundDataTypeName<FTrigger>();
+									bool bIsRerouteNode = false;
+									if (UMetasoundEditorGraphExternalNode* Node = Cast<UMetasoundEditorGraphExternalNode>(GraphPin->GetOwningNode()))
 									{
 										if (Node->GetBreadcrumb().ClassName == FRerouteNodeTemplate::ClassName)
 										{
 											bIsRerouteNode = true;
 										}
 									}
-								}
-
-								if (!bIsDefaultConstructed && !bIsTriggerDataType && !bIsRerouteNode)
-								{
-									return EVisibility::Visible;
+									
+									if (!bIsDefaultConstructed && !bIsTriggerDataType && !bIsRerouteNode)
+									{
+										return EVisibility::Visible;
+									}
 								}
 							}
 						}
@@ -425,32 +392,35 @@ PRAGMA_ENABLE_DEPRECATION_WARNINGS
 								if (UMetasoundEditorGraph* MetaSoundGraph = CastChecked<UMetasoundEditorGraph>(Node->GetGraph()))
 								{
 									UObject& MetaSound = MetaSoundGraph->GetMetasoundChecked();
-									FMetasoundFrontendDocumentModifyContext& ModifyContext = FGraphBuilder::GetOutermostMetaSoundChecked(MetaSound).GetModifyContext();
-									{
-										const FScopedTransaction Transaction(LOCTEXT("MetaSoundEditorResetToClassDefault", "Reset to Class Default"));
-										MetaSound.Modify();
-										MetaSoundGraph->Modify();
+									const FScopedTransaction Transaction(LOCTEXT("MetaSoundEditorResetToClassDefault", "Reset to Class Default"));
+									MetaSound.Modify();
+									MetaSoundGraph->Modify();
 
-										if (UMetasoundEditorGraphMemberNode* MemberNode = Cast<UMetasoundEditorGraphMemberNode>(Node))
+									if (UMetasoundEditorGraphMemberNode* MemberNode = Cast<UMetasoundEditorGraphMemberNode>(Node))
+									{
+										FMetasoundFrontendDocumentModifyContext& ModifyContext = FGraphBuilder::GetOutermostMetaSoundChecked(MetaSound).GetModifyContext();
+										UMetasoundEditorGraphMember* Member = MemberNode->GetMember();
+										if (ensure(Member))
 										{
-											UMetasoundEditorGraphMember* Member = MemberNode->GetMember();
-											if (ensure(Member))
-											{
-												Member->ResetToClassDefault();
-												ModifyContext.AddMemberIDsModified({ Member->GetMemberID() });
-											}
-											else
-											{
-												ModifyContext.SetDocumentModified();
-											}
+											Member->ResetToClassDefault();
+											ModifyContext.AddMemberIDsModified({ Member->GetMemberID() });
 										}
 										else
 										{
-											FInputHandle InputHandle = GetInputHandle();
-											InputHandle->ClearLiteral();
 											ModifyContext.SetDocumentModified();
 										}
 									}
+									else
+									{
+										if (UMetaSoundBuilderBase* Builder = GetBuilder())
+										{
+											FMetaSoundFrontendDocumentBuilder& DocBuilder = Builder->GetBuilder();
+											if (const FMetasoundFrontendVertex* Vertex = FGraphBuilder::GetPinVertex(DocBuilder, Pin))
+											{
+												Builder->GetBuilder().RemoveNodeInputDefault(Node->GetNodeID(), Vertex->VertexID);
+											}
+										}
+									}
 								}
 							}
 						}
@@ -460,7 +430,7 @@ PRAGMA_ENABLE_DEPRECATION_WARNINGS
 					.Content()
 					[
 						SNew(SImage)
-							.Image(FAppStyle::GetBrush("PropertyWindow.DiffersFromDefault"))
+						.Image(FAppStyle::GetBrush("PropertyWindow.DiffersFromDefault"))
 					];
 			}
 
@@ -477,8 +447,8 @@ PRAGMA_ENABLE_DEPRECATION_WARNINGS
 				// For now, arrays do not support literals.
 				// TODO: Support array literals by displaying
 				// default literals (non-array too) in inspector window.
-				FConstInputHandle InputHandle = GetConstInputHandle();
-				if (!InputHandle->IsValid() || ParentPinType::IsArray())
+				const FMetasoundFrontendVertex* FrontendVertex = GetFrontendVertex();
+				if (!FrontendVertex || ParentPinType::IsArray())
 				{
 					return DefaultWidget;
 				}
diff --git a/Engine/Source/Editor/GraphEditor/Public/SGraphNode.h b/Engine/Source/Editor/GraphEditor/Public/SGraphNode.h
--- a/Engine/Source/Editor/GraphEditor/Public/SGraphNode.h
+++ b/Engine/Source/Editor/GraphEditor/Public/SGraphNode.h
@@ -186,7 +186,7 @@ public:
 	UE_API virtual FLinearColor GetNodeTitleTextColor() const;
 
 	/** @return the tooltip to display when over the node */
-	UE_API FText GetNodeTooltip() const;
+	UE_API virtual FText GetNodeTooltip() const;
 	
 	/** @return the node being observed by this widget*/
 	UE_API UEdGraphNode* GetNodeObj() const;
