Add AccessFlags to MetaSound UObject types
- Implement RegisterAsClass as flag to determine what should be registered as a referenceable flag and shown in editor or not
- Move IsDeprecated to associated access flag

[RN] Add flag to specify whether a MetaSound asset is registered as a class in editor for other MetaSounds to reference or not.

diff --git a/Engine/Plugins/Experimental/TechAudioTools/Source/TechAudioToolsMetaSoundEditor/Private/ViewModels/MetaSoundEditorViewModel.cpp b/Engine/Plugins/Experimental/TechAudioTools/Source/TechAudioToolsMetaSoundEditor/Private/ViewModels/MetaSoundEditorViewModel.cpp
--- a/Engine/Plugins/Experimental/TechAudioTools/Source/TechAudioToolsMetaSoundEditor/Private/ViewModels/MetaSoundEditorViewModel.cpp
+++ b/Engine/Plugins/Experimental/TechAudioTools/Source/TechAudioToolsMetaSoundEditor/Private/ViewModels/MetaSoundEditorViewModel.cpp
@@ -39,7 +39,7 @@ void UMetaSoundEditorViewModel::Initialize(UMetaSoundBuilderBase* InBuilder)
 	OnAuthorChanged(FrontendDocument.RootGraph.Metadata.GetAuthor());
 	OnKeywordsChanged(FrontendDocument.RootGraph.Metadata.GetKeywords());
 	OnCategoryHierarchyChanged(FrontendDocument.RootGraph.Metadata.GetCategoryHierarchy());
-	OnIsDeprecatedChanged(FrontendDocument.RootGraph.Metadata.GetIsDeprecated());
+	OnIsDeprecatedChanged(EnumHasAllFlags(FrontendDocument.RootGraph.Metadata.GetAccessFlags(), EMetasoundFrontendClassAccessFlags::Deprecated));
 
 	BindDelegates();
 }
@@ -108,7 +108,14 @@ void UMetaSoundEditorViewModel::SetIsDeprecated(const bool bInIsDeprecated)
 	if (Builder && UE_MVVM_SET_PROPERTY_VALUE(bIsDeprecated, bInIsDeprecated))
 	{
 		FMetaSoundFrontendDocumentBuilder& DocumentBuilder = Builder->GetBuilder();
-		DocumentBuilder.SetIsDeprecated(bInIsDeprecated);
+		if (bInIsDeprecated)
+		{
+			DocumentBuilder.AddAccessFlags(EMetasoundFrontendClassAccessFlags::Deprecated);
+		}
+		else
+		{
+			DocumentBuilder.RemoveAccessFlags(EMetasoundFrontendClassAccessFlags::Deprecated);
+		}
 	}
 }
 
diff --git a/Engine/Plugins/Runtime/Harmonix/Source/HarmonixMetasound/Private/HarmonixMetasound/Nodes/MidiCCTriggerNode.cpp b/Engine/Plugins/Runtime/Harmonix/Source/HarmonixMetasound/Private/HarmonixMetasound/Nodes/MidiCCTriggerNode.cpp
--- a/Engine/Plugins/Runtime/Harmonix/Source/HarmonixMetasound/Private/HarmonixMetasound/Nodes/MidiCCTriggerNode.cpp
+++ b/Engine/Plugins/Runtime/Harmonix/Source/HarmonixMetasound/Private/HarmonixMetasound/Nodes/MidiCCTriggerNode.cpp
@@ -266,7 +266,7 @@ namespace HarmonixMetasound::Nodes::MidiCCTriggerNode
 				Info.PromptIfMissing = PluginNodeMissingPrompt;
 				Info.DefaultInterface = GetVertexInterface();
 				Info.CategoryHierarchy.Emplace(NodeCategories::Music);
-				Info.bDeprecated = true;
+				EnumAddFlags(Info.AccessFlags, ENodeClassAccessFlags::Deprecated);
 				return Info;
 			};
 
diff --git a/Engine/Plugins/Runtime/Metasound/Source/MetasoundEditor/Private/MetasoundDetailCustomization.cpp b/Engine/Plugins/Runtime/Metasound/Source/MetasoundEditor/Private/MetasoundDetailCustomization.cpp
--- a/Engine/Plugins/Runtime/Metasound/Source/MetasoundEditor/Private/MetasoundDetailCustomization.cpp
+++ b/Engine/Plugins/Runtime/Metasound/Source/MetasoundEditor/Private/MetasoundDetailCustomization.cpp
@@ -158,13 +158,13 @@ namespace Metasound::Editor
 			case EMetasoundActiveDetailView::Metasound:
 			{
 				IDetailCategoryBuilder& GeneralCategoryBuilder = DetailLayout.EditCategory("MetaSound");
+				const FName AccessFlagsPropertyPath = BuildChildPath(GetMetadataPropertyPath(), FMetasoundFrontendClassMetadata::GetAccessFlagsPropertyName());
 				const FName AuthorPropertyPath = BuildChildPath(GetMetadataPropertyPath(), FMetasoundFrontendClassMetadata::GetAuthorPropertyName());
 				const FName CategoryHierarchyPropertyPath = BuildChildPath(GetMetadataPropertyPath(), FMetasoundFrontendClassMetadata::GetCategoryHierarchyPropertyName());
 				const FName ClassNamePropertyPath = BuildChildPath(GetMetadataPropertyPath(), FMetasoundFrontendClassMetadata::GetClassNamePropertyName());
 				const FName DescPropertyPath = BuildChildPath(GetMetadataPropertyPath(), FMetasoundFrontendClassMetadata::GetDescriptionPropertyName());
 				const FName DisplayNamePropertyPath = BuildChildPath(GetMetadataPropertyPath(), FMetasoundFrontendClassMetadata::GetDisplayNamePropertyName());
 				const FName KeywordsPropertyPath = BuildChildPath(GetMetadataPropertyPath(), FMetasoundFrontendClassMetadata::GetKeywordsPropertyName());
-				const FName IsDeprecatedPropertyPath = BuildChildPath(GetMetadataPropertyPath(), FMetasoundFrontendClassMetadata::GetIsDeprecatedPropertyName());
 				const FName VersionPropertyPath = BuildChildPath(GetMetadataPropertyPath(), FMetasoundFrontendClassMetadata::GetVersionPropertyName());
 
 				const FName ClassNameNamePropertyPath = BuildChildPath(ClassNamePropertyPath, GET_MEMBER_NAME_CHECKED(FMetasoundFrontendClassName, Name));
@@ -174,13 +174,13 @@ namespace Metasound::Editor
 
 				const FName InterfaceVersionsPropertyPath = GetInterfaceVersionsPropertyPath();
 
+				TSharedPtr<IPropertyHandle> AccessFlagsHandle = DetailLayout.GetProperty(AccessFlagsPropertyPath);
 				TSharedPtr<IPropertyHandle> AuthorHandle = DetailLayout.GetProperty(AuthorPropertyPath);
 				TSharedPtr<IPropertyHandle> CategoryHierarchyHandle = DetailLayout.GetProperty(CategoryHierarchyPropertyPath);
 				TSharedPtr<IPropertyHandle> ClassNameHandle = DetailLayout.GetProperty(ClassNameNamePropertyPath);
 				TSharedPtr<IPropertyHandle> DisplayNameHandle = DetailLayout.GetProperty(DisplayNamePropertyPath);
 				TSharedPtr<IPropertyHandle> DescHandle = DetailLayout.GetProperty(DescPropertyPath);
 				TSharedPtr<IPropertyHandle> KeywordsHandle = DetailLayout.GetProperty(KeywordsPropertyPath);
-				TSharedPtr<IPropertyHandle> IsDeprecatedHandle = DetailLayout.GetProperty(IsDeprecatedPropertyPath);
 				TSharedPtr<IPropertyHandle> InterfaceVersionsHandle = DetailLayout.GetProperty(InterfaceVersionsPropertyPath);
 				TSharedPtr<IPropertyHandle> MajorVersionHandle = DetailLayout.GetProperty(MajorVersionPropertyPath);
 				TSharedPtr<IPropertyHandle> MinorVersionHandle = DetailLayout.GetProperty(MinorVersionPropertyPath);
@@ -254,6 +254,8 @@ namespace Metasound::Editor
 							}
 						}
 					});
+
+					AccessFlagsHandle->SetOnPropertyValueChanged(RegisterOnChange);
 					AuthorHandle->SetOnPropertyValueChanged(RegisterOnChange);
 					DescHandle->SetOnPropertyValueChanged(RegisterOnChange);
 					DisplayNameHandle->SetOnPropertyValueChanged(RegisterOnChange);
@@ -261,13 +263,12 @@ namespace Metasound::Editor
 					KeywordsHandle->SetOnChildPropertyValueChanged(RegisterOnChange);
 					CategoryHierarchyHandle->SetOnPropertyValueChanged(RegisterOnChange);
 					CategoryHierarchyHandle->SetOnChildPropertyValueChanged(RegisterOnChange);
-					IsDeprecatedHandle->SetOnPropertyValueChanged(RegisterOnChange);
 				}
 
 				GeneralCategoryBuilder.AddProperty(DisplayNameHandle);
 				GeneralCategoryBuilder.AddProperty(DescHandle);
 				GeneralCategoryBuilder.AddProperty(AuthorHandle);
-				GeneralCategoryBuilder.AddProperty(IsDeprecatedHandle);
+				GeneralCategoryBuilder.AddProperty(AccessFlagsHandle);
 				GeneralCategoryBuilder.AddProperty(MajorVersionHandle);
 				GeneralCategoryBuilder.AddProperty(MinorVersionHandle);
 
diff --git a/Engine/Plugins/Runtime/Metasound/Source/MetasoundEditor/Private/MetasoundEditorBuilderListener.cpp b/Engine/Plugins/Runtime/Metasound/Source/MetasoundEditor/Private/MetasoundEditorBuilderListener.cpp
--- a/Engine/Plugins/Runtime/Metasound/Source/MetasoundEditor/Private/MetasoundEditorBuilderListener.cpp
+++ b/Engine/Plugins/Runtime/Metasound/Source/MetasoundEditor/Private/MetasoundEditorBuilderListener.cpp
@@ -82,7 +82,7 @@ void UMetaSoundEditorBuilderListener::OnDocumentIsDeprecatedChanged()
 	{
 		const FMetaSoundFrontendDocumentBuilder& ConstBuilder = Builder->GetConstBuilder();
 		const FMetasoundFrontendGraphClass& GraphClass = ConstBuilder.GetConstDocumentChecked().RootGraph;
-		const bool bIsDeprecated = GraphClass.Metadata.GetIsDeprecated();
+		const bool bIsDeprecated = EnumHasAnyFlags(GraphClass.Metadata.GetAccessFlags(), EMetasoundFrontendClassAccessFlags::Deprecated);
 
 		OnDocumentIsDeprecatedChangedDelegate.Broadcast(bIsDeprecated);
 	}
diff --git a/Engine/Plugins/Runtime/Metasound/Source/MetasoundEditor/Private/MetasoundEditorGraphNode.cpp b/Engine/Plugins/Runtime/Metasound/Source/MetasoundEditor/Private/MetasoundEditorGraphNode.cpp
--- a/Engine/Plugins/Runtime/Metasound/Source/MetasoundEditor/Private/MetasoundEditorGraphNode.cpp
+++ b/Engine/Plugins/Runtime/Metasound/Source/MetasoundEditor/Private/MetasoundEditorGraphNode.cpp
@@ -1430,7 +1430,7 @@ void UMetasoundEditorGraphExternalNode::Validate(const FMetaSoundFrontendDocumen
 	FMetasoundFrontendClass RegisteredClass;
 	if (FMetasoundFrontendRegistryContainer::Get()->GetFrontendClassFromRegistered(FNodeRegistryKey(Metadata), RegisteredClass))
 	{
-		if (RegisteredClass.Metadata.GetIsDeprecated())
+		if (EnumHasAnyFlags(RegisteredClass.Metadata.GetAccessFlags(), EMetasoundFrontendClassAccessFlags::Deprecated))
 		{
 			constexpr bool bIncludeNamespace = true;
 			OutResult.SetMessage(EMessageSeverity::Warning,
diff --git a/Engine/Plugins/Runtime/Metasound/Source/MetasoundEditor/Private/MetasoundEditorGraphSchema.cpp b/Engine/Plugins/Runtime/Metasound/Source/MetasoundEditor/Private/MetasoundEditorGraphSchema.cpp
--- a/Engine/Plugins/Runtime/Metasound/Source/MetasoundEditor/Private/MetasoundEditorGraphSchema.cpp
+++ b/Engine/Plugins/Runtime/Metasound/Source/MetasoundEditor/Private/MetasoundEditorGraphSchema.cpp
@@ -469,6 +469,11 @@ PRAGMA_ENABLE_DEPRECATION_WARNINGS
 				return bCanConnect;
 			}
 
+			virtual EMetasoundFrontendClassAccessFlags GetAccessFlags() const override
+			{
+				return ClassInfo.AccessFlags;
+			}
+
 			virtual FMetaSoundAssetKey GetAssetKey() const override
 			{
 				return FMetaSoundAssetKey(ClassInfo.ClassName, ClassInfo.Version);
@@ -598,6 +603,11 @@ PRAGMA_ENABLE_DEPRECATION_WARNINGS
 				return bCanConnect;
 			}
 
+			virtual EMetasoundFrontendClassAccessFlags GetAccessFlags() const override
+			{
+				return Class.Metadata.GetAccessFlags();
+			}
+
 			virtual FMetaSoundAssetKey GetAssetKey() const override
 			{
 				return FMetaSoundAssetKey(Class.Metadata);
@@ -2378,6 +2388,12 @@ void UMetasoundEditorGraphSchema::GetAssetsGraphHoverMessage(const TArray<FAsset
 
 	OutTooltipText = TEXT("Add MetaSound reference to Graph.");
 
+	const UMetasoundEditorGraph* MetaSoundGraph = CastChecked<UMetasoundEditorGraph>(HoverGraph);
+	const UObject& MetaSound = MetaSoundGraph->GetMetasoundChecked();
+
+	const FMetasoundAssetBase* MetaSoundAsset = IMetasoundUObjectRegistry::Get().GetObjectAsAssetBase(&MetaSound);
+	check(MetaSoundAsset);
+
 	for (const FAssetData& Data : Assets)
 	{
 		UClass* AssetClass = Data.GetClass();
@@ -2390,33 +2406,36 @@ void UMetasoundEditorGraphSchema::GetAssetsGraphHoverMessage(const TArray<FAsset
 			break;
 		}
 
-		const UMetasoundEditorGraph* MetaSoundGraph = CastChecked<UMetasoundEditorGraph>(HoverGraph);
-		const UObject& MetaSound = MetaSoundGraph->GetMetasoundChecked();
-
-		const FMetasoundAssetBase* MetaSoundAsset = IMetasoundUObjectRegistry::Get().GetObjectAsAssetBase(&MetaSound);
-		check(MetaSoundAsset);
-
 		if (UObject* DroppedObject = Data.GetAsset())
 		{
 			FMetasoundAssetBase* DroppedMetaSoundAsset = IMetasoundUObjectRegistry::Get().GetObjectAsAssetBase(DroppedObject);
+
+			const EMetasoundFrontendClassAccessFlags DroppedAccessFlags = DroppedMetaSoundAsset->GetConstDocumentChecked().RootGraph.Metadata.GetAccessFlags();
+			if (!EnumHasAnyFlags(DroppedAccessFlags, EMetasoundFrontendClassAccessFlags::Referenceable))
+			{
+				OutOkIcon = false;
+				OutTooltipText = TEXT("Asset(s) do not contain access flag 'Reference' (see 'Access Flags'). Cannot reference provided MetaSound Asset.");
+				break;
+			}
+
 			if (!DroppedMetaSoundAsset)
 			{
 				OutOkIcon = false;
-				OutTooltipText = TEXT("Asset is not a valid MetaSound.");
+				OutTooltipText = TEXT("Asset(s) not (all) MetaSound(s) that can be referenced.");
 				break;
 			}
 
 			if (MetaSoundAsset->AddingReferenceCausesLoop(*DroppedMetaSoundAsset))
 			{
 				OutOkIcon = false;
-				OutTooltipText = TEXT("Cannot add an asset that would create a reference loop.");
+				OutTooltipText = TEXT("Cannot add asset(s) that would create a reference loop/loops.");
 				break;
 			}
 		}
 		else
 		{
 			OutOkIcon = false;
-			OutTooltipText = TEXT("Asset not found.");
+			OutTooltipText = TEXT("Asset(s) not found.");
 			break;
 		}
 	}
@@ -2492,6 +2511,12 @@ void UMetasoundEditorGraphSchema::DroppedAssetsOnGraph(const TArray<FAssetData>&
 				continue;
 			}
 
+			const EMetasoundFrontendClassAccessFlags DroppedAccessFlags = DroppedMetaSoundAsset->GetConstDocumentChecked().RootGraph.Metadata.GetAccessFlags();
+			if (!EnumHasAnyFlags(DroppedAccessFlags, EMetasoundFrontendClassAccessFlags::Referenceable))
+			{
+				continue;
+			}
+
 			if (MetaSoundAsset->AddingReferenceCausesLoop(*DroppedMetaSoundAsset))
 			{
 				continue;
@@ -2692,6 +2717,12 @@ void UMetasoundEditorGraphSchema::GetFunctionActions(FGraphActionMenuBuilder& Ac
 			}
 		}
 
+		const EMetasoundFrontendClassAccessFlags AccessFlags = Result->GetAccessFlags();
+		if (!EnumHasAnyFlags(AccessFlags, EMetasoundFrontendClassAccessFlags::Referenceable))
+		{
+			return;
+		}
+
 		bool bCausesLoop = false;
 		auto IsMatchingKey = [&ParentAssetKey, &bCausesLoop](const FMetaSoundAssetKey& AssetKey)
 		{
@@ -2885,6 +2916,11 @@ void UMetasoundEditorGraphSchema::QueryNodeClasses(TFunctionRef<void(TUniquePtr<
 			TArray<FMetasoundFrontendClass> FrontendClasses = ISearchEngine::Get().FindAllClasses(bIncludeAllVersions);
 			for (const FMetasoundFrontendClass& Class : FrontendClasses)
 			{
+				if (Class.Metadata.GetClassName().Name.ToString() == TEXT("70A917764AC0FE5934B36E9B6BF74FFA"))
+				{
+					UE_LOG(LogMetasoundEditor, Warning, TEXT("found class \"%s\""), *Class.Metadata.GetDisplayName().ToString());
+				}
+
 				FoundKeys.Add(FNodeRegistryKey(Class.Metadata));
 				Results.Add(MakeUnique<SchemaPrivate::FRegistrySchemaQueryResult>(AssetManager, Class));
 			}
diff --git a/Engine/Plugins/Runtime/Metasound/Source/MetasoundEditor/Private/MetasoundEditorModule.cpp b/Engine/Plugins/Runtime/Metasound/Source/MetasoundEditor/Private/MetasoundEditorModule.cpp
--- a/Engine/Plugins/Runtime/Metasound/Source/MetasoundEditor/Private/MetasoundEditorModule.cpp
+++ b/Engine/Plugins/Runtime/Metasound/Source/MetasoundEditor/Private/MetasoundEditorModule.cpp
@@ -108,6 +108,108 @@ namespace Metasound::Editor
 	{
 		static const FLazyName AssetToolName { "AssetTools" };
 
+		static TOptional<EMetasoundFrontendClassAccessFlags> ShowAccessFlagSelectDialog()
+		{
+			TOptional<EMetasoundFrontendClassAccessFlags> TargetFlags = EMetasoundFrontendClassAccessFlags::None;
+			if (!GEditor)
+			{
+				return TargetFlags;
+			}
+
+			TSharedRef<SWindow> DialogWindow =
+				SNew(SWindow)
+				.Title(LOCTEXT("MetaSoundSetAccessFlagTitle", "Set MetaSound(s) Access Flags..."))
+				.SupportsMinimize(false)
+				.SupportsMaximize(false)
+				.SizingRule(ESizingRule::Autosized)
+				.AutoCenter(EAutoCenter::PreferredWorkArea);
+				
+			TSharedPtr<SVerticalBox> AccessFlagVerticalBox;
+			TSharedRef<SBox> DialogContent = SNew(SBox)
+				.HAlign(HAlign_Center)
+				.VAlign(VAlign_Center)
+				[
+					SAssignNew(AccessFlagVerticalBox, SVerticalBox)
+				];
+
+			const UEnum* FlagsEnum = StaticEnum<EMetasoundFrontendClassAccessFlags>();
+			check(FlagsEnum);
+			for (EMetasoundFrontendClassAccessFlags Flag : TEnumRange<EMetasoundFrontendClassAccessFlags>())
+			{
+				FText FlagDisplayName;
+				FlagsEnum->FindDisplayNameTextByValue(FlagDisplayName, (int64)(Flag));
+				AccessFlagVerticalBox->AddSlot()
+					.HAlign(HAlign_Left)
+					.VAlign(VAlign_Center)
+					[
+						SNew(SCheckBox)
+							.IsChecked_Lambda([&TargetFlags, Flag]()
+							{
+								if (EnumHasAllFlags(*TargetFlags, Flag))
+								{
+									return ECheckBoxState::Checked;
+								}
+								return ECheckBoxState::Unchecked;
+							})
+							.OnCheckStateChanged_Lambda([&TargetFlags, Flag](ECheckBoxState InNewState)
+							{
+								if (InNewState == ECheckBoxState::Checked)
+								{
+									EnumAddFlags(*TargetFlags, Flag);
+								}
+								else
+								{
+									EnumRemoveFlags(*TargetFlags, Flag);
+								}
+							})
+							[
+								SNew(STextBlock)
+									.Text(FlagDisplayName)
+									.Font(IDetailLayoutBuilder::GetDetailFont())
+							]
+					];
+			}
+
+			AccessFlagVerticalBox->AddSlot()
+				.HAlign(HAlign_Center)
+				.VAlign(VAlign_Center)
+			[
+				SNew(SHorizontalBox)
+				+ SHorizontalBox::Slot()
+				.AutoWidth()
+				.HAlign(HAlign_Left)
+				.VAlign(VAlign_Bottom)
+				[
+					SNew(SButton)
+					.Text(LOCTEXT("MetaSoundSetAccessFlagAction_Accept", "Accept"))
+					.OnClicked_Lambda([&DialogWindow]()
+					{
+						DialogWindow->RequestDestroyWindow();
+						return FReply::Handled();
+					})
+				]
+				+ SHorizontalBox::Slot()
+				.AutoWidth()
+				.HAlign(HAlign_Right)
+				.VAlign(VAlign_Bottom)
+				[
+					SNew(SButton)
+					.Text(LOCTEXT("MetaSoundSetAccessFlagAction_Cancel", "Cancel"))
+					.OnClicked_Lambda([&TargetFlags, &DialogWindow]()
+					{
+						TargetFlags.Reset();
+						DialogWindow->RequestDestroyWindow();
+						return FReply::Handled();
+					})
+				]
+			];
+				
+			DialogWindow->SetContent(DialogContent);
+			GEditor->EditorAddModalWindow(DialogWindow);
+
+			return TargetFlags;
+		}
+
 		void AddVersionDocumentFolderMenuEntry(UToolMenu* Menu)
 		{
 			if (!Menu)
@@ -122,7 +224,7 @@ namespace Metasound::Editor
 			}
 
 			bool bRecursePaths = false;
-			auto CreateVersioningExecutionLambda = [&Context](bool bRecursePaths) -> FExecuteAction
+			auto CreateVersioningExecutionLambda = [Context](bool bRecursePaths) -> FExecuteAction
 			{
 				return FExecuteAction::CreateLambda([FolderPaths = Context->GetSelectedPackagePaths(), bRecursePaths]()
 				{
@@ -201,6 +303,87 @@ namespace Metasound::Editor
 				FSlateIcon(FAppStyle::GetAppStyleSetName(), "Icons.Save"),
 				FToolUIActionChoice(CreateVersioningExecutionLambda(true /* bRecursePaths */))
 			));
+
+			auto CreateSetAccessFlagsExecutionLambda = [Context](bool bRecursePaths) -> FExecuteAction
+			{
+				return FExecuteAction::CreateLambda([bRecursePaths, FolderPaths = Context->GetSelectedPackagePaths()]()
+				{
+					using namespace Frontend;
+
+					const TOptional<EMetasoundFrontendClassAccessFlags> TargetAccessFlags = ModulePrivate::ShowAccessFlagSelectDialog();
+					if (!TargetAccessFlags.IsSet())
+					{
+						return;
+					}
+
+					FNotificationInfo Info(LOCTEXT("VersioningMetaSoundsNotifyTitle", "Setting MetaSound(s) AccessFlags..."));
+					Info.bFireAndForget = false;
+					Info.ExpireDuration = 0.0f;
+					Info.bUseThrobber = true;
+					TSharedPtr<SNotificationItem> Notify = FSlateNotificationManager::Get().AddNotification(Info);
+					if (Notify.IsValid())
+					{
+						Notify->SetCompletionState(SNotificationItem::CS_Pending);
+					}
+
+					const IMetaSoundAssetManager::FVersionAssetResults Results = IMetaSoundAssetManager::GetChecked().SetAccessFlagsOnAssetsInFolders(FolderPaths, *TargetAccessFlags, bRecursePaths);
+					if (!Results.DocumentsFoundInPackages())
+					{
+						if (Notify.IsValid())
+						{
+							Notify->SetCompletionState(SNotificationItem::CS_None);
+						}
+						UE_LOG(LogMetasoundEditor, Display, TEXT("No MetaSound(s) Versioned: Folder/child folder(s) contain(s) no MetaSound asset(s)."));
+					}
+					else if (Results.PackagesToReserialize.IsEmpty())
+					{
+						if (Notify.IsValid())
+						{
+							Notify->SetCompletionState(SNotificationItem::CS_None);
+						}
+						UE_LOG(LogMetasoundEditor, Display, TEXT("No MetaSound(s) Versioned: Folder/child folder(s) contain(s) no MetaSound asset(s) requiring versioning."));
+					}
+					else if (!Results.FailedPackages.IsEmpty())
+					{
+						if (Notify.IsValid())
+						{
+							Notify->SetCompletionState(SNotificationItem::CS_Fail);
+						}
+						UE_LOG(LogMetasoundEditor, Display, TEXT("%i MetaSound(s) Reserialized but '%i' failed: See output log for details."), Results.PackagesToReserialize.Num(), Results.FailedPackages.Num());
+						FEditorFileUtils::PromptToCheckoutPackages(false, Results.PackagesToReserialize);
+					}
+					else
+					{
+						if (Notify.IsValid())
+						{
+							Notify->SetCompletionState(SNotificationItem::CS_Success);
+						}
+						UE_LOG(LogMetasoundEditor, Display, TEXT("%i MetaSound(s) Successfully Reserialized."), Results.PackagesToReserialize.Num());
+						FEditorFileUtils::PromptToCheckoutPackages(false, Results.PackagesToReserialize);
+					}
+
+					if (Notify.IsValid())
+					{
+						Notify->ExpireAndFadeout();
+					}
+				});
+			};
+
+			Menu->AddMenuEntry("MetaSounds", FToolMenuEntry::InitMenuEntry(
+				"SetMetaSoundAccessFlags",
+				LOCTEXT("SetAccessFlagsMetaSoundAssetsMenuEntry", "Set MetaSound Access Flags"),
+				LOCTEXT("SetMetaSoundFlags", "Sets given AccessFlag values on all MetaSounds in the given folder."),
+				FSlateIcon(FAppStyle::GetAppStyleSetName(), "Icons.Save"),
+				FToolUIActionChoice(CreateSetAccessFlagsExecutionLambda(false /* bRecursePaths */))
+			));
+
+			Menu->AddMenuEntry("MetaSounds", FToolMenuEntry::InitMenuEntry(
+				"SetMetaSoundAccessFlagsRecursively",
+				LOCTEXT("SetAccessFlagsMetaSoundAssetsMenuEntry_Recursive", "Set MetaSound Access Flags (Recursive)"),
+				LOCTEXT("SetMetaSoundFlags_RecurseContext", "Sets given AccessFlag values on all MetaSounds in the given folder (and all sub-folders)."),
+				FSlateIcon(FAppStyle::GetAppStyleSetName(), "Icons.Save"),
+				FToolUIActionChoice(CreateSetAccessFlagsExecutionLambda(true /* bRecursePaths */))
+			));
 		}
 	} // namespace ModulePrivate
 
diff --git a/Engine/Plugins/Runtime/Metasound/Source/MetasoundEditor/Public/MetasoundEditorGraphSchema.h b/Engine/Plugins/Runtime/Metasound/Source/MetasoundEditor/Public/MetasoundEditorGraphSchema.h
--- a/Engine/Plugins/Runtime/Metasound/Source/MetasoundEditor/Public/MetasoundEditorGraphSchema.h
+++ b/Engine/Plugins/Runtime/Metasound/Source/MetasoundEditor/Public/MetasoundEditorGraphSchema.h
@@ -49,6 +49,7 @@ namespace Metasound::Editor
 		virtual bool CanConnectInputOfTypeAndAccess(FName TypeName, EMetasoundFrontendVertexAccessType AccessType) const = 0;
 		virtual bool CanConnectOutputOfTypeAndAccess(FName TypeName, EMetasoundFrontendVertexAccessType AccessType) const = 0;
 
+		virtual EMetasoundFrontendClassAccessFlags GetAccessFlags() const = 0;
 		virtual FMetaSoundAssetKey GetAssetKey() const = 0;
 		virtual const TArray<FText>& GetCategoryHierarchy() const = 0;
 		virtual FText GetDisplayName() const = 0;
diff --git a/Engine/Plugins/Runtime/Metasound/Source/MetasoundEngine/Private/Metasound.cpp b/Engine/Plugins/Runtime/Metasound/Source/MetasoundEngine/Private/Metasound.cpp
--- a/Engine/Plugins/Runtime/Metasound/Source/MetasoundEngine/Private/Metasound.cpp
+++ b/Engine/Plugins/Runtime/Metasound/Source/MetasoundEngine/Private/Metasound.cpp
@@ -133,6 +133,11 @@ const FMetasoundFrontendDocument& UMetaSoundPatch::GetConstDocument() const
 	return RootMetaSoundDocument;
 }
 
+EMetasoundFrontendClassAccessFlags UMetaSoundPatch::GetDefaultAccessFlags() const
+{
+	return EMetasoundFrontendClassAccessFlags::Referenceable;
+}
+
 #if WITH_EDITOR
 void UMetaSoundPatch::PreDuplicate(FObjectDuplicationParameters& DupParams)
 {
diff --git a/Engine/Plugins/Runtime/Metasound/Source/MetasoundEngine/Private/MetasoundAssetSubsystem.cpp b/Engine/Plugins/Runtime/Metasound/Source/MetasoundEngine/Private/MetasoundAssetSubsystem.cpp
--- a/Engine/Plugins/Runtime/Metasound/Source/MetasoundEngine/Private/MetasoundAssetSubsystem.cpp
+++ b/Engine/Plugins/Runtime/Metasound/Source/MetasoundEngine/Private/MetasoundAssetSubsystem.cpp
@@ -251,6 +251,12 @@ namespace Metasound::Engine
 
 		METASOUND_TRACE_CPUPROFILER_EVENT_SCOPE(FMetaSoundAssetManager::AddOrUpdateFromObjectInternal);
 
+		// Don't add temporary assets used for diffing
+		if (UPackage* Package = InObject.GetPackage(); !Package || Package->HasAnyPackageFlags(PKG_ForDiffing))
+		{
+			return FMetaSoundAssetKey::GetInvalid();
+		}
+
 		const FMetasoundAssetBase* MetaSoundAsset = Metasound::IMetasoundUObjectRegistry::Get().GetObjectAsAssetBase(&InObject);
 		check(MetaSoundAsset);
 
@@ -1233,6 +1239,69 @@ namespace Metasound::Engine
 	}
 
 #if WITH_EDITORONLY_DATA
+	FMetaSoundAssetManager::FVersionAssetResults FMetaSoundAssetManager::SetAccessFlagsOnAssetsInFolders(const TArray<FString>& FolderPaths, EMetasoundFrontendClassAccessFlags Flags, bool bRecursePaths) const
+	{
+		TArray<FTopLevelAssetPath> ClassNames;
+		IMetasoundUObjectRegistry::Get().IterateRegisteredUClasses([&ClassNames](UClass& InClass)
+		{
+			ClassNames.Add(InClass.GetClassPathName());
+		});
+
+		FARFilter Filter;
+		Filter.ClassPaths = ClassNames;
+		Filter.bRecursivePaths = bRecursePaths;
+		Filter.bRecursiveClasses = true;
+		Algo::Transform(FolderPaths, Filter.PackagePaths, [](const FString& PackagePath)
+		{
+			return FName(*PackagePath);
+		});
+
+		FVersionAssetResults Results;
+		FAssetRegistryModule& AssetRegistryModule = FModuleManager::LoadModuleChecked<FAssetRegistryModule>("AssetRegistry");
+
+		if (AssetRegistryModule.Get().IsGathering())
+		{
+			UE_LOG(LogMetaSound, Error, TEXT("Asset Registry is still scanning, wait to update MetaSound AccessFlags in given folder(s) until initial scan is complete."));
+			return Results;
+		}
+
+		AssetRegistryModule.Get().EnumerateAssets(Filter, [&Flags, &Results](const FAssetData& AssetData)
+		{
+			using namespace Frontend;
+
+			const FTopLevelAssetPath PackagePath(AssetData.PackageName, AssetData.AssetName);
+			const FMetaSoundAssetClassInfo ClassInfo(AssetData);
+			if (ClassInfo.bIsValid)
+			{
+				if (ClassInfo.AccessFlags == Flags)
+				{
+					Results.PackagesUpToDate.Add(PackagePath);
+					UE_LOG(LogMetaSound, Display, TEXT("MetaSound '%s' already has given flag(s), skipping update."), *AssetData.GetFullName());
+					return true;
+				}
+			}
+
+			UObject* MetaSoundObject = AssetData.GetAsset();
+			FMetasoundAssetBase* MetaSound = IMetasoundUObjectRegistry::Get().GetObjectAsAssetBase(MetaSoundObject);
+			if (MetaSound)
+			{
+				FMetaSoundFrontendDocumentBuilder& DocBuilder = FDocumentBuilderRegistry::GetChecked().FindOrBeginBuilding(MetaSoundObject);
+				constexpr bool bAlwaysMarkDirty = true;
+				MetaSoundObject->Modify(bAlwaysMarkDirty);
+				Results.PackagesToReserialize.Add(AssetData.GetPackage());
+				DocBuilder.SetAccessFlags(Flags);
+				UE_LOG(LogMetaSound, Display, TEXT("MetaSound '%s' access flags updated."), *AssetData.GetFullName());
+				return true;
+			}
+
+			Results.FailedPackages.Add(PackagePath);
+			UE_LOG(LogMetaSound, Error, TEXT("MetaSound asset '%s' failed to load: asset AccessFlags not updated."), *AssetData.GetFullName());
+			return true;
+		});
+
+		return Results;
+	}
+
 	FMetaSoundAssetManager::FVersionAssetResults FMetaSoundAssetManager::VersionAssetsInFolders(const TArray<FString>& FolderPaths, bool bRecursePaths) const
 	{
 		TArray<FTopLevelAssetPath> ClassNames;
diff --git a/Engine/Plugins/Runtime/Metasound/Source/MetasoundEngine/Private/MetasoundEngineModule.cpp b/Engine/Plugins/Runtime/Metasound/Source/MetasoundEngine/Private/MetasoundEngineModule.cpp
--- a/Engine/Plugins/Runtime/Metasound/Source/MetasoundEngine/Private/MetasoundEngineModule.cpp
+++ b/Engine/Plugins/Runtime/Metasound/Source/MetasoundEngine/Private/MetasoundEngineModule.cpp
@@ -188,6 +188,14 @@ namespace Metasound::Engine
 			return;
 		}
 
+		// Ignore requests to prime on assets with AccessFlag 'Referenceable' unset, deffering
+		// registration to initial execution if existing reference from another asset requires it.
+		const FMetaSoundAssetClassInfo ClassInfo(InAssetData);
+		if (!EnumHasAnyFlags(ClassInfo.AccessFlags, EMetasoundFrontendClassAccessFlags::Referenceable))
+		{
+			return;
+		}
+
 		// If an object's class could not be found, ignore this asset.  This can hit for non-MetaSound assets
 		// and it is up to the system in charge of interacting with that asset or the loading behavior to
 		// report the failed load of the class.
diff --git a/Engine/Plugins/Runtime/Metasound/Source/MetasoundEngine/Private/MetasoundSource.cpp b/Engine/Plugins/Runtime/Metasound/Source/MetasoundEngine/Private/MetasoundSource.cpp
--- a/Engine/Plugins/Runtime/Metasound/Source/MetasoundEngine/Private/MetasoundSource.cpp
+++ b/Engine/Plugins/Runtime/Metasound/Source/MetasoundEngine/Private/MetasoundSource.cpp
@@ -438,6 +438,11 @@ const UClass& UMetaSoundSource::GetBuilderUClass() const
 	return *UMetaSoundSourceBuilder::StaticClass();
 }
 
+EMetasoundFrontendClassAccessFlags UMetaSoundSource::GetDefaultAccessFlags() const
+{
+	return EMetasoundFrontendClassAccessFlags::None;
+}
+
 const FMetasoundFrontendDocument& UMetaSoundSource::GetConstDocument() const
 {
 	return RootMetasoundDocument;
diff --git a/Engine/Plugins/Runtime/Metasound/Source/MetasoundEngine/Public/Metasound.h b/Engine/Plugins/Runtime/Metasound/Source/MetasoundEngine/Public/Metasound.h
--- a/Engine/Plugins/Runtime/Metasound/Source/MetasoundEngine/Public/Metasound.h
+++ b/Engine/Plugins/Runtime/Metasound/Source/MetasoundEngine/Public/Metasound.h
@@ -152,6 +152,7 @@ public:
 	UE_API virtual const UClass& GetBaseMetaSoundUClass() const final override;
 	UE_API virtual const UClass& GetBuilderUClass() const final override;
 	UE_API virtual const FMetasoundFrontendDocument& GetConstDocument() const override;
+	UE_API virtual EMetasoundFrontendClassAccessFlags GetDefaultAccessFlags() const final override;
 
 #if WITH_EDITOR
 	UE_API virtual void PreDuplicate(FObjectDuplicationParameters& DupParams) override;
diff --git a/Engine/Plugins/Runtime/Metasound/Source/MetasoundEngine/Public/MetasoundAssetSubsystem.h b/Engine/Plugins/Runtime/Metasound/Source/MetasoundEngine/Public/MetasoundAssetSubsystem.h
--- a/Engine/Plugins/Runtime/Metasound/Source/MetasoundEngine/Public/MetasoundAssetSubsystem.h
+++ b/Engine/Plugins/Runtime/Metasound/Source/MetasoundEngine/Public/MetasoundAssetSubsystem.h
@@ -148,6 +148,7 @@ PRAGMA_ENABLE_DEPRECATION_WARNINGS
 		UE_API virtual bool TryGetAssetIDFromClassName(const FMetasoundFrontendClassName& InClassName, FGuid& OutGuid) const override;
 		UE_API virtual bool TryLoadReferencedAssets(const FMetasoundAssetBase& InAssetBase, TArray<FMetasoundAssetBase*>& OutReferencedAssets) const override;
 #if WITH_EDITORONLY_DATA
+		UE_API FMetaSoundAssetManager::FVersionAssetResults SetAccessFlagsOnAssetsInFolders(const TArray<FString>& FolderPaths, EMetasoundFrontendClassAccessFlags Flags, bool bRecursePaths) const override;
 		UE_API virtual FVersionAssetResults VersionAssetsInFolders(const TArray<FString>& FolderPaths, bool bRecursePaths = true) const override;
 #endif // WITH_EDITORONLY_DATA
 
diff --git a/Engine/Plugins/Runtime/Metasound/Source/MetasoundEngine/Public/MetasoundSource.h b/Engine/Plugins/Runtime/Metasound/Source/MetasoundEngine/Public/MetasoundSource.h
--- a/Engine/Plugins/Runtime/Metasound/Source/MetasoundEngine/Public/MetasoundSource.h
+++ b/Engine/Plugins/Runtime/Metasound/Source/MetasoundEngine/Public/MetasoundSource.h
@@ -319,6 +319,7 @@ public:
 
 	UE_API virtual const UClass& GetBaseMetaSoundUClass() const final override;
 	UE_API virtual const UClass& GetBuilderUClass() const final override;
+	UE_API virtual EMetasoundFrontendClassAccessFlags GetDefaultAccessFlags() const final override;
 
 protected:
 	UE_DEPRECATED(5.6, "AccessPtrs are actively being deprecated. Writable access outside of the builder API "
diff --git a/Engine/Plugins/Runtime/Metasound/Source/MetasoundEngineTest/Private/EngineTestMetaSoundAutomatedNodeTest.cpp b/Engine/Plugins/Runtime/Metasound/Source/MetasoundEngineTest/Private/EngineTestMetaSoundAutomatedNodeTest.cpp
--- a/Engine/Plugins/Runtime/Metasound/Source/MetasoundEngineTest/Private/EngineTestMetaSoundAutomatedNodeTest.cpp
+++ b/Engine/Plugins/Runtime/Metasound/Source/MetasoundEngineTest/Private/EngineTestMetaSoundAutomatedNodeTest.cpp
@@ -13,6 +13,7 @@
 #include "Misc/AutomationTest.h"
 #include "MetasoundAudioBuffer.h"
 #include "MetasoundFrontendDataTypeRegistry.h"
+#include "MetasoundFrontendDocument.h"
 #include "MetasoundFrontendRegistries.h"
 #include "MetasoundFrontendSearchEngine.h"
 #include "MetasoundLog.h"
@@ -21,6 +22,7 @@
 #include "MetasoundVariableNodes.h"
 #include "MetasoundVertex.h"
 #include "MetasoundVertexData.h"
+#include "Misc/EnumClassFlags.h"
 #include "Templates/UniquePtr.h"
 #include "Tests/AutomationCommon.h"
 
@@ -762,7 +764,7 @@ PRAGMA_ENABLE_DEPRECATION_WARNINGS
 			{
 				continue;
 			}
-			else if (!bInIncludeDeprecatedNodes && NodeClass.Metadata.GetIsDeprecated())
+			else if (!bInIncludeDeprecatedNodes && EnumHasAnyFlags(NodeClass.Metadata.GetAccessFlags(), EMetasoundFrontendClassAccessFlags::Deprecated))
 			{
 				// Optionally skip deprecated classes. 
 				continue;
@@ -795,7 +797,7 @@ PRAGMA_ENABLE_DEPRECATION_WARNINGS
 			{
 				continue;
 			}
-			else if (!bInIncludeDeprecatedNodes && NodeClass.Metadata.GetIsDeprecated())
+			else if (!bInIncludeDeprecatedNodes && EnumHasAnyFlags(NodeClass.Metadata.GetAccessFlags(), EMetasoundFrontendClassAccessFlags::Deprecated))
 			{
 				// Optionally skip deprecated classes. 
 				continue;
diff --git a/Engine/Plugins/Runtime/Metasound/Source/MetasoundFrontend/Private/MetasoundArrayNodes.cpp b/Engine/Plugins/Runtime/Metasound/Source/MetasoundFrontend/Private/MetasoundArrayNodes.cpp
--- a/Engine/Plugins/Runtime/Metasound/Source/MetasoundFrontend/Private/MetasoundArrayNodes.cpp
+++ b/Engine/Plugins/Runtime/Metasound/Source/MetasoundFrontend/Private/MetasoundArrayNodes.cpp
@@ -1,6 +1,7 @@
 // Copyright Epic Games, Inc. All Rights Reserved.
 
 #include "MetasoundArrayNodes.h"
+#include "Misc/EnumClassFlags.h"
 
 #define LOCTEXT_NAMESPACE "MetasoundFrontend"
 
@@ -12,7 +13,7 @@ namespace Metasound
 		{
 			FNodeClassMetadata Metadata
 			{
-				FNodeClassName{FName("Array"), InOperatorName, InDataTypeName},
+				FNodeClassName { FName("Array"), InOperatorName, InDataTypeName },
 				MajorVersion, 
 				MinorVersion,
 				InDisplayName, 
@@ -22,10 +23,13 @@ namespace Metasound
 				InDefaultInterface,
 				{ METASOUND_LOCTEXT("ArrayCategory", "Array") },
 				{ METASOUND_LOCTEXT("MetasoundArrayKeyword", "Array") },
-				FNodeDisplayStyle{}, 
-				bIsDeprecated
 			};
 
+			if (bIsDeprecated)
+			{
+				EnumAddFlags(Metadata.AccessFlags, ENodeClassAccessFlags::Deprecated);
+			}
+
 			return Metadata;
 		}
 	}
diff --git a/Engine/Plugins/Runtime/Metasound/Source/MetasoundFrontend/Private/MetasoundAssetManager.cpp b/Engine/Plugins/Runtime/Metasound/Source/MetasoundFrontend/Private/MetasoundAssetManager.cpp
--- a/Engine/Plugins/Runtime/Metasound/Source/MetasoundFrontend/Private/MetasoundAssetManager.cpp
+++ b/Engine/Plugins/Runtime/Metasound/Source/MetasoundFrontend/Private/MetasoundAssetManager.cpp
@@ -305,11 +305,6 @@ namespace Metasound::Frontend
 		check(!AssetManagerPrivate::Instance.IsValid());
 		AssetManagerPrivate::Instance = MoveTemp(InInterface);
 	}
-
-	bool IMetaSoundAssetManager::IsTesting() const
-	{
-		return false;
-	}
 } // namespace Metasound::Frontend
 
 FMetaSoundAssetKey::FMetaSoundAssetKey(const FMetasoundFrontendClassName& InClassName, const FMetasoundFrontendVersionNumber& InVersion)
diff --git a/Engine/Plugins/Runtime/Metasound/Source/MetasoundFrontend/Private/MetasoundFrontendDocument.cpp b/Engine/Plugins/Runtime/Metasound/Source/MetasoundFrontend/Private/MetasoundFrontendDocument.cpp
--- a/Engine/Plugins/Runtime/Metasound/Source/MetasoundFrontend/Private/MetasoundFrontendDocument.cpp
+++ b/Engine/Plugins/Runtime/Metasound/Source/MetasoundFrontend/Private/MetasoundFrontendDocument.cpp
@@ -288,6 +288,21 @@ void FMetasoundFrontendDocumentModifyContext::AddNodeIDsModified(const TSet<FGui
 }
 #endif // WITH_EDITORONLY_DATA
 
+FString LexToString(const EMetasoundFrontendClassAccessFlags& InFlags)
+{
+	switch (InFlags)
+	{
+		case EMetasoundFrontendClassAccessFlags::Deprecated:
+			return TEXT("Deprecated");
+
+		case EMetasoundFrontendClassAccessFlags::Referenceable:
+			return TEXT("Referenceable");
+
+		default:
+			return TEXT("Unset");
+	}
+}
+
 FMetasoundCommentNodeIntVector::FMetasoundCommentNodeIntVector(const FIntVector2& InValue)
 	: FIntVector2(InValue)
 {
@@ -944,6 +959,41 @@ FMetasoundFrontendClassInterface FMetasoundFrontendClassInterface::GenerateClass
 	return ClassInterface;
 }
 
+void FMetasoundFrontendClassMetadata::AddAccessFlags(EMetasoundFrontendClassAccessFlags InAccessFlags)
+{
+	AccessFlags |= static_cast<uint16>(InAccessFlags);
+}
+
+void FMetasoundFrontendClassMetadata::ClearAccessFlags()
+{
+	AccessFlags = static_cast<uint16>(EMetasoundFrontendClassAccessFlags::None);
+}
+
+EMetasoundFrontendClassType FMetasoundFrontendClassMetadata::GetType() const
+{
+	return Type;
+}
+
+const FMetasoundFrontendVersionNumber& FMetasoundFrontendClassMetadata::GetVersion() const
+{
+	return Version;
+}
+
+EMetasoundFrontendClassAccessFlags FMetasoundFrontendClassMetadata::GetAccessFlags() const
+{
+	return static_cast<EMetasoundFrontendClassAccessFlags>(AccessFlags);
+}
+
+void FMetasoundFrontendClassMetadata::RemoveAccessFlags(EMetasoundFrontendClassAccessFlags InAccessFlags)
+{
+	AccessFlags &= ~static_cast<uint16>(InAccessFlags);
+}
+
+void FMetasoundFrontendClassMetadata::SetAccessFlags(EMetasoundFrontendClassAccessFlags InAccessFlags)
+{
+	AccessFlags = static_cast<uint16>(InAccessFlags);
+}
+
 #if WITH_EDITOR
 void FMetasoundFrontendClassMetadata::SetAuthor(const FString& InAuthor)
 {
@@ -1180,8 +1230,7 @@ FMetasoundFrontendClassMetadata FMetasoundFrontendClassMetadata::GenerateClassMe
 	NewMetadata.SetAuthor(InNodeClassMetadata.Author);
 	NewMetadata.SetKeywords(InNodeClassMetadata.Keywords);
 	NewMetadata.SetCategoryHierarchy(InNodeClassMetadata.CategoryHierarchy);
-
-	NewMetadata.bIsDeprecated = InNodeClassMetadata.bDeprecated;
+	NewMetadata.AddAccessFlags(static_cast<EMetasoundFrontendClassAccessFlags>(InNodeClassMetadata.AccessFlags));
 #endif // WITH_EDITOR
 
 	return NewMetadata;
diff --git a/Engine/Plugins/Runtime/Metasound/Source/MetasoundFrontend/Private/MetasoundFrontendDocumentBuilder.cpp b/Engine/Plugins/Runtime/Metasound/Source/MetasoundFrontend/Private/MetasoundFrontendDocumentBuilder.cpp
--- a/Engine/Plugins/Runtime/Metasound/Source/MetasoundFrontend/Private/MetasoundFrontendDocumentBuilder.cpp
+++ b/Engine/Plugins/Runtime/Metasound/Source/MetasoundFrontend/Private/MetasoundFrontendDocumentBuilder.cpp
@@ -749,6 +749,11 @@ const FMetasoundFrontendDocument& UMetaSoundBuilderDocument::GetConstDocument()
 	return Document;
 }
 
+EMetasoundFrontendClassAccessFlags UMetaSoundBuilderDocument::GetDefaultAccessFlags() const
+{
+	return EMetasoundFrontendClassAccessFlags::None;
+}
+
 const UClass& UMetaSoundBuilderDocument::GetBaseMetaSoundUClass() const
 {
 	checkf(MetaSoundUClass, TEXT("BaseMetaSoundUClass must be set upon creation of UMetaSoundBuilderDocument instance"));
@@ -793,6 +798,31 @@ FMetaSoundFrontendDocumentBuilder::~FMetaSoundFrontendDocumentBuilder()
 	FinishBuilding();
 }
 
+void FMetaSoundFrontendDocumentBuilder::AddAccessFlags(EMetasoundFrontendClassAccessFlags AccessFlags)
+{
+	GetDocumentChecked().RootGraph.Metadata.AddAccessFlags(AccessFlags);
+}
+
+void FMetaSoundFrontendDocumentBuilder::ClearAccessFlags()
+{
+	GetDocumentChecked().RootGraph.Metadata.ClearAccessFlags();
+}
+
+EMetasoundFrontendClassAccessFlags FMetaSoundFrontendDocumentBuilder::GetAccessFlags() const
+{
+	return GetDocumentChecked().RootGraph.Metadata.GetAccessFlags();
+}
+
+void FMetaSoundFrontendDocumentBuilder::RemoveAccessFlags(EMetasoundFrontendClassAccessFlags AccessFlags)
+{
+	GetDocumentChecked().RootGraph.Metadata.RemoveAccessFlags(AccessFlags);
+}
+
+void FMetaSoundFrontendDocumentBuilder::SetAccessFlags(EMetasoundFrontendClassAccessFlags AccessFlags)
+{
+	GetDocumentChecked().RootGraph.Metadata.SetAccessFlags(AccessFlags);
+}
+
 const FMetasoundFrontendClass* FMetaSoundFrontendDocumentBuilder::AddDependency(FMetasoundFrontendClass NewDependency)
 {
 	using namespace Metasound::Frontend;
@@ -1727,6 +1757,15 @@ const FMetasoundFrontendNode* FMetaSoundFrontendDocumentBuilder::AddGraphNode(co
 {
 	using FNodeRegistryKey = Metasound::Frontend::FNodeRegistryKey;
 
+	const EMetasoundFrontendClassAccessFlags AccessFlags = InGraphClass.Metadata.GetAccessFlags();
+	if (!EnumHasAnyFlags(AccessFlags, EMetasoundFrontendClassAccessFlags::Referenceable))
+	{
+		UE_LOG(LogMetaSound, Warning, TEXT("Failed to add new node of graph class '%s': Class access flag '%s' not set."),
+			*InGraphClass.Metadata.GetClassName().ToString(),
+			*LexToString(EMetasoundFrontendClassAccessFlags::Referenceable));
+		return nullptr;
+	}
+
 	auto FinalizeNode = [](FMetasoundFrontendNode& InOutNode, const Metasound::Frontend::FNodeRegistryKey& ClassKey)
 	{
 #if WITH_EDITOR
@@ -1787,6 +1826,14 @@ const FMetasoundFrontendNode* FMetaSoundFrontendDocumentBuilder::AddNodeByClassN
 			return nullptr;
 		}
 
+		if (!EnumHasAnyFlags(RegisteredClass.Metadata.GetAccessFlags(), EMetasoundFrontendClassAccessFlags::Referenceable))
+		{
+			UE_LOG(LogMetaSound, Warning, TEXT("Failed to add new node by name of class '%s': Class access flag '%s' not set."),
+				*InClassName.ToString(),
+				*LexToString(EMetasoundFrontendClassAccessFlags::Referenceable));
+			return nullptr;
+		}
+
 		// Dependency is considered "External" when looked up or added as a dependency to a graph
 		RegisteredClass.Metadata.SetType(EMetasoundFrontendClassType::External);
 		ClassKey = FNodeRegistryKey(RegisteredClass.Metadata);
@@ -1815,6 +1862,15 @@ const FMetasoundFrontendNode* FMetaSoundFrontendDocumentBuilder::AddNodeByTempla
 
 	const FMetasoundFrontendClass& TemplateClass = InTemplate.GetFrontendClass();
 	checkf(TemplateClass.Metadata.GetType() == EMetasoundFrontendClassType::Template, TEXT("INodeTemplate ClassType must always be 'Template'"));
+
+	if (!EnumHasAnyFlags(TemplateClass.Metadata.GetAccessFlags(), EMetasoundFrontendClassAccessFlags::Referenceable))
+	{
+		UE_LOG(LogMetaSound, Warning, TEXT("Failed to add new template node of class '%s': Class access flag '%s' not set."),
+			*InTemplate.GetClassName().ToString(),
+			*LexToString(EMetasoundFrontendClassAccessFlags::Referenceable));
+		return nullptr;
+	}
+
 	const FMetasoundFrontendClass* Dependency = FindDependency(TemplateClass.Metadata);
 	if (!Dependency)
 	{
@@ -3821,6 +3877,9 @@ void FMetaSoundFrontendDocumentBuilder::InitDocument(const FMetasoundFrontendDoc
 			FModifyInterfaceOptions Options({ }, InitVersions);
 			ModifyInterfaces(MoveTemp(Options));
 		}
+
+		// 1d. Set default access flags as provided by document object class
+		SetAccessFlags(DocumentInterface->GetDefaultAccessFlags());
 	}
 }
 
@@ -4556,7 +4615,8 @@ bool FMetaSoundFrontendDocumentBuilder::RemoveNodeInputDefault(const FGuid& InNo
 	const IDocumentGraphNodeCache& NodeCache = DocumentCache->GetNodeCache(PageID);
 	if (const int32* NodeIndex = NodeCache.FindNodeIndex(InNodeID))
 	{
-		FMetasoundFrontendGraph& Graph = GetDocumentChecked().RootGraph.FindGraphChecked(PageID);
+		FMetasoundFrontendDocument& Document = GetDocumentChecked();
+		FMetasoundFrontendGraph& Graph = Document.RootGraph.FindGraphChecked(PageID);
 		FMetasoundFrontendNode& Node = Graph.Nodes[*NodeIndex];
 
 		auto IsVertex = [&InVertexID](const FMetasoundFrontendVertex& Vertex) { return Vertex.VertexID == InVertexID; };
@@ -4577,6 +4637,11 @@ bool FMetaSoundFrontendDocumentBuilder::RemoveNodeInputDefault(const FGuid& InNo
 				}
 
 				Node.InputLiterals.RemoveAtSwap(LiteralIndex, EAllowShrinking::No);
+
+#if WITH_EDITORONLY_DATA
+				Document.Metadata.ModifyContext.AddNodeIDModified(Node.GetID());
+#endif // WITH_EDITORONLY_DATA
+
 				if (LiteralIndex != LastIndex)
 				{
 					const FOnMetaSoundFrontendDocumentMutateNodeInputLiteralArray& OnNodeInputLiteralSet = NodeDelegates.OnNodeInputLiteralSet;
@@ -5923,7 +5988,10 @@ void FMetaSoundFrontendDocumentBuilder::SetCategoryHierarchy(const TArray<FText>
 
 void FMetaSoundFrontendDocumentBuilder::SetIsDeprecated(const bool bInIsDeprecated)
 {
+PRAGMA_DISABLE_DEPRECATION_WARNINGS
 	DocumentInterface->GetDocument().RootGraph.Metadata.SetIsDeprecated(bInIsDeprecated);
+PRAGMA_ENABLE_DEPRECATION_WARNINGS
+
 	DocumentDelegates->OnDocumentMetadataChanged.Broadcast();
 }
 
@@ -5980,7 +6048,8 @@ bool FMetaSoundFrontendDocumentBuilder::SetNodeInputDefault(const FGuid& InNodeI
 	using namespace Metasound::Frontend;
 
 	const FGuid& PageID = InPageID ? *InPageID : BuildPageID;
-	FMetasoundFrontendGraph& Graph = GetDocumentChecked().RootGraph.FindGraphChecked(PageID);
+	FMetasoundFrontendDocument& Document = GetDocumentChecked();
+	FMetasoundFrontendGraph& Graph = Document.RootGraph.FindGraphChecked(PageID);
 	const IDocumentGraphNodeCache& NodeCache = DocumentCache->GetNodeCache(PageID);
 	if (const int32* NodeIndex = NodeCache.FindNodeIndex(InNodeID))
 	{
@@ -6008,6 +6077,11 @@ bool FMetaSoundFrontendDocumentBuilder::SetNodeInputDefault(const FGuid& InNodeI
 
 			FNodeModifyDelegates& NodeDelegates = DocumentDelegates->FindNodeDelegatesChecked(PageID);
 			const FOnMetaSoundFrontendDocumentMutateNodeInputLiteralArray& OnNodeInputLiteralSet = NodeDelegates.OnNodeInputLiteralSet;
+
+#if WITH_EDITORONLY_DATA
+			Document.Metadata.ModifyContext.AddNodeIDModified(Node.GetID());
+#endif // WITH_EDITORONLY_DATA
+
 			OnNodeInputLiteralSet.Broadcast(*NodeIndex, VertexIndex, LiteralIndex);
 			return true;
 		}
diff --git a/Engine/Plugins/Runtime/Metasound/Source/MetasoundFrontend/Private/MetasoundFrontendDocumentVersioning.cpp b/Engine/Plugins/Runtime/Metasound/Source/MetasoundFrontend/Private/MetasoundFrontendDocumentVersioning.cpp
--- a/Engine/Plugins/Runtime/Metasound/Source/MetasoundFrontend/Private/MetasoundFrontendDocumentVersioning.cpp
+++ b/Engine/Plugins/Runtime/Metasound/Source/MetasoundFrontend/Private/MetasoundFrontendDocumentVersioning.cpp
@@ -828,6 +828,31 @@ PRAGMA_ENABLE_DEPRECATION_WARNINGS
 			}
 		};
 
+		/** Versions document from 1.15 to 1.16. */
+		class FVersionDocument_1_16 : public FVersionDocumentTransform
+		{
+		public:
+			virtual ~FVersionDocument_1_16() = default;
+
+			FMetasoundFrontendVersionNumber GetTargetVersion() const override
+			{
+				return { 1, 16 };
+			}
+
+			void TransformInternal(FMetaSoundFrontendDocumentBuilder& OutBuilder) const override
+			{
+				const FMetasoundFrontendDocument& Document = OutBuilder.GetConstDocumentChecked();
+
+PRAGMA_DISABLE_DEPRECATION_WARNINGS
+				if (Document.RootGraph.Metadata.GetIsDeprecated())
+				{
+					OutBuilder.AddAccessFlags(EMetasoundFrontendClassAccessFlags::Deprecated);
+					OutBuilder.RemoveAccessFlags(EMetasoundFrontendClassAccessFlags::Referenceable);
+				}
+PRAGMA_ENABLE_DEPRECATION_WARNINGS
+			}
+		};
+
 		bool VersionBuilderDocument(FMetaSoundFrontendDocumentBuilder& Builder)
 		{
 			UObject& DocObject = Builder.CastDocumentObjectChecked<UObject>();
@@ -839,6 +864,7 @@ PRAGMA_ENABLE_DEPRECATION_WARNINGS
 			bWasUpdated |= FVersionDocument_1_13().Transform(Builder);
 			bWasUpdated |= FVersionDocument_1_14().Transform(Builder);
 			bWasUpdated |= FVersionDocument_1_15().Transform(Builder);
+			bWasUpdated |= FVersionDocument_1_16().Transform(Builder);
 
 			return bWasUpdated;
 		}
diff --git a/Engine/Plugins/Runtime/Metasound/Source/MetasoundFrontend/Private/MetasoundFrontendDocumentVersioning.h b/Engine/Plugins/Runtime/Metasound/Source/MetasoundFrontend/Private/MetasoundFrontendDocumentVersioning.h
--- a/Engine/Plugins/Runtime/Metasound/Source/MetasoundFrontend/Private/MetasoundFrontendDocumentVersioning.h
+++ b/Engine/Plugins/Runtime/Metasound/Source/MetasoundFrontend/Private/MetasoundFrontendDocumentVersioning.h
@@ -16,7 +16,7 @@ namespace Metasound::Frontend
 {
 	static FMetasoundFrontendVersionNumber GetMaxDocumentVersion()
 	{
-		return FMetasoundFrontendVersionNumber { 1, 15 };
+		return FMetasoundFrontendVersionNumber { 1, 16 };
 	}
 
 	bool ChangeIDComparisonEnabledInAutoUpdate();
diff --git a/Engine/Plugins/Runtime/Metasound/Source/MetasoundFrontend/Private/MetasoundFrontendQuery.cpp b/Engine/Plugins/Runtime/Metasound/Source/MetasoundFrontend/Private/MetasoundFrontendQuery.cpp
--- a/Engine/Plugins/Runtime/Metasound/Source/MetasoundFrontend/Private/MetasoundFrontendQuery.cpp
+++ b/Engine/Plugins/Runtime/Metasound/Source/MetasoundFrontend/Private/MetasoundFrontendQuery.cpp
@@ -22,6 +22,7 @@ namespace Metasound::Frontend
 {
 	namespace ClassTagsPrivate
 	{
+		static const FLazyName AccessFlags = "AssetFlags";
 		static const FLazyName ClassVersion = "ClassVersion";
 
 #if WITH_EDITORONLY_DATA
@@ -169,6 +170,24 @@ namespace Metasound
 			InterfaceInfo = FMetaSoundClassInterfaceInfo(InAssetData, bSuccess);
 #endif // WITH_EDITORONLY_DATA
 
+			{
+				int32 AccessFlagsValue = 0;
+				const bool bFoundFlags = InAssetData.GetTagValue(ClassTagsPrivate::AccessFlags, AccessFlagsValue);
+				if (bFoundFlags)
+				{
+					AccessFlags = static_cast<EMetasoundFrontendClassAccessFlags>(AccessFlagsValue);
+				}
+				else
+				{
+					// When the referenceable flag was added, the prior document version assumed everything was referenceable.
+					// Therefore, we can reliably set this flag as that is the default that the document version process inevitably
+					// will upgrade the flag to.  Setting here enables caller to at least know it should be shown to end users or
+					// is available by default for referencing if out-of-date.
+					AccessFlags = EMetasoundFrontendClassAccessFlags::Referenceable;
+				}
+				bSuccess &= bFoundFlags;
+			}
+
 			bSuccess &= TryGetClassVersion(InAssetData, Version);
 			bIsValid = (uint8)(bSuccess);
 		}
@@ -178,6 +197,7 @@ namespace Metasound
 			using FAssetRegistryTag = UObject::FAssetRegistryTag;
 
 			OutContext.AddTag(FAssetRegistryTag(ClassTagsPrivate::ClassVersion, Version.ToString(), FAssetRegistryTag::TT_Alphabetical));
+			OutContext.AddTag(FAssetRegistryTag(ClassTagsPrivate::AccessFlags, FString::Printf(TEXT("%i"), (static_cast<uint8>(AccessFlags))), FAssetRegistryTag::TT_Hidden));
 
 #if WITH_EDITORONLY_DATA
 			if (!OutContext.IsCooking())
@@ -192,6 +212,7 @@ namespace Metasound
 			const FMetasoundFrontendDocument& Document = InDocInterface.GetConstDocument();
 			const FMetasoundFrontendClassMetadata& ClassMetadata = Document.RootGraph.Metadata;
 
+			AccessFlags = Document.RootGraph.Metadata.GetAccessFlags();
 			ClassName = Document.RootGraph.Metadata.GetClassName();
 			Version = ClassMetadata.GetVersion();
 
diff --git a/Engine/Plugins/Runtime/Metasound/Source/MetasoundFrontend/Private/MetasoundFrontendSearchEngineCore.cpp b/Engine/Plugins/Runtime/Metasound/Source/MetasoundFrontend/Private/MetasoundFrontendSearchEngineCore.cpp
--- a/Engine/Plugins/Runtime/Metasound/Source/MetasoundFrontend/Private/MetasoundFrontendSearchEngineCore.cpp
+++ b/Engine/Plugins/Runtime/Metasound/Source/MetasoundFrontend/Private/MetasoundFrontendSearchEngineCore.cpp
@@ -203,9 +203,13 @@ namespace Metasound::Frontend
 		{
 			if (InEntry.Value.IsType<FMetasoundFrontendClass>())
 			{
-				return !InEntry.Value.Get<FMetasoundFrontendClass>().Metadata.GetIsDeprecated();
+				const EMetasoundFrontendClassAccessFlags ThisAccessFlags = InEntry.Value.Get<FMetasoundFrontendClass>().Metadata.GetAccessFlags();
+				if (EnumHasAnyFlags(ThisAccessFlags, EMetasoundFrontendClassAccessFlags::Deprecated))
+				{
+					return false;
+				}
 			}
-			return false;
+			return true;
 		}
 
 		class FMapNodeRegistrationEventsToClassAndMajorVersion : public IFrontendQueryMapStep
diff --git a/Engine/Plugins/Runtime/Metasound/Source/MetasoundFrontend/Private/MetasoundFrontendSearchEngineCore.h b/Engine/Plugins/Runtime/Metasound/Source/MetasoundFrontend/Private/MetasoundFrontendSearchEngineCore.h
--- a/Engine/Plugins/Runtime/Metasound/Source/MetasoundFrontend/Private/MetasoundFrontendSearchEngineCore.h
+++ b/Engine/Plugins/Runtime/Metasound/Source/MetasoundFrontend/Private/MetasoundFrontendSearchEngineCore.h
@@ -101,6 +101,8 @@ namespace Metasound
 			class FRemoveDeprecatedClasses : public IFrontendQueryFilterStep
 			{
 			public:
+				FRemoveDeprecatedClasses() = default;
+
 				virtual bool Filter(const FFrontendQueryEntry& InEntry) const override;
 			};
 
@@ -174,13 +176,13 @@ namespace Metasound
 		//			return MetaSoundClass;
 		//		}
 		// };
-		template<typename QueryPolicy>
+		template<typename QueryPolicy, typename... ArgTypes>
 		struct TSearchEngineQuery : public ISearchEngineQuery
 		{
 			using ResultType = typename QueryPolicy::ResultType;
 
-			TSearchEngineQuery()
-			: Query(QueryPolicy::CreateQuery())
+			TSearchEngineQuery(ArgTypes&&... Args)
+			: Query(QueryPolicy::CreateQuery(Forward<ArgTypes>(Args)...))
 			{
 			}
 
diff --git a/Engine/Plugins/Runtime/Metasound/Source/MetasoundFrontend/Private/MetasoundFrontendSearchEngineEditorOnly.h b/Engine/Plugins/Runtime/Metasound/Source/MetasoundFrontend/Private/MetasoundFrontendSearchEngineEditorOnly.h
--- a/Engine/Plugins/Runtime/Metasound/Source/MetasoundFrontend/Private/MetasoundFrontendSearchEngineEditorOnly.h
+++ b/Engine/Plugins/Runtime/Metasound/Source/MetasoundFrontend/Private/MetasoundFrontendSearchEngineEditorOnly.h
@@ -14,7 +14,7 @@ namespace Metasound
 {
 	namespace Frontend
 	{
-		// Policy for finding all registered metasound classes. 
+		// Policy for finding all registered metasound classes that are not deprecated and referenceable. 
 		struct FFindAllClassesQueryPolicy
 		{
 			using ResultType = TArray<FMetasoundFrontendClass>;
@@ -83,7 +83,6 @@ namespace Metasound
 			virtual TArray<FMetasoundFrontendInterface> FindAllInterfaces(bool bInIncludeAllVersions) override;
 
 		private:
-
 			TSearchEngineQuery<FFindAllClassesQueryPolicy> FindAllClassesQuery;
 			TSearchEngineQuery<FFindAllClassesIncludingAllVersionsQueryPolicy> FindAllClassesIncludingAllVersionsQuery;
 			TSearchEngineQuery<FFindClassesWithNameUnsortedQueryPolicy> FindClassesWithNameUnsortedQuery;
diff --git a/Engine/Plugins/Runtime/Metasound/Source/MetasoundFrontend/Public/MetasoundAssetManager.h b/Engine/Plugins/Runtime/Metasound/Source/MetasoundFrontend/Public/MetasoundAssetManager.h
--- a/Engine/Plugins/Runtime/Metasound/Source/MetasoundFrontend/Public/MetasoundAssetManager.h
+++ b/Engine/Plugins/Runtime/Metasound/Source/MetasoundFrontend/Public/MetasoundAssetManager.h
@@ -196,12 +196,6 @@ namespace Metasound::Frontend
 		static UE_API void Deinitialize();
 		static UE_API void Initialize(TUniquePtr<IMetaSoundAssetManager>&& InInterface);
 
-		UE_DEPRECATED(5.5, "Use Initialize/Deinitialize instead")
-		UE_API virtual bool IsTesting() const;
-
-		UE_DEPRECATED(5.5, "Use Initialize/Deinitialize instead")
-		static void Set(IMetaSoundAssetManager& InInterface) { checkNoEntry(); }
-
 		// Passed template function pointer to execute in certain contexts when a MetaSound object is either already
 		// loaded or has completed asynchronous load.  Due to some implementation restrictions when constructing
 		// FLoadPackageAsyncDelegate, this has to be a copyable function (vs. TUniqueFunction).
@@ -365,6 +359,11 @@ PRAGMA_ENABLE_DEPRECATION_WARNINGS
 		virtual void RenameAsset(const FAssetData& InAssetData, const FString& InOldObjectPath) = 0;
 
 #if WITH_EDITORONLY_DATA
+		// Sets all access flags to the provided flags for the given assets. Optionally, recurses the given paths.
+		// Populates the provided array with MetaSound packages whose document versions were out-of-date and consequently updated.
+		// Returns resulting assets that were updated and those not.
+		virtual FVersionAssetResults SetAccessFlagsOnAssetsInFolders(const TArray<FString>& FolderPaths, EMetasoundFrontendClassAccessFlags Flags, bool bRecursePaths) const = 0;
+
 		// Versions all MetaSound asset tags & documents found within the given folder paths. Optionally, recurses the given paths.
 		// Populates the provided array with MetaSound packages whose document versions were out-of-date and consequently updated.
 		// Returns resulting versioned and unversioned paths.
diff --git a/Engine/Plugins/Runtime/Metasound/Source/MetasoundFrontend/Public/MetasoundDocumentInterface.h b/Engine/Plugins/Runtime/Metasound/Source/MetasoundFrontend/Public/MetasoundDocumentInterface.h
--- a/Engine/Plugins/Runtime/Metasound/Source/MetasoundFrontend/Public/MetasoundDocumentInterface.h
+++ b/Engine/Plugins/Runtime/Metasound/Source/MetasoundFrontend/Public/MetasoundDocumentInterface.h
@@ -35,14 +35,6 @@ class IMetaSoundDocumentInterface : public IInterface
 public:
 	virtual FTopLevelAssetPath GetAssetPathChecked() const = 0;
 
-	// Returns read-only reference to the the MetaSoundFrontendDocument
-	// containing all MetaSound runtime & editor data.
-	UE_DEPRECATED(5.4, "Use GetConstDocument instead")
-	virtual const FMetasoundFrontendDocument& GetDocument() const 
-	{
-		return GetConstDocument();
-	}
-
 	// Returns read-only reference to the the MetaSoundFrontendDocument
 	// containing all MetaSound runtime & editor data.
 	virtual const FMetasoundFrontendDocument& GetConstDocument() const = 0;
@@ -53,6 +45,9 @@ public:
 	// Returns the builder class used to modify the given document.
 	virtual const UClass& GetBuilderUClass() const = 0;
 
+	// Returns the default access flags utilized when document is initialized.
+	virtual EMetasoundFrontendClassAccessFlags GetDefaultAccessFlags() const = 0;
+
 	// Conforms UProperty data outside the Frontend Document Model to the document's data.
 	// Returns whether or not object data was modified.
 	virtual bool ConformObjectToDocument() = 0;
diff --git a/Engine/Plugins/Runtime/Metasound/Source/MetasoundFrontend/Public/MetasoundFrontendDocument.h b/Engine/Plugins/Runtime/Metasound/Source/MetasoundFrontend/Public/MetasoundFrontendDocument.h
--- a/Engine/Plugins/Runtime/Metasound/Source/MetasoundFrontend/Public/MetasoundFrontendDocument.h
+++ b/Engine/Plugins/Runtime/Metasound/Source/MetasoundFrontend/Public/MetasoundFrontendDocument.h
@@ -11,6 +11,7 @@
 #include "MetasoundOperatorData.h"
 #include "MetasoundNodeInterface.h"
 #include "MetasoundVertex.h"
+#include "Misc/EnumRange.h"
 #include "Misc/Guid.h"
 #include "StructUtils/InstancedStruct.h"
 #include "StructUtils/SharedStruct.h"
@@ -137,7 +138,36 @@ enum class EMetasoundFrontendVertexAccessType
 				//< or ambiguous accessor level (it is resolved during document pre-processing).
 };
 
-UENUM()
+UENUM(meta = (Bitflags, UseEnumValuesAsMaskValuesInEditor = "true"))
+enum class EMetasoundFrontendClassAccessFlags : uint16
+{
+	None = 0 UMETA(Hidden),
+
+	// Class is marked as deprecated when referenced by
+	// MetaSounds in the editor.
+	Deprecated = 1 << 0,
+
+	// If set, MetaSound can be referenced by other MetaSounds in either
+	// editor or by builder API.
+	Referenceable = 1 << 1,
+
+	Default = Referenceable UMETA(Hidden)
+};
+
+static_assert(static_cast<uint16>(EMetasoundFrontendClassAccessFlags::None) == static_cast<uint16>(Metasound::ENodeClassAccessFlags::None), "Enum values must match");
+static_assert(static_cast<uint16>(EMetasoundFrontendClassAccessFlags::Deprecated) == static_cast<uint16>(Metasound::ENodeClassAccessFlags::Deprecated), "Enum values must match");
+static_assert(static_cast<uint16>(EMetasoundFrontendClassAccessFlags::Referenceable) == static_cast<uint16>(Metasound::ENodeClassAccessFlags::Referenceable), "Enum values must match");
+static_assert(static_cast<uint16>(EMetasoundFrontendClassAccessFlags::Default) == static_cast<uint16>(Metasound::ENodeClassAccessFlags::Default), "Enum values must match");
+
+UE_API FString LexToString(const EMetasoundFrontendClassAccessFlags& InFlag);
+
+ENUM_CLASS_FLAGS(EMetasoundFrontendClassAccessFlags)
+
+ENUM_RANGE_BY_VALUES(EMetasoundFrontendClassAccessFlags,
+	EMetasoundFrontendClassAccessFlags::Deprecated,
+	EMetasoundFrontendClassAccessFlags::Referenceable)
+
+UENUM(MinimalAPI)
 enum class EMetasoundFrontendClassType : uint8
 {
 	// The MetaSound class is defined externally, in compiled code or in another document.
@@ -1517,11 +1547,13 @@ private:
 	UPROPERTY(Transient)
 	TArray<FText> CategoryHierarchyTransient;
 
-#endif // WITH_EDITORONLY_DATA
-
 	// If true, this node is deprecated and should not be used in new MetaSounds.
-	UPROPERTY(EditAnywhere, Category = Metasound)
+	UPROPERTY(meta = (Deprecated = 5.7, DeprecationMessage = "Use Access Flags instead"))
 	bool bIsDeprecated = false;
+#endif // WITH_EDITORONLY_DATA
+
+	UPROPERTY(EditAnywhere, Category = Metasound, meta = (Bitmask, BitmaskEnum = "/Script/MetasoundFrontend.EMetasoundFrontendClassAccessFlags"))
+	uint16 AccessFlags = static_cast<uint8>(EMetasoundFrontendClassAccessFlags::Referenceable);
 
 #if WITH_EDITORONLY_DATA
 	// If true, auto-update will manage (add and remove)
@@ -1546,6 +1578,11 @@ private:
 
 public:
 #if WITH_EDITOR
+	static FName GetAccessFlagsPropertyName()
+	{
+		return GET_MEMBER_NAME_CHECKED(FMetasoundFrontendClassMetadata, AccessFlags);
+	}
+
 	static FName GetAuthorPropertyName()
 	{
 		return GET_MEMBER_NAME_CHECKED(FMetasoundFrontendClassMetadata, Author);
@@ -1566,9 +1603,12 @@ public:
 		return GET_MEMBER_NAME_CHECKED(FMetasoundFrontendClassMetadata, Description);
 	}
 
+	UE_DEPRECATED(5.7, "Deprecated in favor of access flags")
 	static FName GetIsDeprecatedPropertyName()
 	{
+PRAGMA_DISABLE_DEPRECATION_WARNINGS
 		return GET_MEMBER_NAME_CHECKED(FMetasoundFrontendClassMetadata, bIsDeprecated);
+PRAGMA_ENABLE_DEPRECATION_WARNINGS
 	}
 
 	static FName GetKeywordsPropertyName()
@@ -1594,15 +1634,13 @@ public:
 
 	UE_API void SetClassName(const FMetasoundFrontendClassName& InClassName);
 
-	EMetasoundFrontendClassType GetType() const
-	{
-		return Type;
-	}
-
-	const FMetasoundFrontendVersionNumber& GetVersion() const
-	{
-		return Version;
-	}
+	UE_API EMetasoundFrontendClassType GetType() const;
+	UE_API const FMetasoundFrontendVersionNumber& GetVersion() const;
+	UE_API void AddAccessFlags(EMetasoundFrontendClassAccessFlags InAccessFlags);
+	UE_API void ClearAccessFlags();
+	UE_API EMetasoundFrontendClassAccessFlags GetAccessFlags() const;
+	UE_API void RemoveAccessFlags(EMetasoundFrontendClassAccessFlags InAccessFlags);
+	UE_API void SetAccessFlags(EMetasoundFrontendClassAccessFlags InAccessFlags);
 
 #if WITH_EDITOR
 	const FText& GetDisplayName() const
@@ -1639,10 +1677,15 @@ public:
 	UE_API void SetCategoryHierarchy(const TArray<FText>& InCategoryHierarchy);
 	UE_API void SetDescription(const FText& InDescription);
 	UE_API void SetDisplayName(const FText& InDisplayName);
+
+	UE_DEPRECATED(5.7, "Use Add/Remove/Reset/Set/ClearAccessFlags instead")
 	UE_API void SetIsDeprecated(bool bInIsDeprecated);
+
 	UE_API void SetKeywords(const TArray<FText>& InKeywords);
 	UE_API void SetPromptIfMissing(const FText& InPromptIfMissing);
 
+
+
 	UE_API void SetSerializeText(bool bInSerializeText);
 #endif // WITH_EDITOR
 
@@ -1653,10 +1696,13 @@ public:
 		return ChangeID;
 	}
 
+#if WITH_EDITORONLY_DATA
+	UE_DEPRECATED(5.7, "Use GetAccessFlags instead")
 	bool GetIsDeprecated() const
 	{
 		return bIsDeprecated;
 	}
+#endif // WITH_EDITORONLY_DATA
 
 	void SetType(const EMetasoundFrontendClassType InType)
 	{
diff --git a/Engine/Plugins/Runtime/Metasound/Source/MetasoundFrontend/Public/MetasoundFrontendDocumentBuilder.h b/Engine/Plugins/Runtime/Metasound/Source/MetasoundFrontend/Public/MetasoundFrontendDocumentBuilder.h
--- a/Engine/Plugins/Runtime/Metasound/Source/MetasoundFrontend/Public/MetasoundFrontendDocumentBuilder.h
+++ b/Engine/Plugins/Runtime/Metasound/Source/MetasoundFrontend/Public/MetasoundFrontendDocumentBuilder.h
@@ -116,6 +116,9 @@ public:
 	// Returns the document
 	UE_API virtual const FMetasoundFrontendDocument& GetConstDocument() const override;
 
+	// Returns the default access flags utilized when document is initialized.
+	UE_API virtual EMetasoundFrontendClassAccessFlags GetDefaultAccessFlags() const final override;
+
 	// Returns temp path of builder document
 	UE_API virtual FTopLevelAssetPath GetAssetPathChecked() const override;
 
@@ -161,6 +164,12 @@ public:
 	// Call when the builder will no longer modify the IMetaSoundDocumentInterface
 	UE_API void FinishBuilding();
 
+	UE_API void AddAccessFlags(EMetasoundFrontendClassAccessFlags AccessFlags);
+	UE_API void ClearAccessFlags();
+	UE_API EMetasoundFrontendClassAccessFlags GetAccessFlags() const;
+	UE_API void RemoveAccessFlags(EMetasoundFrontendClassAccessFlags AccessFlags);
+	UE_API void SetAccessFlags(EMetasoundFrontendClassAccessFlags AccessFlags);
+
 	// Adds new dependency to MetaSound. Typically not necessary to call directly as dependencies are added automatically via registry when nodes are added, and can be removed when no longer referenced (see 'RemoveUnusedDependencies`).
 	UE_API const FMetasoundFrontendClass* AddDependency(FMetasoundFrontendClass NewDependency);
 
@@ -608,6 +617,8 @@ public:
 	UE_API void SetDescription(const FText& InDescription);
 	UE_API void SetKeywords(const TArray<FText>& InKeywords);
 	UE_API void SetCategoryHierarchy(const TArray<FText>& InCategoryHierarchy);
+
+	UE_DEPRECATED(5.7, "Use AccessFlags API calls with deprecation flag instead")
 	UE_API void SetIsDeprecated(const bool bInIsDeprecated);
 
 	UE_API void SetMemberMetadata(UMetaSoundFrontendMemberMetadata& NewMetadata);
diff --git a/Engine/Plugins/Runtime/Metasound/Source/MetasoundFrontend/Public/MetasoundFrontendQuery.h b/Engine/Plugins/Runtime/Metasound/Source/MetasoundFrontend/Public/MetasoundFrontendQuery.h
--- a/Engine/Plugins/Runtime/Metasound/Source/MetasoundFrontend/Public/MetasoundFrontendQuery.h
+++ b/Engine/Plugins/Runtime/Metasound/Source/MetasoundFrontend/Public/MetasoundFrontendQuery.h
@@ -190,6 +190,8 @@ namespace Metasound
 			UE_API bool InheritsInterface(FName InterfaceName, const  FMetasoundFrontendVersionNumber& VersionNumber = { 0, 0 }) const;
 #endif // WITH_EDITORONLY_DATA
 
+			EMetasoundFrontendClassAccessFlags AccessFlags;
+
 			// If true, class info is valid and accurately reflects
 			// that from what was provided on construction. If false,
 			// info was constructed from AssetData that failed to
diff --git a/Engine/Plugins/Runtime/Metasound/Source/MetasoundFrontend/Public/MetasoundReceiveNode.h b/Engine/Plugins/Runtime/Metasound/Source/MetasoundFrontend/Public/MetasoundReceiveNode.h
--- a/Engine/Plugins/Runtime/Metasound/Source/MetasoundFrontend/Public/MetasoundReceiveNode.h
+++ b/Engine/Plugins/Runtime/Metasound/Source/MetasoundFrontend/Public/MetasoundReceiveNode.h
@@ -85,10 +85,10 @@ namespace Metasound
 					Info.CategoryHierarchy = { METASOUND_LOCTEXT("Metasound_TransmissionNodeCategory", "Transmission") };
 					Info.Keywords = { };
 
-					// Then send & receive nodes do not work as expected, particularly 
+					// Then send & receive nodes do not work as expected, particularly
 					// around multiple-consumer scenarios. Deprecate them to avoid
-					// metasound assets from relying on send & receive nodes. 
-					Info.bDeprecated = true; 
+					// metasound assets from relying on send & receive nodes.
+					EnumAddFlags(Info.AccessFlags, ENodeClassAccessFlags::Deprecated);
 
 					return Info;
 				};
diff --git a/Engine/Plugins/Runtime/Metasound/Source/MetasoundFrontend/Public/MetasoundSendNode.h b/Engine/Plugins/Runtime/Metasound/Source/MetasoundFrontend/Public/MetasoundSendNode.h
--- a/Engine/Plugins/Runtime/Metasound/Source/MetasoundFrontend/Public/MetasoundSendNode.h
+++ b/Engine/Plugins/Runtime/Metasound/Source/MetasoundFrontend/Public/MetasoundSendNode.h
@@ -125,7 +125,7 @@ namespace Metasound
 					// Then send & receive nodes do not work as expected, particularly 
 					// around multiple-consumer scenarios. Deprecate them to avoid
 					// metasound assets from relying on send & receive nodes. 
-					Info.bDeprecated = true;
+					EnumAddFlags(Info.AccessFlags, ENodeClassAccessFlags::Deprecated);
 
 					return Info;
 				};
diff --git a/Engine/Plugins/Runtime/Metasound/Source/MetasoundGraphCore/Private/MetasoundNodeInterface.cpp b/Engine/Plugins/Runtime/Metasound/Source/MetasoundGraphCore/Private/MetasoundNodeInterface.cpp
--- a/Engine/Plugins/Runtime/Metasound/Source/MetasoundGraphCore/Private/MetasoundNodeInterface.cpp
+++ b/Engine/Plugins/Runtime/Metasound/Source/MetasoundGraphCore/Private/MetasoundNodeInterface.cpp
@@ -111,6 +111,128 @@ namespace Metasound
 		return *this != InvalidNodeClassName;
 	}
 
+	FNodeClassMetadata::FNodeClassMetadata()
+		: ClassName()
+		, MajorVersion(-1)
+		, MinorVersion(-1)
+		, DisplayName()
+		, Description()
+		, Author()
+		, PromptIfMissing()
+		, DefaultInterface()
+		, CategoryHierarchy()
+		, Keywords()
+		, DisplayStyle()
+		, AccessFlags(ENodeClassAccessFlags::Default)
+#if WITH_EDITORONLY_DATA
+		PRAGMA_DISABLE_DEPRECATION_WARNINGS
+		, bDeprecated(false)
+		PRAGMA_ENABLE_DEPRECATION_WARNINGS
+#endif // WITH_EDITORONLY_DATA
+	{
+	}
+
+	FNodeClassMetadata::FNodeClassMetadata(const FNodeClassMetadata& InOther)
+		: ClassName(InOther.ClassName)
+		, MajorVersion(InOther.MajorVersion)
+		, MinorVersion(InOther.MinorVersion)
+		, DisplayName(InOther.DisplayName)
+		, Description(InOther.Description)
+		, Author(InOther.Author)
+		, PromptIfMissing(InOther.PromptIfMissing)
+		, DefaultInterface(InOther.DefaultInterface)
+		, CategoryHierarchy(InOther.CategoryHierarchy)
+		, Keywords(InOther.Keywords)
+		, DisplayStyle(InOther.DisplayStyle)
+		, AccessFlags(InOther.AccessFlags)
+#if WITH_EDITORONLY_DATA
+		PRAGMA_DISABLE_DEPRECATION_WARNINGS
+		, bDeprecated(InOther.bDeprecated)
+		PRAGMA_ENABLE_DEPRECATION_WARNINGS
+#endif // WITH_EDITORONLY_DATA
+	{
+	}
+
+	FNodeClassMetadata::FNodeClassMetadata(
+		FNodeClassName InClassName,
+		int32 InMajorVersion,
+		int32 InMinorVersion,
+		FText InDisplayName,
+		FText InDescription,
+		FString InAuthor,
+		FText InPromptIfMissing,
+		FVertexInterface InDefaultInterface,
+		TArray<FText> InCategoryHierarchy,
+		TArray<FText> InKeywords,
+		FNodeDisplayStyle InDisplayStyle)
+		: ClassName(MoveTemp(InClassName))
+		, MajorVersion(InMajorVersion)
+		, MinorVersion(InMinorVersion)
+		, DisplayName(MoveTemp(InDisplayName))
+		, Description(MoveTemp(InDescription))
+		, Author(MoveTemp(InAuthor))
+		, PromptIfMissing(MoveTemp(InPromptIfMissing))
+		, DefaultInterface(MoveTemp(InDefaultInterface))
+		, CategoryHierarchy(MoveTemp(InCategoryHierarchy))
+		, Keywords(MoveTemp(InKeywords))
+		, DisplayStyle(MoveTemp(InDisplayStyle))
+		, AccessFlags(ENodeClassAccessFlags::Default)
+#if WITH_EDITORONLY_DATA
+		PRAGMA_DISABLE_DEPRECATION_WARNINGS
+		, bDeprecated(false)
+		PRAGMA_ENABLE_DEPRECATION_WARNINGS
+#endif // WITH_EDITORONLY_DATA
+	{
+	}
+
+	FNodeClassMetadata& FNodeClassMetadata::operator=(const FNodeClassMetadata& InOther)
+	{
+		ClassName = InOther.ClassName;
+		MajorVersion = InOther.MajorVersion;
+		MinorVersion = InOther.MinorVersion;
+		DisplayName = InOther.DisplayName;
+		Description = InOther.Description;
+		Author = InOther.Author;
+		PromptIfMissing = InOther.PromptIfMissing;
+		DefaultInterface = InOther.DefaultInterface;
+		CategoryHierarchy = InOther.CategoryHierarchy;
+		Keywords = InOther.Keywords;
+		DisplayStyle = InOther.DisplayStyle;
+		AccessFlags = InOther.AccessFlags;
+
+#if WITH_EDITORONLY_DATA
+PRAGMA_DISABLE_DEPRECATION_WARNINGS
+		bDeprecated = false;
+PRAGMA_ENABLE_DEPRECATION_WARNINGS
+#endif // WITH_EDITORONLY_DATA
+
+		return *this;
+	}
+
+	FNodeClassMetadata& FNodeClassMetadata::operator=(FNodeClassMetadata&& InOther)
+	{
+		ClassName = MoveTemp(InOther.ClassName);
+		MajorVersion = MoveTemp(InOther.MajorVersion);
+		MinorVersion = MoveTemp(InOther.MinorVersion);
+		DisplayName = MoveTemp(InOther.DisplayName);
+		Description = MoveTemp(InOther.Description);
+		Author = MoveTemp(InOther.Author);
+		PromptIfMissing = MoveTemp(InOther.PromptIfMissing);
+		DefaultInterface = MoveTemp(InOther.DefaultInterface);
+		CategoryHierarchy = MoveTemp(InOther.CategoryHierarchy);
+		Keywords = MoveTemp(InOther.Keywords);
+		DisplayStyle = MoveTemp(InOther.DisplayStyle);
+		AccessFlags = MoveTemp(InOther.AccessFlags);
+
+#if WITH_EDITORONLY_DATA
+PRAGMA_DISABLE_DEPRECATION_WARNINGS
+		bDeprecated = false;
+PRAGMA_ENABLE_DEPRECATION_WARNINGS
+#endif // WITH_EDITORONLY_DATA
+
+		return *this;
+	}
+
 	FNodeClassMetadata::~FNodeClassMetadata() = default;
 
 	const FNodeClassMetadata& FNodeClassMetadata::GetEmpty()
diff --git a/Engine/Plugins/Runtime/Metasound/Source/MetasoundGraphCore/Public/MetasoundNodeInterface.h b/Engine/Plugins/Runtime/Metasound/Source/MetasoundGraphCore/Public/MetasoundNodeInterface.h
--- a/Engine/Plugins/Runtime/Metasound/Source/MetasoundGraphCore/Public/MetasoundNodeInterface.h
+++ b/Engine/Plugins/Runtime/Metasound/Source/MetasoundGraphCore/Public/MetasoundNodeInterface.h
@@ -124,17 +124,53 @@ namespace Metasound
 		FName Variant;
 	};
 
+	enum class ENodeClassAccessFlags : uint16
+	{
+		None = 0,
+
+		/** Node class is deprecated and should not be used in new MetaSounds. */
+		Deprecated = 1 << 0,
+
+		/** Node class can be referenced by MetaSound graphs */
+		Referenceable = 1 << 1,
+
+		Default = Referenceable
+	};
+
 	/** Provides metadata for a given node. */
 	struct FNodeClassMetadata
 	{
+		UE_API FNodeClassMetadata();
+		UE_API FNodeClassMetadata(const FNodeClassMetadata& InOther);
+		UE_API FNodeClassMetadata(
+			FNodeClassName InClassName,
+			int32 InMajorVersion = -1,
+			int32 InMinorVersion = -1,
+			FText InDisplayName = { },
+			FText InDescription = { },
+			FString InAuthor = { },
+			FText InPromptIfMissing = { },
+			FVertexInterface InDefaultInterface = { },
+			TArray<FText> InCategoryHierarchy = { },
+			TArray<FText> InKeywords = { },
+			FNodeDisplayStyle InDisplayStyle = { });
+
+		UE_API ~FNodeClassMetadata();
+
+		UE_API FNodeClassMetadata& operator=(const FNodeClassMetadata& InOther);
+		UE_API FNodeClassMetadata& operator=(FNodeClassMetadata&& InOther);
+
+		/** Returns an empty FNodeClassMetadata object. */
+		static UE_API const FNodeClassMetadata& GetEmpty();
+
 		/** Name of class. Used for registration and lookup. */
 		FNodeClassName ClassName;
 
 		/** Major version of node. Used for registration and lookup. */
-		int32 MajorVersion = -1;
+		int32 MajorVersion;
 
 		/** Minor version of node. */
-		int32 MinorVersion = -1;
+		int32 MinorVersion;
 
 		/** Display name of node class. */
 		FText DisplayName;
@@ -160,14 +196,12 @@ namespace Metasound
 		/** Display style for node when visualized. */
 		FNodeDisplayStyle DisplayStyle;
 
-		/** If true, the node is deprecated and should not be used in new MetaSounds. */
-		bool bDeprecated = false;
+		ENodeClassAccessFlags AccessFlags;
 
-		/** Destructor. */
-		METASOUNDGRAPHCORE_API ~FNodeClassMetadata();
-
-		/** Returns an empty FNodeClassMetadata object. */
-		static METASOUNDGRAPHCORE_API const FNodeClassMetadata& GetEmpty();
+#if WITH_EDITORONLY_DATA
+		UE_DEPRECATED(5.7, "Use AccessFlags instead")
+		bool bDeprecated;
+#endif // WITH_EDITORONLY_DATA
 	};
 
 	/** INodeBase
diff --git a/Engine/Plugins/Runtime/Metasound/Source/MetasoundStandardNodes/Private/MetasoundAudioFormats.cpp b/Engine/Plugins/Runtime/Metasound/Source/MetasoundStandardNodes/Private/MetasoundAudioFormats.cpp
--- a/Engine/Plugins/Runtime/Metasound/Source/MetasoundStandardNodes/Private/MetasoundAudioFormats.cpp
+++ b/Engine/Plugins/Runtime/Metasound/Source/MetasoundStandardNodes/Private/MetasoundAudioFormats.cpp
@@ -196,7 +196,7 @@ namespace Metasound
 			Info.Author = PluginAuthor;
 			Info.PromptIfMissing = PluginNodeMissingPrompt;
 			Info.DefaultInterface = CreateVertexInterface(InOutputName);
-			Info.bDeprecated = true;
+			Info.AccessFlags = ENodeClassAccessFlags::Deprecated;
 
 			return Info;
 		};
@@ -331,7 +331,7 @@ namespace Metasound
 			Info.Author = PluginAuthor;
 			Info.PromptIfMissing = PluginNodeMissingPrompt;
 			Info.DefaultInterface = CreateVertexInterface(InInputName);
-			Info.bDeprecated = true;
+			Info.AccessFlags = ENodeClassAccessFlags::Deprecated;
 
 			return Info;
 		};
@@ -436,7 +436,7 @@ namespace Metasound
 			Info.Author = PluginAuthor;
 			Info.PromptIfMissing = PluginNodeMissingPrompt;
 			Info.DefaultInterface = CreateVertexInterface(InOutputName);
-			Info.bDeprecated = true;
+			Info.AccessFlags = ENodeClassAccessFlags::Deprecated;
 
 			return Info;
 		};
@@ -570,7 +570,7 @@ namespace Metasound
 			Info.Author = PluginAuthor;
 			Info.PromptIfMissing = PluginNodeMissingPrompt;
 			Info.DefaultInterface = CreateVertexInterface(InInputName);
-			Info.bDeprecated = true;
+			Info.AccessFlags = ENodeClassAccessFlags::Deprecated;
 
 			return Info;
 		};
