Fix static analysis
#rnx

diff --git a/Engine/Plugins/Runtime/Metasound/Source/MetasoundEditor/Private/MetasoundEditorGraphSchema.cpp b/Engine/Plugins/Runtime/Metasound/Source/MetasoundEditor/Private/MetasoundEditorGraphSchema.cpp
--- a/Engine/Plugins/Runtime/Metasound/Source/MetasoundEditor/Private/MetasoundEditorGraphSchema.cpp
+++ b/Engine/Plugins/Runtime/Metasound/Source/MetasoundEditor/Private/MetasoundEditorGraphSchema.cpp
@@ -2409,19 +2409,18 @@ void UMetasoundEditorGraphSchema::GetAssetsGraphHoverMessage(const TArray<FAsset
 		if (UObject* DroppedObject = Data.GetAsset())
 		{
 			FMetasoundAssetBase* DroppedMetaSoundAsset = IMetasoundUObjectRegistry::Get().GetObjectAsAssetBase(DroppedObject);
-
-			const EMetasoundFrontendClassAccessFlags DroppedAccessFlags = DroppedMetaSoundAsset->GetConstDocumentChecked().RootGraph.Metadata.GetAccessFlags();
-			if (!EnumHasAnyFlags(DroppedAccessFlags, EMetasoundFrontendClassAccessFlags::Referenceable))
+			if (!DroppedMetaSoundAsset)
 			{
 				OutOkIcon = false;
-				OutTooltipText = TEXT("Asset(s) do not contain access flag 'Reference' (see 'Access Flags'). Cannot reference provided MetaSound Asset.");
+				OutTooltipText = TEXT("Asset(s) not (all) MetaSound(s) that can be referenced.");
 				break;
 			}
 
-			if (!DroppedMetaSoundAsset)
+			const EMetasoundFrontendClassAccessFlags DroppedAccessFlags = DroppedMetaSoundAsset->GetConstDocumentChecked().RootGraph.Metadata.GetAccessFlags();
+			if (!EnumHasAnyFlags(DroppedAccessFlags, EMetasoundFrontendClassAccessFlags::Referenceable))
 			{
 				OutOkIcon = false;
-				OutTooltipText = TEXT("Asset(s) not (all) MetaSound(s) that can be referenced.");
+				OutTooltipText = TEXT("Asset(s) do not contain access flag 'Reference' (see 'Access Flags'). Cannot reference provided MetaSound Asset.");
 				break;
 			}
 
