Downgrade MetaSound dependency cache log warning that is expected to occur during versioning

#rnx

diff --git a/Engine/Plugins/Runtime/Metasound/Source/MetasoundFrontend/Private/MetasoundFrontendDocumentCache.cpp b/Engine/Plugins/Runtime/Metasound/Source/MetasoundFrontend/Private/MetasoundFrontendDocumentCache.cpp
--- a/Engine/Plugins/Runtime/Metasound/Source/MetasoundFrontend/Private/MetasoundFrontendDocumentCache.cpp
+++ b/Engine/Plugins/Runtime/Metasound/Source/MetasoundFrontend/Private/MetasoundFrontendDocumentCache.cpp
@@ -1085,7 +1085,7 @@ namespace Metasound::Frontend
 #if !NO_LOGGING
 			if (IDToIndex.Contains(Class.ID))
 			{
-				UE_LOG(LogMetaSound, Warning, TEXT("MetaSound document with class name %s found containing multiple dependencies with the same ID while initializing builder cache. Document may be currently versioning or require versioning past version 1.15"), *InDocument.RootGraph.Metadata.GetClassName().ToString());
+				UE_LOG(LogMetaSound, Display, TEXT("MetaSound document with class name %s found containing multiple dependencies with the same ID while initializing builder cache. Document may be currently versioning or require versioning past version 1.15"), *InDocument.RootGraph.Metadata.GetClassName().ToString());
 			}
 #endif // !NO_LOGGING
 			IDToIndex.Add(Class.ID, Index);
