Fix for crash with copy/paste of non-MetaSound asset type introduced in earlier CL. Clarify integer serialization casting of asset tag.
#rnx

diff --git a/Engine/Plugins/Runtime/Metasound/Source/MetasoundFrontend/Private/MetasoundFrontendQuery.cpp b/Engine/Plugins/Runtime/Metasound/Source/MetasoundFrontend/Private/MetasoundFrontendQuery.cpp
--- a/Engine/Plugins/Runtime/Metasound/Source/MetasoundFrontend/Private/MetasoundFrontendQuery.cpp
+++ b/Engine/Plugins/Runtime/Metasound/Source/MetasoundFrontend/Private/MetasoundFrontendQuery.cpp
@@ -175,7 +175,7 @@ namespace Metasound
 				const bool bFoundFlags = InAssetData.GetTagValue(ClassTagsPrivate::AccessFlags, AccessFlagsValue);
 				if (bFoundFlags)
 				{
-					AccessFlags = static_cast<EMetasoundFrontendClassAccessFlags>(AccessFlagsValue);
+					AccessFlags = static_cast<EMetasoundFrontendClassAccessFlags>((uint16)(AccessFlagsValue));
 				}
 				else
 				{
@@ -197,7 +197,7 @@ namespace Metasound
 			using FAssetRegistryTag = UObject::FAssetRegistryTag;
 
 			OutContext.AddTag(FAssetRegistryTag(ClassTagsPrivate::ClassVersion, Version.ToString(), FAssetRegistryTag::TT_Alphabetical));
-			OutContext.AddTag(FAssetRegistryTag(ClassTagsPrivate::AccessFlags, FString::Printf(TEXT("%i"), (static_cast<uint8>(AccessFlags))), FAssetRegistryTag::TT_Hidden));
+			OutContext.AddTag(FAssetRegistryTag(ClassTagsPrivate::AccessFlags, FString::Printf(TEXT("%i"), (int32)(AccessFlags)), FAssetRegistryTag::TT_Hidden));
 
 #if WITH_EDITORONLY_DATA
 			if (!OutContext.IsCooking())
