Remove test log
#rnx

diff --git a/Engine/Plugins/Runtime/Metasound/Source/MetasoundEditor/Private/MetasoundEditorGraphSchema.cpp b/Engine/Plugins/Runtime/Metasound/Source/MetasoundEditor/Private/MetasoundEditorGraphSchema.cpp
--- a/Engine/Plugins/Runtime/Metasound/Source/MetasoundEditor/Private/MetasoundEditorGraphSchema.cpp
+++ b/Engine/Plugins/Runtime/Metasound/Source/MetasoundEditor/Private/MetasoundEditorGraphSchema.cpp
@@ -2915,11 +2915,6 @@ void UMetasoundEditorGraphSchema::QueryNodeClasses(TFunctionRef<void(TUniquePtr<
 			TArray<FMetasoundFrontendClass> FrontendClasses = ISearchEngine::Get().FindAllClasses(bIncludeAllVersions);
 			for (const FMetasoundFrontendClass& Class : FrontendClasses)
 			{
-				if (Class.Metadata.GetClassName().Name.ToString() == TEXT("70A917764AC0FE5934B36E9B6BF74FFA"))
-				{
-					UE_LOG(LogMetasoundEditor, Warning, TEXT("found class \"%s\""), *Class.Metadata.GetDisplayName().ToString());
-				}
-
 				FoundKeys.Add(FNodeRegistryKey(Class.Metadata));
 				Results.Add(MakeUnique<SchemaPrivate::FRegistrySchemaQueryResult>(AssetManager, Class));
 			}
