Fix crash with MetaSound user preset widget on widget BP recompile and tab reopening

diff --git a/Engine/Plugins/Runtime/Metasound/Source/MetasoundEditor/Private/MetasoundEditor.cpp b/Engine/Plugins/Runtime/Metasound/Source/MetasoundEditor/Private/MetasoundEditor.cpp
--- a/Engine/Plugins/Runtime/Metasound/Source/MetasoundEditor/Private/MetasoundEditor.cpp
+++ b/Engine/Plugins/Runtime/Metasound/Source/MetasoundEditor/Private/MetasoundEditor.cpp
@@ -1028,13 +1028,19 @@ namespace Metasound
 					return Builder.IsValid() && Builder->IsPreset();
 				});
 
-				InTabManager->RegisterTabSpawner(TabNamesPrivate::UserPresetWidget, FOnSpawnTab::CreateLambda([this, InUserPresetWidget = UserPresetWidget](const FSpawnTabArgs&)
+				InTabManager->RegisterTabSpawner(TabNamesPrivate::UserPresetWidget, FOnSpawnTab::CreateLambda([this](const FSpawnTabArgs&)
 				{
-					IMetaSoundPresetWidgetInterface::Execute_OnConstructed(InUserPresetWidget.Get(), Builder.Get());
+					// Refresh widget if widget blueprint has been recompiled
+					if (UserPresetWidget->GetClass()->HasAnyClassFlags(EClassFlags::CLASS_NewerVersionExists))
+					{
+						UserPresetWidget = TStrongObjectPtr<UUserWidget>(CreateUserPresetWidget());
+					}
+
+					IMetaSoundPresetWidgetInterface::Execute_OnConstructed(UserPresetWidget.Get(), Builder.Get());
 					return SNew(SDockTab)
 						.Label(LOCTEXT("UserPresetWidgetTitle", "Preset Widget"))
 						[
-							InUserPresetWidget.Get()->TakeWidget()
+							UserPresetWidget.Get()->TakeWidget()
 						];
 				}), CanSpawnTab)
 				.SetDisplayName(LOCTEXT("UserPresetWidgetTab", "Preset Widget"))
@@ -1114,6 +1120,14 @@ namespace Metasound
 			FLevelEditorModule& LevelEditor = FModuleManager::LoadModuleChecked<FLevelEditorModule>("LevelEditor");
 			LevelEditor.OnMapChanged().RemoveAll(this);
 
+			if (UserPresetWidget.IsValid())
+			{
+				if (UBlueprint* OwningBlueprint = UBlueprint::GetBlueprintFromClass(UserPresetWidget->GetClass()))
+				{
+					OwningBlueprint->OnCompiled().RemoveAll(this);
+				}
+			}
+
 			check(GEditor);
 			GEditor->UnregisterForUndo(this);
 		}
@@ -1863,11 +1877,45 @@ namespace Metasound
 			// Currently only support a single user widget per preset 
 			if (UserWidgets.Num() > 0)
 			{
-				return UserWidgets.Last();
+				UUserWidget* UserWidget = UserWidgets.Last();
+				check(UserWidget);
+
+				// Set recompile refresh delegate 
+				if (UBlueprint* OwningBlueprint = UBlueprint::GetBlueprintFromClass(UserWidget->GetClass()))
+				{
+					OwningBlueprint->OnCompiled().RemoveAll(this);
+					OwningBlueprint->OnCompiled().AddSP(this, &FEditor::OnPresetWidgetBlueprintCompiled);
+				}
+				return UserWidget;
 			}
 			return nullptr;
 		}
 
+		void FEditor::OnPresetWidgetBlueprintCompiled(UBlueprint* Blueprint)
+		{
+			if (!Blueprint || !UserPresetWidget)
+			{
+				return;
+			}
+			if (UBlueprint* PresetWidgetBlueprint = UBlueprint::GetBlueprintFromClass(UserPresetWidget->GetClass()))
+			{
+				if (PresetWidgetBlueprint == Blueprint)
+				{
+					// Regenerate user preset widget
+					TSharedPtr<SDockTab> UserPresetWidgetTab = TabManager->FindExistingLiveTab(TabNamesPrivate::UserPresetWidget);
+					if (UserPresetWidgetTab.IsValid())
+					{
+						UserPresetWidget = TStrongObjectPtr<UUserWidget>(CreateUserPresetWidget());
+						if (UserPresetWidget.IsValid())
+						{
+							IMetaSoundPresetWidgetInterface::Execute_OnConstructed(UserPresetWidget.Get(), Builder.Get());
+							UserPresetWidgetTab->SetContent(UserPresetWidget.Get()->TakeWidget());
+						}
+					}
+				}
+			}
+		}
+		
 		void FEditor::ChangeUserPresetWidgetWorld(UWorld* World, EMapChangeType MapChangeType)
 		{
 			if (MapChangeType == EMapChangeType::TearDownWorld)
diff --git a/Engine/Plugins/Runtime/Metasound/Source/MetasoundEditor/Private/MetasoundEditor.h b/Engine/Plugins/Runtime/Metasound/Source/MetasoundEditor/Private/MetasoundEditor.h
--- a/Engine/Plugins/Runtime/Metasound/Source/MetasoundEditor/Private/MetasoundEditor.h
+++ b/Engine/Plugins/Runtime/Metasound/Source/MetasoundEditor/Private/MetasoundEditor.h
@@ -457,6 +457,9 @@ namespace Metasound::Editor
 			Based on UEditorUtilityWidgetBlueprint::ChangeTabWorld */
 		void ChangeUserPresetWidgetWorld(UWorld* World, EMapChangeType MapChangeType);
 
+		/** Reload preset widget tab on associated widget blueprint recompile */
+		void OnPresetWidgetBlueprintCompiled(UBlueprint* Blueprint);
+
 		/** Builds the toolbar widget for the Metasound editor */
 		void ExtendToolbarInternal();
 
