Make referenceable MetaSound access flag restriction apply to BP builder, not frontend document builder
- Fix nodes with major version of 0 not being copy pasted properly

diff --git a/Engine/Plugins/Runtime/Metasound/Source/MetasoundEditor/Private/MetasoundEditorDocumentClipboardUtils.cpp b/Engine/Plugins/Runtime/Metasound/Source/MetasoundEditor/Private/MetasoundEditorDocumentClipboardUtils.cpp
--- a/Engine/Plugins/Runtime/Metasound/Source/MetasoundEditor/Private/MetasoundEditorDocumentClipboardUtils.cpp
+++ b/Engine/Plugins/Runtime/Metasound/Source/MetasoundEditor/Private/MetasoundEditorDocumentClipboardUtils.cpp
@@ -371,17 +371,38 @@ namespace Metasound::Editor
 				}
 				else
 				{
-					const FMetasoundFrontendNode* NewNode = Builder.AddNodeByClassName(Breadcrumb.ClassName);
-					if (NewNode)
+					FMetasoundFrontendClass RegisteredClass;
+					if (ISearchEngine::Get().FindClassWithHighestMinorVersion(Breadcrumb.ClassName, Breadcrumb.Version.Major, RegisteredClass))
 					{
-						const FGuid& NewNodeID = NewNode->GetID();
-						ExternalNode->NodeID = NewNodeID;
-						Builder.SetNodeConfiguration(NewNodeID, Breadcrumb.NodeConfiguration);
+						if (EnumHasAnyFlags(RegisteredClass.Metadata.GetAccessFlags(), EMetasoundFrontendClassAccessFlags::Referenceable))
+						{
+							const FMetasoundFrontendNode* NewNode = Builder.AddNodeByClassName(Breadcrumb.ClassName, Breadcrumb.Version.Major);
+							if (NewNode)
+							{
+								const FGuid& NewNodeID = NewNode->GetID();
+								ExternalNode->NodeID = NewNodeID;
+								Builder.SetNodeConfiguration(NewNodeID, Breadcrumb.NodeConfiguration);
+							}
+							else
+							{
+								OutPastedNodes.RemoveAtSwap(Index, 1, EAllowShrinking::No);
+								Graph.RemoveNode(ExternalNode);
+							}
+						}
+						else
+						{
+							OutPastedNodes.RemoveAtSwap(Index, 1, EAllowShrinking::No);
+							Graph.RemoveNode(ExternalNode);
+							UE_LOG(LogMetaSound, Warning, TEXT("Failed to add new node by name of class '%s': Class access flag '%s' not set."),
+								*Breadcrumb.ClassName.ToString(),
+								*LexToString(EMetasoundFrontendClassAccessFlags::Referenceable));
+						}
 					}
 					else
 					{
 						OutPastedNodes.RemoveAtSwap(Index, 1, EAllowShrinking::No);
 						Graph.RemoveNode(ExternalNode);
+						UE_LOG(LogMetaSound, Error, TEXT("Failed to add new node by class name '%s' and major version '%d': Class not found"), *Breadcrumb.ClassName.ToString(), Breadcrumb.Version.Major);
 					}
 				}
 			}
@@ -404,19 +425,30 @@ namespace Metasound::Editor
 				else
 				{
 					FMetasoundFrontendClass ExternalClass;
-					if (ISearchEngine::Get().FindClassWithHighestVersion(Breadcrumb.ClassName, ExternalClass))
+					if (ISearchEngine::Get().FindClassWithHighestMinorVersion(Breadcrumb.ClassName, Breadcrumb.Version.Major, ExternalClass))
 					{
-						const FMetasoundFrontendNode* NewNode = Builder.AddNodeByClassName(Breadcrumb.ClassName);
-						if (NewNode)
+						if (EnumHasAnyFlags(ExternalClass.Metadata.GetAccessFlags(), EMetasoundFrontendClassAccessFlags::Referenceable))
 						{
-							const FGuid& NewNodeID = NewNode->GetID();
-							ExternalNode->NodeID = NewNodeID;
-							Builder.SetNodeConfiguration(NewNodeID, Breadcrumb.NodeConfiguration);
+							const FMetasoundFrontendNode* NewNode = Builder.AddNodeByClassName(Breadcrumb.ClassName, Breadcrumb.Version.Major);
+							if (NewNode)
+							{
+								const FGuid& NewNodeID = NewNode->GetID();
+								ExternalNode->NodeID = NewNodeID;
+								Builder.SetNodeConfiguration(NewNodeID, Breadcrumb.NodeConfiguration);
+							}
+							else
+							{
+								OutPastedNodes.RemoveAtSwap(Index, 1, EAllowShrinking::No);
+								Graph.RemoveNode(ExternalNode);
+							}
 						}
 						else
 						{
 							OutPastedNodes.RemoveAtSwap(Index, 1, EAllowShrinking::No);
 							Graph.RemoveNode(ExternalNode);
+							UE_LOG(LogMetaSound, Warning, TEXT("Failed to add new node by name of class '%s': Class access flag '%s' not set."),
+								*Breadcrumb.ClassName.ToString(),
+								*LexToString(EMetasoundFrontendClassAccessFlags::Referenceable));
 						}
 					}
 					else
diff --git a/Engine/Plugins/Runtime/Metasound/Source/MetasoundEditor/Private/MetasoundEditorGraphBuilder.cpp b/Engine/Plugins/Runtime/Metasound/Source/MetasoundEditor/Private/MetasoundEditorGraphBuilder.cpp
--- a/Engine/Plugins/Runtime/Metasound/Source/MetasoundEditor/Private/MetasoundEditorGraphBuilder.cpp
+++ b/Engine/Plugins/Runtime/Metasound/Source/MetasoundEditor/Private/MetasoundEditorGraphBuilder.cpp
@@ -425,11 +425,17 @@ namespace Metasound
 		UMetasoundEditorGraphExternalNode* FGraphBuilder::AddExternalNode(UObject& InMetaSound, const FMetasoundFrontendClassMetadata& InMetadata, bool bInSelectNewNode)
 		{
 			using namespace Frontend;
+			using namespace Engine;
 
-			FMetaSoundFrontendDocumentBuilder& Builder = IDocumentBuilderRegistry::GetChecked().FindOrBeginBuilding(&InMetaSound);
-			if (const FMetasoundFrontendNode* NewNode = Builder.AddNodeByClassName(InMetadata.GetClassName(), InMetadata.GetVersion().Major))
+			UMetaSoundBuilderBase& Builder = FDocumentBuilderRegistry::GetChecked().FindOrBeginBuilding(InMetaSound);
+			EMetaSoundBuilderResult Result;
+			const FMetaSoundNodeHandle NewNodeHandle = Builder.AddNodeByClassName(InMetadata.GetClassName(), Result, InMetadata.GetVersion().Major);
+			if (Result == EMetaSoundBuilderResult::Succeeded && NewNodeHandle.IsSet())
 			{
-				const FMetasoundFrontendClass* Dependency = Builder.FindDependency(NewNode->ClassID);
+				const FMetaSoundFrontendDocumentBuilder& DocBuilder = Builder.GetConstBuilder();
+				const FMetasoundFrontendNode* NewNode = DocBuilder.FindNode(NewNodeHandle.NodeID);
+				check(NewNode);
+				const FMetasoundFrontendClass* Dependency = DocBuilder.FindDependency(NewNode->ClassID);
 				if (ensure(Dependency))
 				{
 					FMetasoundAssetBase* MetaSoundAsset = IMetasoundUObjectRegistry::Get().GetObjectAsAssetBase(&InMetaSound);
diff --git a/Engine/Plugins/Runtime/Metasound/Source/MetasoundEditor/Private/MetasoundEditorGraphNode.cpp b/Engine/Plugins/Runtime/Metasound/Source/MetasoundEditor/Private/MetasoundEditorGraphNode.cpp
--- a/Engine/Plugins/Runtime/Metasound/Source/MetasoundEditor/Private/MetasoundEditorGraphNode.cpp
+++ b/Engine/Plugins/Runtime/Metasound/Source/MetasoundEditor/Private/MetasoundEditorGraphNode.cpp
@@ -1171,6 +1171,7 @@ void UMetasoundEditorGraphExternalNode::CacheBreadcrumb()
 
 		Breadcrumb.bIsClassNative = !bIsAssetClass;
 		Breadcrumb.ClassName = Metadata.GetClassName();
+		Breadcrumb.Version = Metadata.GetVersion();
 
 		const FMetasoundFrontendNode& Node = GetFrontendNodeChecked();
 		Breadcrumb.NodeConfiguration = Node.Configuration;
diff --git a/Engine/Plugins/Runtime/Metasound/Source/MetasoundEditor/Public/MetasoundEditorGraphNode.h b/Engine/Plugins/Runtime/Metasound/Source/MetasoundEditor/Public/MetasoundEditorGraphNode.h
--- a/Engine/Plugins/Runtime/Metasound/Source/MetasoundEditor/Public/MetasoundEditorGraphNode.h
+++ b/Engine/Plugins/Runtime/Metasound/Source/MetasoundEditor/Public/MetasoundEditorGraphNode.h
@@ -59,6 +59,9 @@ struct FMetasoundEditorGraphNodeBreadcrumb
 	// For use with template nodes only
 	UPROPERTY()
 	TOptional<FNodeTemplateGenerateInterfaceParams> TemplateParams;
+
+	UPROPERTY()
+	FMetasoundFrontendVersionNumber Version;
 };
 
 USTRUCT()
diff --git a/Engine/Plugins/Runtime/Metasound/Source/MetasoundEngine/Private/MetasoundBuilderBase.cpp b/Engine/Plugins/Runtime/Metasound/Source/MetasoundEngine/Private/MetasoundBuilderBase.cpp
--- a/Engine/Plugins/Runtime/Metasound/Source/MetasoundEngine/Private/MetasoundBuilderBase.cpp
+++ b/Engine/Plugins/Runtime/Metasound/Source/MetasoundEngine/Private/MetasoundBuilderBase.cpp
@@ -268,6 +268,16 @@ FMetaSoundNodeHandle UMetaSoundBuilderBase::AddNode(const TScriptInterface<IMeta
 
 			const FMetasoundFrontendDocument& NodeClassDoc = NodeClass->GetConstDocument();
 			const FMetasoundFrontendGraphClass& NodeClassGraph = NodeClassDoc.RootGraph;
+
+			const EMetasoundFrontendClassAccessFlags AccessFlags = NodeClassGraph.Metadata.GetAccessFlags();
+			if (!EnumHasAnyFlags(AccessFlags, EMetasoundFrontendClassAccessFlags::Referenceable))
+			{
+				UE_LOG(LogMetaSound, Warning, TEXT("Failed to add new node of graph class '%s': Class access flag '%s' not set."),
+					*NodeClassGraph.Metadata.GetClassName().ToString(),
+					*LexToString(EMetasoundFrontendClassAccessFlags::Referenceable));
+				OutResult = EMetaSoundBuilderResult::Failed;
+				return NewHandle;
+			}
 			if (const FMetasoundFrontendNode* NewNode = Builder.AddGraphNode(NodeClassGraph))
 			{
 				NewHandle.NodeID = NewNode->GetID();
@@ -290,6 +300,23 @@ FMetaSoundNodeHandle UMetaSoundBuilderBase::AddNodeByClassName(const FMetasoundF
 	using namespace Metasound::Frontend;
 
 	FMetaSoundNodeHandle NewHandle;
+	FMetasoundFrontendClass RegisteredClass;
+	if (!ISearchEngine::Get().FindClassWithHighestMinorVersion(InClassName, MajorVersion, RegisteredClass))
+	{
+		UE_LOG(LogMetaSound, Error, TEXT("Failed to add new node by class name '%s' and major version '%d': Class not found"), *InClassName.ToString(), MajorVersion);
+		OutResult = EMetaSoundBuilderResult::Failed;
+		return NewHandle;
+	}
+
+	if (!EnumHasAnyFlags(RegisteredClass.Metadata.GetAccessFlags(), EMetasoundFrontendClassAccessFlags::Referenceable))
+	{
+		UE_LOG(LogMetaSound, Warning, TEXT("Failed to add new node by name of class '%s': Class access flag '%s' not set."),
+			*InClassName.ToString(),
+			*LexToString(EMetasoundFrontendClassAccessFlags::Referenceable));
+		OutResult = EMetaSoundBuilderResult::Failed;
+		return NewHandle;
+	}
+
 	if (const FMetasoundFrontendNode* NewNode = Builder.AddNodeByClassName(InClassName, MajorVersion))
 	{
 		NewHandle.NodeID = NewNode->GetID();
diff --git a/Engine/Plugins/Runtime/Metasound/Source/MetasoundFrontend/Private/MetasoundFrontendDocumentBuilder.cpp b/Engine/Plugins/Runtime/Metasound/Source/MetasoundFrontend/Private/MetasoundFrontendDocumentBuilder.cpp
--- a/Engine/Plugins/Runtime/Metasound/Source/MetasoundFrontend/Private/MetasoundFrontendDocumentBuilder.cpp
+++ b/Engine/Plugins/Runtime/Metasound/Source/MetasoundFrontend/Private/MetasoundFrontendDocumentBuilder.cpp
@@ -1757,15 +1757,6 @@ const FMetasoundFrontendNode* FMetaSoundFrontendDocumentBuilder::AddGraphNode(co
 {
 	using FNodeRegistryKey = Metasound::Frontend::FNodeRegistryKey;
 
-	const EMetasoundFrontendClassAccessFlags AccessFlags = InGraphClass.Metadata.GetAccessFlags();
-	if (!EnumHasAnyFlags(AccessFlags, EMetasoundFrontendClassAccessFlags::Referenceable))
-	{
-		UE_LOG(LogMetaSound, Warning, TEXT("Failed to add new node of graph class '%s': Class access flag '%s' not set."),
-			*InGraphClass.Metadata.GetClassName().ToString(),
-			*LexToString(EMetasoundFrontendClassAccessFlags::Referenceable));
-		return nullptr;
-	}
-
 	auto FinalizeNode = [](FMetasoundFrontendNode& InOutNode, const Metasound::Frontend::FNodeRegistryKey& ClassKey)
 	{
 #if WITH_EDITOR
@@ -1826,14 +1817,6 @@ const FMetasoundFrontendNode* FMetaSoundFrontendDocumentBuilder::AddNodeByClassN
 			return nullptr;
 		}
 
-		if (!EnumHasAnyFlags(RegisteredClass.Metadata.GetAccessFlags(), EMetasoundFrontendClassAccessFlags::Referenceable))
-		{
-			UE_LOG(LogMetaSound, Warning, TEXT("Failed to add new node by name of class '%s': Class access flag '%s' not set."),
-				*InClassName.ToString(),
-				*LexToString(EMetasoundFrontendClassAccessFlags::Referenceable));
-			return nullptr;
-		}
-
 		// Dependency is considered "External" when looked up or added as a dependency to a graph
 		RegisteredClass.Metadata.SetType(EMetasoundFrontendClassType::External);
 		ClassKey = FNodeRegistryKey(RegisteredClass.Metadata);
@@ -1863,14 +1846,6 @@ const FMetasoundFrontendNode* FMetaSoundFrontendDocumentBuilder::AddNodeByTempla
 	const FMetasoundFrontendClass& TemplateClass = InTemplate.GetFrontendClass();
 	checkf(TemplateClass.Metadata.GetType() == EMetasoundFrontendClassType::Template, TEXT("INodeTemplate ClassType must always be 'Template'"));
 
-	if (!EnumHasAnyFlags(TemplateClass.Metadata.GetAccessFlags(), EMetasoundFrontendClassAccessFlags::Referenceable))
-	{
-		UE_LOG(LogMetaSound, Warning, TEXT("Failed to add new template node of class '%s': Class access flag '%s' not set."),
-			*InTemplate.GetClassName().ToString(),
-			*LexToString(EMetasoundFrontendClassAccessFlags::Referenceable));
-		return nullptr;
-	}
-
 	const FMetasoundFrontendClass* Dependency = FindDependency(TemplateClass.Metadata);
 	if (!Dependency)
 	{
diff --git a/Engine/Plugins/Runtime/Metasound/Source/MetasoundFrontend/Public/MetasoundFrontendDocument.h b/Engine/Plugins/Runtime/Metasound/Source/MetasoundFrontend/Public/MetasoundFrontendDocument.h
--- a/Engine/Plugins/Runtime/Metasound/Source/MetasoundFrontend/Public/MetasoundFrontendDocument.h
+++ b/Engine/Plugins/Runtime/Metasound/Source/MetasoundFrontend/Public/MetasoundFrontendDocument.h
@@ -148,7 +148,7 @@ enum class EMetasoundFrontendClassAccessFlags : uint16
 	Deprecated = 1 << 0,
 
 	// If set, MetaSound can be referenced by other MetaSounds in either
-	// editor or by builder API.
+	// editor or by builder Blueprint API.
 	Referenceable = 1 << 1,
 
 	Default = Referenceable UMETA(Hidden)
