Add "MaintainAudioSync" input to Wave Player Node as an advanced input to to allow the Wave Player to use the same playback synchronization mechanism the Fusion Sampler node already uses.

Does not eliminate latency. But if there is latency, then once the sound starts playing, it will compencate for the latency and seek forward to the point it would be if there was no latency once it starts playing.

[RN] MetaSounds: Add "Maintain Audio Sync" flag to Wave Player Node to allow for frame accurate playback
#DisabledValidations NonUnityAndStaticAnalysis

diff --git a/Engine/Plugins/Runtime/Metasound/Source/MetasoundEngine/Private/MetasoundWavePlayerNode.cpp b/Engine/Plugins/Runtime/Metasound/Source/MetasoundEngine/Private/MetasoundWavePlayerNode.cpp
--- a/Engine/Plugins/Runtime/Metasound/Source/MetasoundEngine/Private/MetasoundWavePlayerNode.cpp
+++ b/Engine/Plugins/Runtime/Metasound/Source/MetasoundEngine/Private/MetasoundWavePlayerNode.cpp
@@ -33,6 +33,7 @@ namespace Metasound
 		METASOUND_PARAM(InputLoop, "Loop", "Whether or not to loop between the start and specified end times.")
 		METASOUND_PARAM(InputLoopStart, "Loop Start", "When to start the loop.")
 		METASOUND_PARAM(InputLoopDuration, "Loop Duration", "The duration of the loop when wave player is enabled for looping. A negative value will loop the whole wave asset.")
+		METASOUND_PARAM(InputMaintainAudioSync, "Maintain Audio Sync", "Force audio to seek to maintain sample accurate playback even if the source audio is latent")
 		
 		METASOUND_PARAM(OutputTriggerOnPlay, "On Play", "Triggers when Play is triggered.")
 		METASOUND_PARAM(OutputTriggerOnDone, "On Finished", "Triggers when the wave played has finished playing or Stop is triggered.")
@@ -308,6 +309,7 @@ namespace Metasound
 		FBoolReadRef bLoop;
 		FTimeReadRef LoopStartTime;
 		FTimeReadRef LoopDuration;
+		bool bMaintainAudioSync;
 		FNodeRenderCost CostReporter;
 	};
 
@@ -343,6 +345,9 @@ namespace Metasound
 			FDataVertexMetadata InputLoopDurationMetaData = METASOUND_GET_PARAM_METADATA(InputLoopDuration);
 			InputLoopDurationMetaData.bIsAdvancedDisplay = true;
 
+			FDataVertexMetadata InputMaintainAudioSyncMetaData = METASOUND_GET_PARAM_METADATA(InputMaintainAudioSync);
+			InputMaintainAudioSyncMetaData.bIsAdvancedDisplay = true;
+
 			FDataVertexMetadata OutputTriggerOnNearlyDoneMetaData = METASOUND_GET_PARAM_METADATA(OutputTriggerOnNearlyDone);
 			OutputTriggerOnNearlyDoneMetaData.bIsAdvancedDisplay = true;
 
@@ -375,7 +380,8 @@ namespace Metasound
 					TInputDataVertex<float>(METASOUND_GET_PARAM_NAME_AND_METADATA(InputPitchShift), 0.0f),
 					TInputDataVertex<bool>(METASOUND_GET_PARAM_NAME_AND_METADATA(InputLoop), false),
 					TInputDataVertex<FTime>(METASOUND_GET_PARAM_NAME(InputLoopStart), InputLoopStartMetaData, 0.0f),
-					TInputDataVertex<FTime>(METASOUND_GET_PARAM_NAME(InputLoopDuration), InputLoopDurationMetaData, -1.0f)
+					TInputDataVertex<FTime>(METASOUND_GET_PARAM_NAME(InputLoopDuration), InputLoopDurationMetaData, -1.0f),
+					TInputConstructorVertex<bool>(METASOUND_GET_PARAM_NAME(InputMaintainAudioSync), InputMaintainAudioSyncMetaData, false)
 				),
 				FOutputVertexInterface(
 					TOutputDataVertex<FTrigger>(METASOUND_GET_PARAM_NAME_AND_METADATA(OutputTriggerOnPlay)),
@@ -438,6 +444,7 @@ namespace Metasound
 			, bLoop(InArgs.bLoop)
 			, LoopStartTime(InArgs.LoopStartTime)
 			, LoopDuration(InArgs.LoopDuration)
+			, bMaintainAudioSync(InArgs.bMaintainAudioSync)
 			, TriggerOnDone(FTriggerWriteRef::CreateNew(InArgs.Settings))
 			, TriggerOnNearlyDone(FTriggerWriteRef::CreateNew(InArgs.Settings))
 			, TriggerOnLooped(FTriggerWriteRef::CreateNew(InArgs.Settings))
@@ -476,6 +483,7 @@ namespace Metasound
 			InOutVertexData.BindReadVertex(METASOUND_GET_PARAM_NAME(InputLoop), bLoop);
 			InOutVertexData.BindReadVertex(METASOUND_GET_PARAM_NAME(InputLoopStart), LoopStartTime);
 			InOutVertexData.BindReadVertex(METASOUND_GET_PARAM_NAME(InputLoopDuration), LoopDuration);
+			InOutVertexData.SetValue(METASOUND_GET_PARAM_NAME(InputMaintainAudioSync), bMaintainAudioSync);
 		}
 
 		virtual void BindOutputs(FOutputVertexInterfaceData& InOutVertexData) override
@@ -814,6 +822,7 @@ namespace Metasound
 				WaveReaderSettings.LoopStartTimeInSeconds = LoopStartTime->GetSeconds();
 				WaveReaderSettings.LoopDurationInSeconds = LoopDuration->GetSeconds(); 
 				WaveReaderSettings.bIsLooping = *bLoop;
+				WaveReaderSettings.bMaintainAudioSync = bMaintainAudioSync;
 
 				WaveProxyReader = FSoundWaveProxyReader::Create(WaveProxy.ToSharedRef(), WaveReaderSettings);
 
@@ -1129,6 +1138,7 @@ namespace Metasound
 		FBoolReadRef bLoop;
 		FTimeReadRef LoopStartTime;
 		FTimeReadRef LoopDuration;
+		const bool bMaintainAudioSync;
 
 		FTriggerWriteRef TriggerOnDone;
 		FTriggerWriteRef TriggerOnNearlyDone;
@@ -1192,6 +1202,7 @@ namespace Metasound
 				Inputs.GetOrCreateDefaultDataReadReference<bool>(METASOUND_GET_PARAM_NAME(InputLoop), InParams.OperatorSettings),
 				Inputs.GetOrCreateDefaultDataReadReference<FTime>(METASOUND_GET_PARAM_NAME(InputLoopStart), InParams.OperatorSettings),
 				Inputs.GetOrCreateDefaultDataReadReference<FTime>(METASOUND_GET_PARAM_NAME(InputLoopDuration), InParams.OperatorSettings),
+				Inputs.GetOrCreateDefaultValue<bool>(METASOUND_GET_PARAM_NAME(InputMaintainAudioSync), InParams.OperatorSettings),
 				FNodeRenderCost(InParams.GraphRenderCost ? InParams.GraphRenderCost->AddNode(InParams.Node.GetInstanceID(), InParams.Environment) : FNodeRenderCost{})
 			};
 
