Refactor MetaSound rebuild preset and autoupdate transforms to use builder API
- Move logic for dependency updates to builder API

diff --git a/Engine/Plugins/Runtime/Metasound/Source/MetasoundEditor/Private/MetasoundEditorGraphNode.cpp b/Engine/Plugins/Runtime/Metasound/Source/MetasoundEditor/Private/MetasoundEditorGraphNode.cpp
--- a/Engine/Plugins/Runtime/Metasound/Source/MetasoundEditor/Private/MetasoundEditorGraphNode.cpp
+++ b/Engine/Plugins/Runtime/Metasound/Source/MetasoundEditor/Private/MetasoundEditorGraphNode.cpp
@@ -1156,7 +1156,8 @@ bool UMetasoundEditorGraphExternalNode::CanAutoUpdate() const
 	using namespace Metasound::Frontend;
 
 	FClassInterfaceUpdates InterfaceUpdates;
-	return GetConstNodeHandle()->CanAutoUpdate(InterfaceUpdates);
+	FMetaSoundFrontendDocumentBuilder& DocBuilder = GetBuilderChecked().GetBuilder();
+	return DocBuilder.CanAutoUpdate(GetNodeID(), InterfaceUpdates);
 }
 
 void UMetasoundEditorGraphExternalNode::CacheBreadcrumb()
diff --git a/Engine/Plugins/Runtime/Metasound/Source/MetasoundEditor/Private/MetasoundEditorSubsystem.cpp b/Engine/Plugins/Runtime/Metasound/Source/MetasoundEditor/Private/MetasoundEditorSubsystem.cpp
--- a/Engine/Plugins/Runtime/Metasound/Source/MetasoundEditor/Private/MetasoundEditorSubsystem.cpp
+++ b/Engine/Plugins/Runtime/Metasound/Source/MetasoundEditor/Private/MetasoundEditorSubsystem.cpp
@@ -370,7 +370,8 @@ void UMetaSoundEditorSubsystem::InitAsset(UObject& InNewMetaSound, UObject* InRe
 
 		// Initialize preset with referenced Metasound 
 		TScriptInterface<IMetaSoundDocumentInterface> ReferencedDocInterface = InReferencedMetaSound;
-		Builder.ConvertToPreset(ReferencedDocInterface->GetConstDocument());
+		const FMetaSoundFrontendDocumentBuilder& ReferenceBuilder = IDocumentBuilderRegistry::GetChecked().FindOrBeginBuilding(ReferencedDocInterface);
+		Builder.ConvertToPreset(ReferenceBuilder);
 		
 		if (&DocInterface->GetBaseMetaSoundUClass() == UMetaSoundSource::StaticClass())
 		{
diff --git a/Engine/Plugins/Runtime/Metasound/Source/MetasoundEngine/Private/MetasoundBuilderBase.cpp b/Engine/Plugins/Runtime/Metasound/Source/MetasoundEngine/Private/MetasoundBuilderBase.cpp
--- a/Engine/Plugins/Runtime/Metasound/Source/MetasoundEngine/Private/MetasoundBuilderBase.cpp
+++ b/Engine/Plugins/Runtime/Metasound/Source/MetasoundEngine/Private/MetasoundBuilderBase.cpp
@@ -633,12 +633,14 @@ void UMetaSoundBuilderBase::ConvertToPreset(const TScriptInterface<IMetaSoundDoc
 		RegOptions.PageOrder = UMetaSoundSettings::GetPageOrder();
 		ReferencedMetaSoundAsset->UpdateAndRegisterForExecution(RegOptions);
 	}
-
-	const FMetasoundFrontendDocument& ReferencedDocument = ReferencedInterface->GetConstDocument();
-
-	TSharedRef<FDocumentModifyDelegates> DocumentDelegates = MakeShared<FDocumentModifyDelegates>(ReferencedDocument);
+#if WITH_EDITORONLY_DATA
+	const FMetaSoundFrontendDocumentBuilder& ReferencedBuilder = Metasound::Frontend::IDocumentBuilderRegistry::GetChecked().FindOrBeginBuilding(ReferencedObject);
+#else
+	const FMetaSoundFrontendDocumentBuilder& ReferencedBuilder = FMetaSoundFrontendDocumentBuilder(ReferencedObject);
+#endif // WITH_EDITORONLY_DATA
+	TSharedRef<FDocumentModifyDelegates> DocumentDelegates = MakeShared<FDocumentModifyDelegates>(ReferencedInterface->GetConstDocument());
 	InitDelegates(*DocumentDelegates);
-	const bool bConvertToPreset = Builder.ConvertToPreset(ReferencedDocument, DocumentDelegates);
+	const bool bConvertToPreset = Builder.ConvertToPreset(ReferencedBuilder, DocumentDelegates);
 	ConformObjectToDocument();
 	OutResult = bConvertToPreset ? EMetaSoundBuilderResult::Succeeded : EMetaSoundBuilderResult::Failed;
 }
diff --git a/Engine/Plugins/Runtime/Metasound/Source/MetasoundFrontend/Private/MetasoundAssetBase.cpp b/Engine/Plugins/Runtime/Metasound/Source/MetasoundFrontend/Private/MetasoundAssetBase.cpp
--- a/Engine/Plugins/Runtime/Metasound/Source/MetasoundFrontend/Private/MetasoundAssetBase.cpp
+++ b/Engine/Plugins/Runtime/Metasound/Source/MetasoundFrontend/Private/MetasoundAssetBase.cpp
@@ -922,20 +922,8 @@ bool FMetasoundAssetBase::VersionDependencies(FMetaSoundFrontendDocumentBuilder&
 	bool bDocumentModified = false;
 
 #if WITH_EDITORONLY_DATA
-	const FGuid InitBuildPageID = Builder.GetBuildPageID();
-	Metasound::Frontend::FDocumentHandle DocHandle = GetDocumentHandle();
-	const FMetasoundFrontendGraphClass& RootGraph = Builder.GetConstDocumentChecked().RootGraph;
-	constexpr bool bBroadcastPageIDDelegate = false;
-
 	FAutoUpdateRootGraph AutoUpdateTransform(GetOwningAssetName(), bInLogWarningsOnDroppedConnection);
-	RootGraph.IterateGraphPages([&](const FMetasoundFrontendGraph& Graph)
-	{
-		// Set the build page ID to this graph as a hack to apply dependency versioning logic using
-		// the controller/handle API until auto-update is renamed & moved to use document builder API.
-		Builder.SetBuildPageID(Graph.PageID, bBroadcastPageIDDelegate);
-		bDocumentModified |= AutoUpdateTransform.Transform(DocHandle);
-	});
-	Builder.SetBuildPageID(InitBuildPageID, bBroadcastPageIDDelegate);
+	bDocumentModified |= AutoUpdateTransform.Transform(Builder);
 #endif // WITH_EDITORONLY_DATA
 
 	return bDocumentModified;
diff --git a/Engine/Plugins/Runtime/Metasound/Source/MetasoundFrontend/Private/MetasoundFrontendDocument.cpp b/Engine/Plugins/Runtime/Metasound/Source/MetasoundFrontend/Private/MetasoundFrontendDocument.cpp
--- a/Engine/Plugins/Runtime/Metasound/Source/MetasoundFrontend/Private/MetasoundFrontendDocument.cpp
+++ b/Engine/Plugins/Runtime/Metasound/Source/MetasoundFrontend/Private/MetasoundFrontendDocument.cpp
@@ -32,6 +32,14 @@ namespace Metasound
 
 	namespace Frontend
 	{
+		int32 MetaSoundAutoUpdateNativeClassesOfEqualVersionCVar = 1;
+		FAutoConsoleVariableRef CVarMetaSoundAutoUpdateNativeClass(
+			TEXT("au.MetaSound.AutoUpdate.NativeClassesOfEqualVersion"),
+			MetaSoundAutoUpdateNativeClassesOfEqualVersionCVar,
+			TEXT("If true, node references to native classes that share a version number will attempt to auto-update if the interface is different, which results in slower graph load times.\n")
+			TEXT("0: Don't auto-update native classes of the same version with interface discrepancies, !0: Auto-update native classes of the same version with interface discrepancies (default)"),
+			ECVF_Default);
+
 		namespace DocumentPrivate
 		{
 			template <typename ResolveType>
@@ -1420,13 +1428,16 @@ bool FMetasoundFrontendClassInput::RemoveDefault(const FGuid& InPageID)
 	return Defaults.RemoveAllSwap(IsPage) > 0;
 }
 
-void FMetasoundFrontendClassInput::ResetDefaults()
+void FMetasoundFrontendClassInput::ResetDefaults(bool bInitializeDefaultPage)
 {
 	using namespace Metasound::Frontend;
 
 	Defaults.Reset();
-	InitDefault();
-	Defaults.Shrink();
+	if (bInitializeDefaultPage)
+	{
+		InitDefault();
+		Defaults.Shrink();
+	}
 }
 
 void FMetasoundFrontendClassInput::SetDefaults(TArray<FMetasoundFrontendClassInputDefault> InputDefaults)
diff --git a/Engine/Plugins/Runtime/Metasound/Source/MetasoundFrontend/Private/MetasoundFrontendDocumentBuilder.cpp b/Engine/Plugins/Runtime/Metasound/Source/MetasoundFrontend/Private/MetasoundFrontendDocumentBuilder.cpp
--- a/Engine/Plugins/Runtime/Metasound/Source/MetasoundFrontend/Private/MetasoundFrontendDocumentBuilder.cpp
+++ b/Engine/Plugins/Runtime/Metasound/Source/MetasoundFrontend/Private/MetasoundFrontendDocumentBuilder.cpp
@@ -277,8 +277,11 @@ namespace Metasound::Frontend
 			{
 				for (const FMetasoundFrontendInterface& FromInterface : Options.InterfacesToRemove)
 				{
-					InputsToRemove.Append(FromInterface.Inputs);
-					OutputsToRemove.Append(FromInterface.Outputs);
+					if (Document.Interfaces.Contains(FromInterface.Metadata.Version))
+					{
+						InputsToRemove.Append(FromInterface.Inputs);
+						OutputsToRemove.Append(FromInterface.Outputs);
+					}
 				}
 
 				for (const FMetasoundFrontendInterface& ToInterface : Options.InterfacesToAdd)
@@ -1701,7 +1704,7 @@ const FMetasoundFrontendNode* FMetaSoundFrontendDocumentBuilder::AddGraphVariabl
 	return nullptr;
 }
 
-bool FMetaSoundFrontendDocumentBuilder::AddInterface(FName InterfaceName)
+bool FMetaSoundFrontendDocumentBuilder::AddInterface(FName InterfaceName, bool bAddUserModifiableInterfaceOnly)
 {
 	using namespace Metasound::Frontend;
 
@@ -1736,9 +1739,9 @@ bool FMetaSoundFrontendDocumentBuilder::AddInterface(FName InterfaceName)
 			const FTopLevelAssetPath BuilderClassPath = GetBuilderClassPath();
 			auto FindClassOptionsPredicate = [&BuilderClassPath](const FMetasoundFrontendInterfaceUClassOptions& Options) { return Options.ClassPath == BuilderClassPath; };
 			const FMetasoundFrontendInterfaceUClassOptions* ClassOptions = Entry->GetInterface().Metadata.UClassOptions.FindByPredicate(FindClassOptionsPredicate);
-			if (ClassOptions && !ClassOptions->bIsModifiable)
+			if (ClassOptions && bAddUserModifiableInterfaceOnly && !ClassOptions->bIsModifiable)
 			{
-				UE_LOG(LogMetaSound, Error, TEXT("DocumentBuilder failed to add MetaSound Interface '%s' to document: is not set to be modifiable for given UClass '%s'"), *InterfaceName.ToString(), *BuilderClassPath.ToString());
+				UE_LOG(LogMetaSound, Error, TEXT("DocumentBuilder failed to add MetaSound Interface '%s' to document: is not set to be user modifiable for given UClass '%s'"), *InterfaceName.ToString(), *BuilderClassPath.ToString());
 				return false;
 			}
 
@@ -1838,6 +1841,56 @@ const FMetasoundFrontendNode* FMetaSoundFrontendDocumentBuilder::AddNodeByClassN
 	return nullptr;
 }
 
+const FMetasoundFrontendNode* FMetaSoundFrontendDocumentBuilder::AddNodeByClassName(const FMetasoundFrontendClassName& InClassName, int32 InMajorVersion, int32 InMinorVersion, FGuid InNodeID, const FGuid* InPageID)
+{
+	using namespace Metasound;
+	using namespace Metasound::Frontend;
+
+	const FMetasoundFrontendClass* Dependency = nullptr;
+	FNodeRegistryKey ClassKey;
+	{
+		INodeClassRegistry* NodeRegistry = INodeClassRegistry::Get();
+		if (ensure(NodeRegistry))
+		{
+			FNodeRegistryKey LookupKey(EMetasoundFrontendClassType::External, InClassName, InMajorVersion, InMinorVersion);
+			FMetasoundFrontendClass RegisteredClass;
+			if (!NodeRegistry->FindFrontendClassFromRegistered(LookupKey, RegisteredClass))
+			{
+				UE_LOG(LogMetaSound, Error, TEXT("Failed to add new node by class name '%s' and major version '%d' and minor version '%d': Class not found"), *InClassName.ToString(), InMajorVersion, InMinorVersion);
+				return nullptr;
+			}
+
+			const EMetasoundFrontendClassType ClassType = RegisteredClass.Metadata.GetType();
+			if (ClassType != EMetasoundFrontendClassType::External && ClassType != EMetasoundFrontendClassType::Graph)
+			{
+				UE_LOG(LogMetaSound, Warning, TEXT("Failed to add new node by class name '%s': Class is restricted type '%s' that cannot be added via this function."),
+					*InClassName.ToString(),
+					LexToString(ClassType));
+				return nullptr;
+			}
+
+			// Dependency is considered "External" when looked up or added as a dependency to a graph
+			RegisteredClass.Metadata.SetType(EMetasoundFrontendClassType::External);
+			ClassKey = FNodeRegistryKey(RegisteredClass.Metadata);
+			Dependency = FindDependency(RegisteredClass.Metadata);
+			if (!Dependency)
+			{
+				Dependency = AddDependency(MoveTemp(RegisteredClass));
+			}
+		}
+	}
+
+	if (Dependency)
+	{
+		auto FinalizeNode = [](const FMetasoundFrontendNode& Node, const Metasound::Frontend::FNodeRegistryKey& ClassKey) { return Node.Name; };
+		constexpr int32* NewNodeIndex = nullptr;
+		const FGuid& PageID = InPageID ? *InPageID : BuildPageID;
+		return AddNodeInternal(Dependency->Metadata, FinalizeNode, PageID, InNodeID, NewNodeIndex);
+	}
+
+	return nullptr;
+}
+
 const FMetasoundFrontendNode* FMetaSoundFrontendDocumentBuilder::AddNodeByTemplate(const Metasound::Frontend::INodeTemplate& InTemplate, FNodeTemplateGenerateInterfaceParams Params, FGuid InNodeID, const FGuid* InPageID)
 {
 	using namespace Metasound;
@@ -2052,6 +2105,92 @@ bool FMetaSoundFrontendDocumentBuilder::CanAddEdge(const FMetasoundFrontendEdge&
 	return false;
 }
 
+#if WITH_EDITORONLY_DATA
+bool FMetaSoundFrontendDocumentBuilder::CanAutoUpdate(const FGuid& InNodeID, Metasound::Frontend::FClassInterfaceUpdates& OutInterfaceUpdates, const FGuid* InPageID)
+{
+	using namespace Metasound::Frontend;
+	METASOUND_TRACE_CPUPROFILER_EVENT_SCOPE(FMetaSoundFrontendDocumentBuilder::CanAutoUpdate);
+
+	OutInterfaceUpdates = { };
+	const FGuid& PageID = InPageID ? *InPageID : BuildPageID;
+	const FMetasoundFrontendNode* Node = FindNode(InNodeID, InPageID);
+	if (!Node)
+	{
+		return false;
+	}
+	const FMetasoundFrontendClass* Class = FindDependency(Node->ClassID);
+	if (!Class)
+	{
+		return false;
+	}
+
+	const FMetasoundFrontendClassMetadata& NodeClassMetadata = Class->Metadata;
+	const IMetaSoundAssetManager& AssetManager = IMetaSoundAssetManager::GetChecked();
+	
+	FMetasoundFrontendClass RegistryClass;
+	if (!ISearchEngine::Get().FindClassWithHighestMinorVersion(NodeClassMetadata.GetClassName().ToNodeClassName(), NodeClassMetadata.GetVersion().Major, RegistryClass))
+	{
+		return false;
+	}
+
+	// 1. Document's class version is somehow higher than registries, so can't update.
+	if (RegistryClass.Metadata.GetVersion() < NodeClassMetadata.GetVersion())
+	{
+		return false;
+	}
+
+	// 2. Document's class version is equal, so have to dif and check change IDs
+	// to ensure a change wasn't created that didn't contain a version change.
+	if (RegistryClass.Metadata.GetVersion() == NodeClassMetadata.GetVersion())
+	{
+		if (ChangeIDComparisonEnabledInAutoUpdate())
+		{
+			// WIP: Merging these paths.  Should no longer require using different
+			// logic to determine changes in native vs asset class definitions.
+			// Its less performant to ignore ChangeIDs, but auto-update is no longer
+			// called at runtime and the editor now avoids running on mass numbers of
+			// assets at once due to the utilization of AssetTag data for querying class info.
+			const bool bIsAssetClass = AssetManager.IsAssetClass(RegistryClass.Metadata);
+			if (bIsAssetClass)
+			{
+				if (RegistryClass.Metadata.GetChangeID() == NodeClassMetadata.GetChangeID())
+				{
+					const FMetasoundFrontendClassInterface& ClassInterface = Class->GetInterfaceForNode(*Node);
+					const FGuid& NodeClassInterfaceChangeID = ClassInterface.GetChangeID();
+					if (RegistryClass.GetDefaultInterface().GetChangeID() == NodeClassInterfaceChangeID)
+					{
+						return false;
+					}
+				}
+			}
+			else
+			{
+				if (!MetaSoundAutoUpdateNativeClassesOfEqualVersionCVar)
+				{
+					return false;
+				}
+			}
+		}
+
+		constexpr bool bUseHighestMinorVersion = true;
+		constexpr bool bForceRegenerateClassInterfaceOverride = true;
+
+		DiffAgainstRegistryInterface(InNodeID, PageID, bUseHighestMinorVersion, OutInterfaceUpdates, bForceRegenerateClassInterfaceOverride);
+
+		return OutInterfaceUpdates.ContainsChanges();
+	}
+
+	// 3. Document's class version is out-of-date, so dif and always return true that can auto-update
+	// (Unlike the case where the version is equal, the version must be updated even if the interface
+	// contains no changes).
+	constexpr bool bUseHighestMinorVersion = true;
+	constexpr bool bForceRegenerateClassInterfaceOverride = true;
+	DiffAgainstRegistryInterface(InNodeID, PageID, bUseHighestMinorVersion, OutInterfaceUpdates, bForceRegenerateClassInterfaceOverride);
+
+	return true;
+}
+#endif // WITH_EDITORONLY_DATA
+
 void FMetaSoundFrontendDocumentBuilder::ClearDocument(TSharedRef<Metasound::Frontend::FDocumentModifyDelegates> ModifyDelegates)
 {
 	FMetasoundFrontendDocument& Doc = GetDocumentChecked();
@@ -2066,7 +2205,6 @@ void FMetaSoundFrontendDocumentBuilder::ClearDocument(TSharedRef<Metasound::Fron
 #endif // WITH_EDITOR
 
 	GraphClass.PresetOptions.InputsInheritingDefault.Empty();
-	GraphClass.PresetOptions.bIsPreset = false;
 
 	// Removing graph pages is not necessary when editor only data is not available as graph mutation
 	// is only supported in builds with editor data loaded. Otherwise, anything calling ClearDocument
@@ -2253,21 +2391,19 @@ bool FMetaSoundFrontendDocumentBuilder::ConvertFromPreset()
 }
 
 bool FMetaSoundFrontendDocumentBuilder::ConvertToPreset(const FMetasoundFrontendDocument& InReferencedDocument, TSharedPtr<Metasound::Frontend::FDocumentModifyDelegates> ModifyDelegates)
+{
+	return false;
+}
+
+bool FMetaSoundFrontendDocumentBuilder::ConvertToPreset(const FMetaSoundFrontendDocumentBuilder& InReferencedDocumentBuilder, TSharedPtr<Metasound::Frontend::FDocumentModifyDelegates> ModifyDelegates)
 {
 	using namespace Metasound;
 	using namespace Metasound::Frontend;
 
-	TSharedRef<FDocumentModifyDelegates> ModifyDelegatesRef = ModifyDelegates.IsValid() ? ModifyDelegates->AsShared() : MakeShared<FDocumentModifyDelegates>(InReferencedDocument);
+	TSharedRef<FDocumentModifyDelegates> ModifyDelegatesRef = ModifyDelegates.IsValid() ? ModifyDelegates->AsShared() : MakeShared<FDocumentModifyDelegates>(InReferencedDocumentBuilder.GetConstDocumentChecked());
 	ClearDocument(ModifyDelegatesRef);
 
 	FMetasoundFrontendGraphClass& PresetAssetRootGraph = GetDocumentChecked().RootGraph;
-	PresetAssetRootGraph.IterateGraphPages([](FMetasoundFrontendGraph& PresetAssetGraph)
-	{
-#if WITH_EDITORONLY_DATA
-		PresetAssetGraph.Style.bIsGraphEditable = false;
-#endif // WITH_EDITORONLY_DATA
-	});
-
 	// Mark all inputs as inherited by default
 	{
 		PresetAssetRootGraph.PresetOptions.InputsInheritingDefault.Reset();
@@ -2277,18 +2413,10 @@ bool FMetaSoundFrontendDocumentBuilder::ConvertToPreset(const FMetasoundFrontend
 	}
 
 	// Apply root graph transform
-	FRebuildPresetRootGraph RebuildPresetRootGraph(InReferencedDocument);
-	if (RebuildPresetRootGraph.Transform(GetDocumentChecked()))
+	FRebuildPresetRootGraph RebuildPresetRootGraph(InReferencedDocumentBuilder);
+	if (RebuildPresetRootGraph.Transform(*this))
 	{
 		DocumentInterface->ConformObjectToDocument();
-
-		// TL/DR: Have to reload and assign delegates here due to the rebuild preset transform still being implemented via controllers.
-		// Onces its reimplemented with the builder API, this can be removed.
-		// 
-		// The invalidate cache call when accessing the mutable document handle from within the transform unfortunately doesn't reach this
-		// builder's cache indirectly as converting to preset can be called by transient builders that are not registered with the MetaSound
-		// builder subsystem.
-		Reload(ModifyDelegates);
 		return true;
 	}
 
@@ -2307,8 +2435,21 @@ TOptional<Metasound::FAnyDataReference> FMetaSoundFrontendDocumentBuilder::Creat
 	return IDataTypeRegistry::Get().CreateDataReference(DataType, AccessType, InLiteral, InOperatorSettings);
 };
 
+void FMetaSoundFrontendDocumentBuilder::SetPresetFlags(bool bIsPreset)
+{
+	FMetasoundFrontendGraphClass& RootGraph = GetDocumentChecked().RootGraph;
+	RootGraph.PresetOptions.bIsPreset = bIsPreset;
+	RootGraph.IterateGraphPages([&](FMetasoundFrontendGraph& PresetAssetGraph)
+	{
+#if WITH_EDITORONLY_DATA
+		PresetAssetGraph.Style.bIsGraphEditable = !bIsPreset;
+#endif // WITH_EDITORONLY_DATA
+	});
+}
+
 #if WITH_EDITORONLY_DATA
-bool FMetaSoundFrontendDocumentBuilder::DiffAgainstRegistryInterface(const FGuid& InNodeID, const FGuid& InPageID, bool bInUseHighestMinorVersion, Metasound::Frontend::FClassInterfaceUpdates& OutInterfaceUpdates) const
+bool FMetaSoundFrontendDocumentBuilder::DiffAgainstRegistryInterface(const FGuid& InNodeID, const FGuid& InPageID, bool bInUseHighestMinorVersion, Metasound::Frontend::FClassInterfaceUpdates& OutInterfaceUpdates, bool bInForceRegenerateClassInterfaceOverride) const
+
 {
 	METASOUND_TRACE_CPUPROFILER_EVENT_SCOPE(FMetaSoundFrontendDocumentBuilder::DiffAgainstRegistryInterface);
 
@@ -2333,6 +2474,7 @@ bool FMetaSoundFrontendDocumentBuilder::DiffAgainstRegistryInterface(const FGuid
 
 	Metasound::FNodeClassName NodeClassName = NodeClassMetadata.GetClassName().ToNodeClassName();
 
+	// Get class 
 	if (bInUseHighestMinorVersion)
 	{
 		if (!ISearchEngine::Get().FindClassWithHighestMinorVersion(NodeClassName, NodeClassMetadata.GetVersion().Major, OutInterfaceUpdates.RegistryClass))
@@ -2345,7 +2487,7 @@ bool FMetaSoundFrontendDocumentBuilder::DiffAgainstRegistryInterface(const FGuid
 	else
 	{
 		// Find class with same metadata in the node registry.
-		FMetasoundFrontendRegistryContainer* Registry = FMetasoundFrontendRegistryContainer::Get();
+		INodeClassRegistry* Registry = INodeClassRegistry::Get();
 		checkf(nullptr != Registry, TEXT("The metasound node registry should always be available if the metasound plugin is loaded"));
 		bool bFoundRegisteredClass = Registry->FindFrontendClassFromRegistered(FNodeRegistryKey(NodeClassMetadata), OutInterfaceUpdates.RegistryClass);
 		if (!bFoundRegisteredClass)
@@ -2358,6 +2500,14 @@ bool FMetaSoundFrontendDocumentBuilder::DiffAgainstRegistryInterface(const FGuid
 		}
 	}
 
+	// Get updates of the node's configuration
+	FNodeConfigurationUpdateData NodeConfigUpdates;
+	FindNodeConfigurationUpdates(InNodeID, InPageID, OutInterfaceUpdates.RegistryClass, NodeConfigUpdates, bInForceRegenerateClassInterfaceOverride);
+
+	// Based upon all the node configuration updates, determine which class interface is the correct up-to-date interface
+	const FMetasoundFrontendClassInterface& ApplicableRegistryInterface = GetApplicableRegistryInterface(OutInterfaceUpdates.RegistryClass, NodeConfigUpdates);
+
+	// Diff interface 
 	Algo::Transform(OutInterfaceUpdates.RegistryClass.GetInterfaceForNode(*Node).Inputs, OutInterfaceUpdates.AddedInputs, [](const FMetasoundFrontendClassInput& Input) { return &Input; });
 	for (const FMetasoundFrontendClassInput& Input : NodeClassInterface.Inputs)
 	{
@@ -2397,6 +2547,27 @@ bool FMetaSoundFrontendDocumentBuilder::DiffAgainstRegistryInterface(const FGuid
 		}
 	}
 
+	const bool bFoundDifferenceInInterface = OutInterfaceUpdates.ContainsRemovedMembers() || OutInterfaceUpdates.ContainsAddedMembers();
+
+	// Add node config and class interface override info if necessary. 
+	if (ShouldReplaceExistingNodeConfig(NodeConfigUpdates.RegisteredConfig, NodeConfigUpdates.ExistingConfig))
+	{
+		OutInterfaceUpdates.AddedConfiguration = MoveTemp(NodeConfigUpdates.RegisteredConfig);
+		OutInterfaceUpdates.RemovedConfiguration = NodeConfigUpdates.ExistingConfig;
+	}
+
+	if (NodeConfigUpdates.bDidUpdateClassInterfaceOverride)
+	{
+		// The class interface override is updated if there was any difference found
+		// in the interface, or the override was added/removed. 
+		const bool bReplaceClassInterfaceOverride = bFoundDifferenceInInterface || (NodeConfigUpdates.ExistingClassInterfaceOverride.IsValid() != NodeConfigUpdates.RegeneratedClassInterfaceOverride.IsValid());
+		if (bReplaceClassInterfaceOverride)
+		{
+			OutInterfaceUpdates.AddedClassInterfaceOverride = MoveTemp(NodeConfigUpdates.RegeneratedClassInterfaceOverride);
+			OutInterfaceUpdates.RemovedClassInterfaceOverride = NodeConfigUpdates.ExistingClassInterfaceOverride;
+		}
+	}
+
 	return true;
 }
 #endif // WITH_EDITORONLY_DATA
@@ -2553,6 +2724,16 @@ const FMetasoundFrontendGraph& FMetaSoundFrontendDocumentBuilder::FindConstBuild
 	return GetConstDocumentChecked().RootGraph.FindConstGraphChecked(BuildPageID);
 }
 
+TConstStructView<FMetasoundFrontendClassInterface> FMetaSoundFrontendDocumentBuilder::FindClassInterfaceOverride(const FGuid& InNodeID, const FGuid* InPageID) const
+{
+	TConstStructView<FMetasoundFrontendClassInterface> ClassInterfaceOverride;
+	if (const FMetasoundFrontendNode* Node = FindNode(InNodeID, InPageID))
+	{
+		ClassInterfaceOverride = Node->ClassInterfaceOverride;
+	}
+	return ClassInterfaceOverride;
+}
+
 bool FMetaSoundFrontendDocumentBuilder::FindDeclaredInterfaces(TArray<const Metasound::Frontend::IInterfaceRegistryEntry*>& OutInterfaces) const
 {
 	return FindDeclaredInterfaces(GetConstDocumentChecked(), OutInterfaces);
@@ -2596,6 +2777,11 @@ const FMetasoundFrontendClass* FMetaSoundFrontendDocumentBuilder::FindDependency
 	return DocumentCache->FindDependency(RegistryKey);
 }
 
+const FMetasoundFrontendClass* FMetaSoundFrontendDocumentBuilder::FindDependency(const Metasound::Frontend::FNodeRegistryKey& InRegistryKey) const
+{
+	return DocumentCache->FindDependency(InRegistryKey);
+}
+
 TArray<const FMetasoundFrontendEdge*> FMetaSoundFrontendDocumentBuilder::FindEdges(const FGuid& InNodeID, const FGuid& InVertexID, const FGuid* InPageID) const
 {
 	using namespace Metasound::Frontend;
@@ -2965,6 +3151,69 @@ TInstancedStruct<FMetaSoundFrontendNodeConfiguration> FMetaSoundFrontendDocument
 	return TInstancedStruct<FMetaSoundFrontendNodeConfiguration>();
 }
 
+bool FMetaSoundFrontendDocumentBuilder::ShouldReplaceExistingNodeConfig(TConstStructView<FMetaSoundFrontendNodeConfiguration> InRegisteredNodeConfig, TConstStructView<FMetaSoundFrontendNodeConfiguration> InExistingConfig) const
+{
+	// The existing node configuration is not always replaced. Replacing
+	// the node configuration struct loses any modified state on the existing
+	// node config. In the scenario where the node configurations point to
+	// the same derived class type, we do not replace them. 
+	return InRegisteredNodeConfig.GetScriptStruct() != InExistingConfig.GetScriptStruct();
+}
+
+const FMetasoundFrontendClassInterface& FMetaSoundFrontendDocumentBuilder::GetApplicableRegistryInterface(const FMetasoundFrontendClass& InRegisteredClass, const FNodeConfigurationUpdateData& InNodeConfigurationUpdates) const
+{
+	// Check for the latest class interface overrides
+	if (InNodeConfigurationUpdates.bDidUpdateClassInterfaceOverride)
+	{
+		if (const FMetasoundFrontendClassInterface* Interface = InNodeConfigurationUpdates.RegeneratedClassInterfaceOverride.GetPtr())
+		{
+			return *Interface;
+		}
+	}
+	else
+	{
+		// The class interface override wasn't updated. Use the original if it exists. 
+		if (const FMetasoundFrontendClassInterface* Interface = InNodeConfigurationUpdates.ExistingClassInterfaceOverride.GetPtr())
+		{
+			return *Interface;
+		}
+	}
+
+	// If the override is invalid, use the interface on the class.
+	return InRegisteredClass.GetDefaultInterface();
+}
+
+void FMetaSoundFrontendDocumentBuilder::FindNodeConfigurationUpdates(const FGuid& InNodeID, const FGuid& InPageID, const FMetasoundFrontendClass& InRegisteredClass, FNodeConfigurationUpdateData& OutNodeConfigurationUpdates, bool bInForceRegenerateClassInterfaceOverride) const
+{
+	using namespace Metasound::Frontend;
+
+	const INodeClassRegistry* Registry = INodeClassRegistry::Get();
+	checkf(nullptr != Registry, TEXT("The metasound node registry should always be available if the metasound plugin is loaded"));
+
+	OutNodeConfigurationUpdates.ExistingConfig = FindNodeConfiguration(InNodeID, &InPageID);
+	OutNodeConfigurationUpdates.ExistingClassInterfaceOverride = FindClassInterfaceOverride(InNodeID, &InPageID);
+
+	OutNodeConfigurationUpdates.RegisteredConfig = Registry->CreateFrontendNodeConfiguration(FNodeRegistryKey(InRegisteredClass.Metadata));
+
+	// Determine which node config will be used to create the class interface override;
+	const bool bWillReplaceNodeConfig = ShouldReplaceExistingNodeConfig(OutNodeConfigurationUpdates.RegisteredConfig, OutNodeConfigurationUpdates.ExistingConfig);
+	TConstStructView<FMetaSoundFrontendNodeConfiguration> EffectiveNodeConfig = bWillReplaceNodeConfig ? OutNodeConfigurationUpdates.RegisteredConfig : OutNodeConfigurationUpdates.ExistingConfig;
+
+	// Update the node config
+	if (bWillReplaceNodeConfig || bInForceRegenerateClassInterfaceOverride)
+	{
+		// Set this flag to true even if the effective node configuration pointer 
+		// is invalid. If the effective node configuration is null, this captures
+		// the fact that a null node configuration should produce an empty 
+		// ClassInterfaceOverride. 
+		OutNodeConfigurationUpdates.bDidUpdateClassInterfaceOverride = true;
+		if (const FMetaSoundFrontendNodeConfiguration* ConfigPtr = EffectiveNodeConfig.GetPtr())
+		{
+			OutNodeConfigurationUpdates.RegeneratedClassInterfaceOverride = ConfigPtr->OverrideDefaultInterface(InRegisteredClass);
+		}
+	}
+}
+
 const int32* FMetaSoundFrontendDocumentBuilder::FindNodeIndex(const FGuid& InNodeID, const FGuid* InPageID) const
 {
 	using namespace Metasound::Frontend;
@@ -4178,6 +4427,41 @@ void FMetaSoundFrontendDocumentBuilder::IterateNodes(Metasound::Frontend::FConst
 	}
 }
 
+void FMetaSoundFrontendDocumentBuilder::IterateNodesByPredicate(Metasound::Frontend::FConstClassAndNodeAndPageIDFunctionRef Func, Metasound::Frontend::FNodePredicateFunctionRef PredicateFunc, const FGuid* InPageID, bool bIterateAllPages) const
+{
+	using namespace Metasound::Frontend;
+	FMetasoundFrontendDocument& Doc = GetDocumentChecked();
+	if (bIterateAllPages)
+	{
+		Doc.RootGraph.IterateGraphPages([&](FMetasoundFrontendGraph& Graph)
+		{
+			for (const FMetasoundFrontendNode& Node : Graph.Nodes)
+			{
+				if (PredicateFunc(Node))
+				{
+					if (const FMetasoundFrontendClass* Class = FindDependency(Node.ClassID))
+					{
+						Func(*Class, Node, Graph.PageID);
+					}
+				}
+			}
+		});
+	}
+	else
+	{
+		const FGuid& PageID = InPageID ? *InPageID : BuildPageID;
+		const FMetasoundFrontendGraph& Graph = Doc.RootGraph.FindGraphChecked(PageID);
+
+		for (const FMetasoundFrontendNode& Node : Graph.Nodes)
+		{
+			if (const FMetasoundFrontendClass* Class = FindDependency(Node.ClassID))
+			{
+				Func(*Class, Node, Graph.PageID);
+			}
+		}
+	}
+}
+
 void FMetaSoundFrontendDocumentBuilder::IterateNodesByClassType(Metasound::Frontend::FConstClassAndNodeFunctionRef Func, EMetasoundFrontendClassType ClassType, const FGuid* InPageID) const
 {
 	using namespace Metasound::Frontend;
@@ -5107,12 +5391,14 @@ int32 FMetaSoundFrontendDocumentBuilder::RemoveNodeLocation(const FGuid& InNodeI
 
 bool FMetaSoundFrontendDocumentBuilder::RemoveUnusedDependencies()
 {
-	bool bDidEdit = false;
-
 	const FMetasoundFrontendDocument& Document = GetConstDocumentChecked();
 	const FMetasoundFrontendGraphClass& RootGraph = Document.RootGraph;
 	const TArray<FMetasoundFrontendClass>& Dependencies = Document.Dependencies;
 
+	// Remove duplicate entries, keeping the highest version dependencies. 
+	bool bDidEdit = false;
+
+	// Remove unused dependencies
 	for (int32 Index = Dependencies.Num() - 1; Index >= 0; --Index)
 	{
 		const FGuid& ClassID = Dependencies[Index].ID;
@@ -5132,6 +5418,73 @@ bool FMetaSoundFrontendDocumentBuilder::RenameRootGraphClass(const FMetasoundFro
 	return false;
 }
 
+bool FMetaSoundFrontendDocumentBuilder::ReplaceDependency(const Metasound::Frontend::FNodeClassRegistryKey& OldNodeClassKey, const Metasound::Frontend::FNodeClassRegistryKey& NewNodeClassKey, TArray<FVertexNameAndType>* OutDisconnectedInputs, TArray<FVertexNameAndType>* OutDisconnectedOutputs)
+{
+	using namespace Metasound::Frontend;
+
+	// Early out if there is no dependency to replace
+	const FMetasoundFrontendClass* OldDependency = FindDependency(OldNodeClassKey);
+	if (!OldDependency)
+	{
+		return false;
+	}
+
+	// For all matching nodes, cache off connections and other data 
+	TArray<FNodeInstanceReplacementData> NodeData;
+
+	auto CaptureNodeData = [&](const FMetasoundFrontendClass&, const FMetasoundFrontendNode& Node, const FGuid& PageID)
+	{
+		FNodeInstanceReplacementData NodeInstanceData = CaptureNodeInstanceReplacementData(*const_cast<FMetasoundFrontendNode*>(&Node), &PageID);
+		NodeData.Add(MoveTemp(NodeInstanceData));
+	};
+
+	auto GetNodesOfOldClass = [&](const FMetasoundFrontendNode& Node)
+	{
+		const FMetasoundFrontendClass* Class = FindDependency(Node.ClassID);
+		if (Class)
+		{
+			const FNodeClassRegistryKey KeyToCompare(Class->Metadata);
+			return KeyToCompare == OldNodeClassKey;
+		}
+		return false;
+	};
+
+	IterateNodesByPredicate(CaptureNodeData, GetNodesOfOldClass, /*InPageID =*/nullptr, /*bIterateAllPages=*/true);
+	
+	// Remove dependency and associated nodes
+	bool bSuccess = RemoveDependency(OldNodeClassKey.Type, OldNodeClassKey.ClassName, OldNodeClassKey.Version);
+
+	// Clean up dependency
+	RemoveUnusedDependencies();
+
+	// Readd nodes with new class
+	for (FNodeInstanceReplacementData& NodeInstanceData : NodeData)
+	{
+		const FMetasoundFrontendNode* NewNode = AddNodeByClassName(NewNodeClassKey.ClassName, NewNodeClassKey.Version.Major, NodeInstanceData.NodeID, &NodeInstanceData.PageID);
+		bSuccess |= (NewNode != nullptr);
+	}
+
+	// Reapply connections and node data 
+	const bool bNodeVersionUpdated = NewNodeClassKey.Version > OldNodeClassKey.Version;
+	for (FNodeInstanceReplacementData& NodeInstanceData : NodeData)
+	{
+		FMetasoundFrontendNode* NewNode = FindNodeInternal(NodeInstanceData.NodeID, &NodeInstanceData.PageID);
+		if (NewNode)
+		{
+			const FMetasoundFrontendClass* Class = FindDependency(NewNode->ClassID);
+			if (Class)
+			{
+#if WITH_EDITOR
+				NodeInstanceData.Style.bMessageNodeUpdated = bNodeVersionUpdated;
+#endif // WITH_EDITOR
+				ApplyNodeInstanceReplacementData(*NewNode, MoveTemp(NodeInstanceData), OutDisconnectedInputs, OutDisconnectedOutputs);
+			}
+		}
+	}
+
+	return bSuccess;
+}
+
 void FMetaSoundFrontendDocumentBuilder::ReloadCache()
 {
 	using namespace Metasound::Frontend;
@@ -5139,6 +5492,199 @@ void FMetaSoundFrontendDocumentBuilder::ReloadCache()
 	Reload(DocumentDelegates, true);
 }
 
+void FMetaSoundFrontendDocumentBuilder::ApplyNodeInstanceReplacementData(FMetasoundFrontendNode& InReplacementNode, FMetaSoundFrontendDocumentBuilder::FNodeInstanceReplacementData&& InInstanceData, TArray<FVertexNameAndType>* OutDisconnectedInputs, TArray<FVertexNameAndType>* OutDisconnectedOutputs)
+{
+	using namespace Metasound::Frontend;
+	const FGuid& NodeID = InReplacementNode.GetID();
+	const FGuid* PageID = &InInstanceData.PageID;
+
+#if WITH_EDITOR
+	InReplacementNode.Style = MoveTemp(InInstanceData.Style);
+#endif // WITH_EDITOR
+
+	// The logic may appear a little backwards here because adding the node already instantiated the replacement
+	// node configuration instances. The purpose is to revert the replacement if it should not have happened in 
+	// the first place.
+	if (!ShouldReplaceExistingNodeConfig(InReplacementNode.Configuration, InInstanceData.Configuration))
+	{
+		// We should have kept the original node configuration. Set back to the original
+		SetNodeConfiguration(NodeID, MoveTemp(InInstanceData.Configuration), PageID);
+	}
+	
+	// Set literal and reconnect input connections if possible
+	for (FMetasoundFrontendVertex& InputVertex : InReplacementNode.Interface.Inputs)
+	{
+		const FVertexNameAndType ConnectionKey(InputVertex.Name, InputVertex.TypeName);
+		if (FInputConnectionInfo* ConnectionInfo = InInstanceData.InputConnections.Find(ConnectionKey))
+		{
+			if (ConnectionInfo->bLiteralSet)
+			{
+				SetNodeInputDefault(NodeID, InputVertex.VertexID, ConnectionInfo->DefaultValue, PageID);
+			}
+
+			if (ConnectionInfo->ConnectedOutput.IsSet())
+			{
+				FMetasoundFrontendEdge NewEdge
+				{ 
+					ConnectionInfo->ConnectedOutput.NodeID,
+					ConnectionInfo->ConnectedOutput.VertexID,
+					NodeID, 
+					InputVertex.VertexID, 
+				};
+
+				const bool bIsValidEdge = IsValidEdge(NewEdge, PageID) == Metasound::Frontend::EInvalidEdgeReason::None;
+				// It is possible the edge already exists if the edge is between 
+				// this and another node of the same dependency that just added it when it was replaced 
+				if (bIsValidEdge && !ContainsEdge(NewEdge, PageID))
+				{
+					AddEdge(MoveTemp(NewEdge), PageID);
+				}
+
+				// Remove connection to track missing connections between 
+				// node versions.
+				InInstanceData.InputConnections.Remove(ConnectionKey);
+			}
+		}
+	}
+
+	// Track missing input connections
+	if (nullptr != OutDisconnectedInputs)
+	{
+		for (const auto& ConnectionInfoKV : InInstanceData.InputConnections)
+		{
+			const FInputConnectionInfo& ConnectionInfo = ConnectionInfoKV.Value;
+			if (ConnectionInfo.ConnectedOutput.IsSet())
+			{
+				OutDisconnectedInputs->Add(FVertexNameAndType{ ConnectionInfo.Name, ConnectionInfo.DataType });
+			}
+		}
+	}
+
+	// Reconnect output connections if possible
+	for (FMetasoundFrontendVertex& OutputVertex : InReplacementNode.Interface.Outputs)
+	{
+		const FVertexNameAndType ConnectionKey(OutputVertex.Name, OutputVertex.TypeName);
+		if (FOutputConnectionInfo* ConnectionInfo = InInstanceData.OutputConnections.Find(ConnectionKey))
+		{
+			bool bConnectionSuccess = false;
+			for (const FMetasoundFrontendVertexHandle& ConnectedInputHandle : ConnectionInfo->ConnectedInputs)
+			{
+				if (ConnectedInputHandle.IsSet())
+				{
+					FMetasoundFrontendEdge NewEdge
+					{ 
+						NodeID,
+						OutputVertex.VertexID,
+						ConnectedInputHandle.NodeID,
+						ConnectedInputHandle.VertexID
+					};
+
+					const bool bIsValidEdge = IsValidEdge(NewEdge, PageID) == Metasound::Frontend::EInvalidEdgeReason::None;
+					if (bIsValidEdge)
+					{
+						AddEdge(MoveTemp(NewEdge), PageID);
+					}
+					bConnectionSuccess = true;
+				}
+			}
+
+			// Remove connection to track missing connections between 
+			// node versions.
+			if (bConnectionSuccess)
+			{
+				InInstanceData.OutputConnections.Remove(ConnectionKey);
+			}
+		}
+	}
+
+	// Track missing output connections
+	if (nullptr != OutDisconnectedOutputs)
+	{
+		for (const auto& ConnectionInfoKV : InInstanceData.OutputConnections)
+		{
+			const FOutputConnectionInfo& ConnectionInfo = ConnectionInfoKV.Value;
+			const bool bAnyConnectedInputs = Algo::AnyOf(ConnectionInfo.ConnectedInputs, [](const FMetasoundFrontendVertexHandle& Input) { return Input.IsSet(); });
+			if (bAnyConnectedInputs)
+			{
+				OutDisconnectedOutputs->Add(FVertexNameAndType{ ConnectionInfo.Name, ConnectionInfo.DataType });
+			}
+		}
+	}
+}
+
+FMetaSoundFrontendDocumentBuilder::FNodeInstanceReplacementData FMetaSoundFrontendDocumentBuilder::CaptureNodeInstanceReplacementData(FMetasoundFrontendNode& InOriginalNode, const FGuid* InPageID)
+{
+	const FGuid& PageID = InPageID ? *InPageID : BuildPageID;
+
+	FNodeInstanceReplacementData ReplacementData;
+	ReplacementData.PageID = PageID;
+
+#if WITH_EDITOR
+	ReplacementData.Style = MoveTemp(InOriginalNode.Style);
+#endif // WITH_EDITOR
+
+	// Move any configuration and override data living on the node.
+	ReplacementData.Configuration = MoveTemp(InOriginalNode.Configuration);
+	ReplacementData.ClassInterfaceOverride = MoveTemp(InOriginalNode.ClassInterfaceOverride);
+	ReplacementData.NodeID = InOriginalNode.GetID();
+
+	// Cache input/output connections by name to try so they can be
+	// hooked back up after swapping to the new class version.
+	for (const FMetasoundFrontendVertex& InputVertex : InOriginalNode.Interface.Inputs)
+	{
+		bool bLiteralSet = false;
+		FMetasoundFrontendLiteral DefaultLiteral;
+
+		auto VertexLiteralMatchesID = [&](const FMetasoundFrontendVertexLiteral& VertexLiteral)
+		{
+			return VertexLiteral.VertexID == InputVertex.VertexID;
+		};
+
+		if (const FMetasoundFrontendVertexLiteral* VertexLiteral = InOriginalNode.InputLiterals.FindByPredicate(VertexLiteralMatchesID))
+		{
+			DefaultLiteral = VertexLiteral->Value;
+			bLiteralSet = true;
+		}
+
+		const FVertexNameAndType ConnectionKey(InputVertex.Name, InputVertex.TypeName);
+		const FMetasoundFrontendNode* ConnectedOutputNode = nullptr;
+		const FMetasoundFrontendVertex* ConnectedOutputVertex = FindNodeOutputConnectedToNodeInput(InOriginalNode.GetID(), InputVertex.VertexID, &ConnectedOutputNode, &PageID);
+		FMetasoundFrontendVertexHandle ConnectedOutputHandle = (ConnectedOutputNode && ConnectedOutputVertex) ?
+			FMetasoundFrontendVertexHandle{ ConnectedOutputNode->GetID(), ConnectedOutputVertex->VertexID } : FMetasoundFrontendVertexHandle();
+		ReplacementData.InputConnections.Add(ConnectionKey, FInputConnectionInfo
+		{
+			MoveTemp(ConnectedOutputHandle),
+			InputVertex.Name, 
+			InputVertex.TypeName,
+			MoveTemp(DefaultLiteral),
+			bLiteralSet
+		});
+	}
+
+	for (const FMetasoundFrontendVertex& OutputVertex : InOriginalNode.Interface.Outputs)
+	{
+		const FVertexNameAndType ConnectionKey(OutputVertex.Name, OutputVertex.TypeName);
+		TArray<const FMetasoundFrontendNode*> ConnectedInputNodes;
+		TArray<const FMetasoundFrontendVertex*> ConnectedInputVertices = FindNodeInputsConnectedToNodeOutput(InOriginalNode.GetID(), OutputVertex.VertexID, &ConnectedInputNodes, &PageID);
+		TArray<FMetasoundFrontendVertexHandle> ConnectedInputInfo;
+		check(ConnectedInputNodes.Num() == ConnectedInputVertices.Num());
+		for (int i = 0; i < ConnectedInputVertices.Num(); ++i)
+		{
+			check(ConnectedInputNodes[i]);
+			check(ConnectedInputVertices[i]);
+			ConnectedInputInfo.Add(FMetasoundFrontendVertexHandle{ ConnectedInputNodes[i]->GetID(), ConnectedInputVertices[i]->VertexID });
+		}
+		ReplacementData.OutputConnections.Add(ConnectionKey, FOutputConnectionInfo
+		{
+			ConnectedInputInfo,
+			OutputVertex.Name, 
+			OutputVertex.TypeName
+		});
+	}
+	
+	return ReplacementData;
+}
+
 #if WITH_EDITORONLY_DATA
 bool FMetaSoundFrontendDocumentBuilder::ResetGraphInputDefault(FName InputName)
 {
@@ -5562,6 +6108,13 @@ bool FMetaSoundFrontendDocumentBuilder::SetGraphInputInheritsDefault(FName InNam
 	return bValueChanged;
 }
 
+bool FMetaSoundFrontendDocumentBuilder::SetGraphInputsInheritingDefault(TSet<FName>&& InNames)
+{
+	FMetasoundFrontendGraphClassPresetOptions& PresetOptions = GetDocumentChecked().RootGraph.PresetOptions;
+	PresetOptions.InputsInheritingDefault = MoveTemp(InNames);
+	return true;
+}
+
 bool FMetaSoundFrontendDocumentBuilder::SetGraphInputName(FName InputName, FName NewName)
 {
 	using namespace Metasound::Frontend;
@@ -5970,6 +6523,18 @@ PRAGMA_ENABLE_DEPRECATION_WARNINGS
 	DocumentDelegates->OnDocumentMetadataChanged.Broadcast();
 }
 
+void FMetaSoundFrontendDocumentBuilder::SetInputStyle(FMetasoundFrontendInterfaceStyle&& Style)
+{
+	FMetasoundFrontendClassInterface& RootGraphClassInterface = GetDocumentChecked().RootGraph.GetDefaultInterface();
+	RootGraphClassInterface.SetInputStyle(MoveTemp(Style));
+}
+
+void FMetaSoundFrontendDocumentBuilder::SetOutputStyle(FMetasoundFrontendInterfaceStyle&& Style)
+{
+	FMetasoundFrontendClassInterface& RootGraphClassInterface = GetDocumentChecked().RootGraph.GetDefaultInterface();
+	RootGraphClassInterface.SetOutputStyle(MoveTemp(Style));
+}
+
 void FMetaSoundFrontendDocumentBuilder::SetMemberMetadata(UMetaSoundFrontendMemberMetadata& NewMetadata)
 {
 	check(NewMetadata.MemberID.IsValid());
@@ -6547,6 +7112,41 @@ bool FMetaSoundFrontendDocumentBuilder::UnlinkVariableNode(const FGuid& InNodeID
 }
 
 #if WITH_EDITOR
+bool FMetaSoundFrontendDocumentBuilder::SynchronizeDependencyMetadata(TArray<const FMetasoundFrontendClass*>* InOutModifiedClasses)
+{
+	using namespace Metasound::Frontend;
+
+	if (InOutModifiedClasses)
+	{
+		InOutModifiedClasses->Reset();
+	}
+
+	bool bSuccess = true;
+	FMetasoundFrontendDocument& Document = GetDocumentChecked();
+	for (FMetasoundFrontendClass& Dependency : Document.Dependencies)
+	{
+		FMetasoundFrontendClass RegistryVersion;
+		const FNodeRegistryKey RegistryKey(Dependency.Metadata);
+		bool bFoundRegisteredClass = INodeClassRegistry::Get()->FindFrontendClassFromRegistered(RegistryKey, RegistryVersion);
+		if (bFoundRegisteredClass)
+		{
+			if (Dependency.Metadata.GetChangeID() != RegistryVersion.Metadata.GetChangeID())
+			{
+				Dependency.Metadata = MoveTemp(RegistryVersion.Metadata);
+				if (InOutModifiedClasses)
+				{
+					InOutModifiedClasses->Add(&Dependency);
+				}
+			}
+		}
+		else
+		{
+			bSuccess = false;
+		}
+	}
+	return bSuccess;
+}
+
 bool FMetaSoundFrontendDocumentBuilder::UpdateDependencyRegistryData(const TMap<Metasound::Frontend::FNodeRegistryKey, Metasound::Frontend::FNodeRegistryKey>& OldToNewClassKeys)
 {
 	using namespace Metasound::Frontend;
diff --git a/Engine/Plugins/Runtime/Metasound/Source/MetasoundFrontend/Private/MetasoundFrontendNodeController.cpp b/Engine/Plugins/Runtime/Metasound/Source/MetasoundFrontend/Private/MetasoundFrontendNodeController.cpp
--- a/Engine/Plugins/Runtime/Metasound/Source/MetasoundFrontend/Private/MetasoundFrontendNodeController.cpp
+++ b/Engine/Plugins/Runtime/Metasound/Source/MetasoundFrontend/Private/MetasoundFrontendNodeController.cpp
@@ -4,7 +4,6 @@
 
 #include "Algo/AnyOf.h"
 #include "Algo/Transform.h"
-#include "HAL/IConsoleManager.h"
 #include "MetasoundFrontendDocumentAccessPtr.h"
 #include "MetasoundFrontendDocumentVersioning.h"
 #include "MetasoundFrontendInputController.h"
@@ -16,14 +15,6 @@
 
 #define LOCTEXT_NAMESPACE "MetasoundFrontendNodeController"
 
-static int32 MetaSoundAutoUpdateNativeClassesOfEqualVersionCVar = 1;
-FAutoConsoleVariableRef CVarMetaSoundAutoUpdateNativeClass(
-	TEXT("au.MetaSound.AutoUpdate.NativeClassesOfEqualVersion"),
-	MetaSoundAutoUpdateNativeClassesOfEqualVersionCVar,
-	TEXT("If true, node references to native classes that share a version number will attempt to auto-update if the interface is different, which results in slower graph load times.\n")
-	TEXT("0: Don't auto-update native classes of the same version with interface discrepancies, !0: Auto-update native classes of the same version with interface discrepancies (default)"),
-	ECVF_Default);
-
 namespace Metasound
 {
 	namespace Frontend
@@ -756,7 +747,6 @@ namespace Metasound
 
 		FNodeHandle FBaseNodeController::ReplaceWithVersion(const FMetasoundFrontendVersionNumber& InNewVersion, TArray<FVertexNameAndType>* OutDisconnectedInputs, TArray<FVertexNameAndType>* OutDisconnectedOutputs)
 		{
-			// TODO: Need to deal with node configuration and class interface override here. 
 			const FMetasoundFrontendClassMetadata Metadata = GetClassMetadata();
 
 			// Lookup new version in node registry
diff --git a/Engine/Plugins/Runtime/Metasound/Source/MetasoundFrontend/Private/MetasoundFrontendNodeController.h b/Engine/Plugins/Runtime/Metasound/Source/MetasoundFrontend/Private/MetasoundFrontendNodeController.h
--- a/Engine/Plugins/Runtime/Metasound/Source/MetasoundFrontend/Private/MetasoundFrontendNodeController.h
+++ b/Engine/Plugins/Runtime/Metasound/Source/MetasoundFrontend/Private/MetasoundFrontendNodeController.h
@@ -67,6 +67,7 @@ namespace Metasound
 			UE_DEPRECATED(5.6, "Moved to FMetaSoundFrontendDocumentBuilder")
 			virtual bool DiffAgainstRegistryInterface(FClassInterfaceUpdates& OutInterfaceUpdates, bool bInUseHighestMinorVersion, bool bInForceRegenerateClassInterfaceOverride) const override;
 
+			UE_DEPRECATED(5.7, "Moved to FMetaSoundFrontendDocumentBuilder")
 			virtual bool CanAutoUpdate(FClassInterfaceUpdates& OutInterfaceUpdates) const override;
 #endif // WITH_EDITORONLY_DATA
 
diff --git a/Engine/Plugins/Runtime/Metasound/Source/MetasoundFrontend/Private/MetasoundFrontendTransform.cpp b/Engine/Plugins/Runtime/Metasound/Source/MetasoundFrontend/Private/MetasoundFrontendTransform.cpp
--- a/Engine/Plugins/Runtime/Metasound/Source/MetasoundFrontend/Private/MetasoundFrontendTransform.cpp
+++ b/Engine/Plugins/Runtime/Metasound/Source/MetasoundFrontend/Private/MetasoundFrontendTransform.cpp
@@ -80,11 +80,6 @@ namespace Metasound
 			return false;
 		}
 
-		bool INodeTransform::Transform(FMetasoundFrontendNode& InOutNode) const
-		{
-			return false;
-		}
-
 		FModifyRootGraphInterfaces::FModifyRootGraphInterfaces(const TArray<FMetasoundFrontendInterface>& InInterfacesToRemove, const TArray<FMetasoundFrontendInterface>& InInterfacesToAdd)
 			: InterfacesToRemove(InInterfacesToRemove)
 			, InterfacesToAdd(InInterfacesToAdd)
@@ -521,502 +516,496 @@ namespace Metasound
 			}
 		}
 
-		bool FAutoUpdateRootGraph::Transform(FDocumentHandle InDocument)
+		bool FAutoUpdateRootGraph::Transform(FMetaSoundFrontendDocumentBuilder& InOutBuilderToTransform)
 		{
 			METASOUND_TRACE_CPUPROFILER_EVENT_SCOPE(FAutoUpdateRootGraph::Transform);
 			bool bDidEdit = false;
 
-			FMetasoundAssetBase* PresetReferencedMetaSoundAsset = nullptr;
-			TArray<TPair<FNodeHandle, FMetasoundFrontendVersionNumber>> NodesToUpdate;
-
-			FGraphHandle RootGraph = InDocument->GetRootGraph();
-			const bool bIsPreset = RootGraph->GetGraphPresetOptions().bIsPreset;
-			RootGraph->IterateNodes([&](FNodeHandle NodeHandle)
+			// <Old dependency, new dependency>
+			using FDependencyPair = TPair<FNodeClassRegistryKey, FNodeClassRegistryKey>;
+			TSet<FDependencyPair> DependenciesToUpdate;
+			const bool bIsPreset = InOutBuilderToTransform.IsPreset();
+			const FMetasoundFrontendGraphClass& RootGraph = InOutBuilderToTransform.GetConstDocumentChecked().RootGraph;
+			
+			// Go through external nodes and find dependencies to update
+			RootGraph.IterateGraphPages([&](const FMetasoundFrontendGraph& Graph)
 			{
-				using namespace Metasound::Frontend;
-
-				const FMetasoundFrontendClassMetadata& ClassMetadata = NodeHandle->GetClassMetadata();
-				const FNodeRegistryKey RegistryKey(ClassMetadata);
-
-				if (FMetasoundAssetBase* ReferencedMetaSoundAsset = IMetaSoundAssetManager::GetChecked().FindAsset(RegistryKey))
+				const FGuid PageID = Graph.PageID;
+				InOutBuilderToTransform.IterateNodesByClassType([&](const FMetasoundFrontendClass& Class, const FMetasoundFrontendNode& Node)
 				{
-					if (bIsPreset)
-					{
-						PresetReferencedMetaSoundAsset = ReferencedMetaSoundAsset;
-					}
-				}
-				else
-				{
-					if (bIsPreset)
+					const FNodeClassRegistryKey& NodeClassKey = FNodeClassRegistryKey(Class.Metadata);
+
+					FClassInterfaceUpdates InterfaceUpdates;
+					const bool bHasUpdated = UpdatedClasses.Contains(Class.ID);
+					const bool bCanAutoUpdate = InOutBuilderToTransform.CanAutoUpdate(Node.GetID(), InterfaceUpdates);
+					if (bHasUpdated || !bCanAutoUpdate)
 					{
-						METASOUND_VERSIONING_LOG(Error, TEXT("Auto-Updating preset '%s' failed: Referenced class '%s' missing."), *DebugAssetPath, *ClassMetadata.GetClassName().ToString());
 						return;
 					}
-				}
-
-				FClassInterfaceUpdates InterfaceUpdates;
-				const FGuid& ClassID = NodeHandle->GetClassID();
-				const bool bHasUpdated = UpdatedClasses.Contains(ClassID);
-				if (!bHasUpdated && !NodeHandle->CanAutoUpdate(InterfaceUpdates))
-				{
-					return;
-				}
-
-				UpdatedClasses.Add(ClassID);
-
-				// Check if a updated minor version exists.
-				FMetasoundFrontendClass ClassWithHighestMinorVersion;
-				const bool bFoundClassInSearchEngine = Frontend::ISearchEngine::Get().FindClassWithHighestMinorVersion(ClassMetadata.GetClassName(), ClassMetadata.GetVersion().Major, ClassWithHighestMinorVersion);
 
-				if (bFoundClassInSearchEngine && (ClassWithHighestMinorVersion.Metadata.GetVersion() > ClassMetadata.GetVersion()))
-				{
-					const FMetasoundFrontendVersionNumber UpdateVersion = ClassWithHighestMinorVersion.Metadata.GetVersion();
+					UpdatedClasses.Add(Class.ID);
 
-					METASOUND_VERSIONING_LOG(Display, TEXT("Auto-Updating '%s' node class '%s': Newer version '%s' found."), *DebugAssetPath, *ClassMetadata.GetClassName().ToString(), *UpdateVersion.ToString());
+					// Check if a updated minor version exists.
+					FMetasoundFrontendClass ClassWithHighestMinorVersion;
+					const bool bFoundClassInSearchEngine = Frontend::ISearchEngine::Get().FindClassWithHighestMinorVersion(NodeClassKey.ClassName, NodeClassKey.Version.Major, ClassWithHighestMinorVersion);
 
-					NodesToUpdate.Add(TPair<FNodeHandle, FMetasoundFrontendVersionNumber>(NodeHandle, UpdateVersion));
-				}
-				else if (InterfaceUpdates.ContainsChanges())
-				{
-					const FMetasoundFrontendVersionNumber UpdateVersion = ClassMetadata.GetVersion();
-					METASOUND_VERSIONING_LOG(Display, TEXT("Auto-Updating '%s' node class '%s (%s)': Interface change detected."), *DebugAssetPath, *ClassMetadata.GetClassName().ToString(), *UpdateVersion.ToString());
+					if (bFoundClassInSearchEngine && (ClassWithHighestMinorVersion.Metadata.GetVersion() > NodeClassKey.Version))
+					{
+						const FMetasoundFrontendVersionNumber UpdateVersion = ClassWithHighestMinorVersion.Metadata.GetVersion();
 
-					NodesToUpdate.Add(TPair<FNodeHandle, FMetasoundFrontendVersionNumber>(NodeHandle, UpdateVersion));
-				}
+						METASOUND_VERSIONING_LOG(Display, TEXT("Auto-Updating '%s' node class '%s': Newer version '%s' found."), *DebugAssetPath, *NodeClassKey.ClassName.ToString(), *UpdateVersion.ToString());
 
-				else
+						DependenciesToUpdate.Add(FDependencyPair(NodeClassKey, FNodeClassRegistryKey(ClassWithHighestMinorVersion.Metadata)));
+					}
+					else if (InterfaceUpdates.ContainsChanges())
+					{
+						METASOUND_VERSIONING_LOG(Display, TEXT("Auto-Updating '%s' node class '%s (%s)': Interface change detected."), *DebugAssetPath, *NodeClassKey.ClassName.ToString(), *NodeClassKey.Version.ToString());
+						DependenciesToUpdate.Add(FDependencyPair(NodeClassKey, NodeClassKey));
+					}
+					else
+					{
+						DependenciesToUpdate.Add(FDependencyPair(NodeClassKey, NodeClassKey));
+					}
+				}, EMetasoundFrontendClassType::External, &PageID);
+			});
+			
+			// If preset, rebuild root graph 
+			if (bIsPreset)
+			{
+				FMetasoundAssetBase* PresetReferencedMetaSoundAsset = InOutBuilderToTransform.GetReferencedPresetAsset();
+				if (!PresetReferencedMetaSoundAsset)
 				{
-					NodesToUpdate.Add(TPair<FNodeHandle, FMetasoundFrontendVersionNumber>(NodeHandle, ClassMetadata.GetVersion()));
+					METASOUND_VERSIONING_LOG(Error, TEXT("Auto-Updating preset '%s' failed: Referenced class missing."), *DebugAssetPath);
+					return false;
 				}
-			}, EMetasoundFrontendClassType::External);
+				TScriptInterface<IMetaSoundDocumentInterface> RefDocInterface = PresetReferencedMetaSoundAsset->GetOwningAsset();
+				FMetaSoundFrontendDocumentBuilder& ReferenceBuilder = IDocumentBuilderRegistry::GetChecked().FindOrBeginBuilding(RefDocInterface);
 
-			if (PresetReferencedMetaSoundAsset)
-			{
-				if (bIsPreset)
+				bDidEdit |= FRebuildPresetRootGraph(ReferenceBuilder).Transform(InOutBuilderToTransform);
+				if (bDidEdit)
 				{
-					// This call to FindOrBeginBuilding is a bit of a hack to guarantee instantiation of the referenced graph's builder while
-					// using the deprecated controllers to rebuild the preset root graph.  This is needed in case the referenced graph implements
-					// page data as the builder is required to resolve the correct graph data used to rebuild the preset.  Once the rebuild transform
-					// is migrated to the Builder API, this will no longer be necessary.
-					TScriptInterface<IMetaSoundDocumentInterface> RefDocInterface = PresetReferencedMetaSoundAsset->GetOwningAsset();
-					FMetaSoundFrontendDocumentBuilder& ReferenceBuilder = IDocumentBuilderRegistry::GetChecked().FindOrBeginBuilding(RefDocInterface);
-
-					bDidEdit |= FRebuildPresetRootGraph(PresetReferencedMetaSoundAsset->GetDocumentHandle()).Transform(InDocument);
-					if (bDidEdit)
+					FMetasoundFrontendClassMetadata PresetMetadata = InOutBuilderToTransform.GetConstDocumentChecked().RootGraph.Metadata;
+					PresetMetadata.SetType(EMetasoundFrontendClassType::External);
+					const FNodeRegistryKey RegistryKey(PresetMetadata);
+					FMetasoundAssetBase* PresetMetaSoundAsset = IMetaSoundAssetManager::GetChecked().TryLoadAssetFromKey(RegistryKey);
+					if (ensure(PresetMetaSoundAsset))
 					{
-						FMetasoundFrontendClassMetadata PresetMetadata = InDocument->GetRootGraphClass().Metadata;
-						PresetMetadata.SetType(EMetasoundFrontendClassType::External);
-						const FNodeRegistryKey RegistryKey(PresetMetadata);
-						FMetasoundAssetBase* PresetMetaSoundAsset = IMetaSoundAssetManager::GetChecked().TryLoadAssetFromKey(RegistryKey);
-						if (ensure(PresetMetaSoundAsset))
-						{
-							TScriptInterface<IMetaSoundDocumentInterface> PresetInterface = PresetMetaSoundAsset->GetOwningAsset();
-							check(PresetInterface);
-							PresetInterface->ConformObjectToDocument();
-						}
-
-						InDocument->RemoveUnreferencedDependencies();
-						InDocument->SynchronizeDependencyMetadata();
+						TScriptInterface<IMetaSoundDocumentInterface> PresetInterface = PresetMetaSoundAsset->GetOwningAsset();
+						check(PresetInterface);
+						PresetInterface->ConformObjectToDocument();
 					}
+						
+					InOutBuilderToTransform.RemoveUnusedDependencies();
+					InOutBuilderToTransform.SynchronizeDependencyMetadata();
 				}
 			}
-			else
-			{
-				using FVertexNameAndType = INodeController::FVertexNameAndType;
 
-				bDidEdit |= !NodesToUpdate.IsEmpty();
-				for (const TPair<FNodeHandle, FMetasoundFrontendVersionNumber>& Pair : NodesToUpdate)
-				{
-					FNodeHandle ExistingNode = Pair.Key;
-					FMetasoundFrontendVersionNumber InitialVersion = ExistingNode->GetClassMetadata().GetVersion();
+			bDidEdit |= !DependenciesToUpdate.IsEmpty();
 
-					TArray<FVertexNameAndType> DisconnectedInputs;
-					TArray<FVertexNameAndType> DisconnectedOutputs;
-					FNodeHandle NewNode = ExistingNode->ReplaceWithVersion(Pair.Value, &DisconnectedInputs, &DisconnectedOutputs);
+			// Update each dependency and log disconnections
+			for (const FDependencyPair& DependencyPair : DependenciesToUpdate)
+			{
+				using FVertexNameAndType = FMetaSoundFrontendDocumentBuilder::FVertexNameAndType;
+				TArray<FVertexNameAndType> DisconnectedInputs;
+				TArray<FVertexNameAndType> DisconnectedOutputs;
+				InOutBuilderToTransform.ReplaceDependency(DependencyPair.Key, DependencyPair.Value, &DisconnectedInputs, &DisconnectedOutputs);
 
-					// Log warnings for any disconnections
-					if (bLogWarningOnDroppedConnection)
+				// Log warnings for any disconnections
+				if (bLogWarningOnDroppedConnection)
+				{
+					if ((DisconnectedInputs.Num() > 0) || (DisconnectedOutputs.Num() > 0))
 					{
-						if ((DisconnectedInputs.Num() > 0) || (DisconnectedOutputs.Num() > 0))
-						{
-							const FString NodeClassName = NewNode->GetClassMetadata().GetClassName().ToString();
-							const FString NewClassVersion = Pair.Value.ToString();
+						const FString NodeClassName = DependencyPair.Key.ClassName.ToString();
+						const FString NewClassVersion = DependencyPair.Value.Version.ToString();
 
-							for (const FVertexNameAndType& InputPin : DisconnectedInputs)
-							{
-								DocumentTransform::LogAutoUpdateWarning(FString::Printf(TEXT("Auto-Updating '%s' node class '%s (%s)': Previously connected input '%s' with data type '%s' no longer exists."), *DebugAssetPath, *NodeClassName, *NewClassVersion, *InputPin.Get<0>().ToString(), *InputPin.Get<1>().ToString()));
-							}
+						for (const FVertexNameAndType& InputPin : DisconnectedInputs)
+						{
+							DocumentTransform::LogAutoUpdateWarning(FString::Printf(TEXT("Auto-Updating '%s' node class '%s (%s)': Previously connected input '%s' with data type '%s' no longer exists."), *DebugAssetPath, *NodeClassName, *NewClassVersion, *InputPin.Get<0>().ToString(), *InputPin.Get<1>().ToString()));
+						}
 
-							for (const FVertexNameAndType& OutputPin : DisconnectedOutputs)
-							{
-								DocumentTransform::LogAutoUpdateWarning(FString::Printf(TEXT("Auto-Updating '%s' node class '%s (%s)': Previously connected output '%s' with data type '%s' no longer exists."), *DebugAssetPath, *NodeClassName, *NewClassVersion, *OutputPin.Get<0>().ToString(), *OutputPin.Get<1>().ToString()));
-							}
+						for (const FVertexNameAndType& OutputPin : DisconnectedOutputs)
+						{
+							DocumentTransform::LogAutoUpdateWarning(FString::Printf(TEXT("Auto-Updating '%s' node class '%s (%s)': Previously connected output '%s' with data type '%s' no longer exists."), *DebugAssetPath, *NodeClassName, *NewClassVersion, *OutputPin.Get<0>().ToString(), *OutputPin.Get<1>().ToString()));
 						}
 					}
 				}
-
-				InDocument->RemoveUnreferencedDependencies();
-				InDocument->SynchronizeDependencyMetadata();
 			}
 
+			InOutBuilderToTransform.RemoveUnusedDependencies();
+			InOutBuilderToTransform.SynchronizeDependencyMetadata();
+
 			return bDidEdit;
 		}
+
+		bool FAutoUpdateRootGraph::Transform(FDocumentHandle InDocument)
+		{
+			return false;
+		}
 #endif // WITH_EDITORONLY_DATA
 
 		bool FRebuildPresetRootGraph::Transform(FDocumentHandle InDocument) const
 		{
-			METASOUND_TRACE_CPUPROFILER_EVENT_SCOPE(Metasound::Frontend::FRebuildPresetRootGraph::Transform);
+			return false;
+		}
+	
+		FRebuildPresetRootGraph::FRebuildPresetRootGraph(const FMetasoundFrontendDocument& InReferencedDocument)
+		{
+		}
 
-			FGraphHandle PresetGraphHandle = InDocument->GetRootGraph();
-			if (!ensure(PresetGraphHandle->IsValid()))
-			{
-				return false;
-			}
+		bool FRebuildPresetRootGraph::Transform(FMetasoundFrontendDocument& InDocument) const
+		{
+			return false;
+		}
 
+		bool FRebuildPresetRootGraph::Transform(FMetaSoundFrontendDocumentBuilder& InOutBuilderToTransform) const
+		{
+			METASOUND_TRACE_CPUPROFILER_EVENT_SCOPE(Metasound::Frontend::FRebuildPresetRootGraph::Transform);
+			
 			// Callers of this transform should check that the graph is supposed to
-			// be managed externally before calling this transform. If a scenario
+			// be managed externally and the parent builder is valid before calling this transform. If a scenario
 			// arises where this transform is used outside of AutoUpdate, then this
 			// early exist should be removed as it's mostly here to protect against
 			// accidental manipulation of metasound graphs.
-			if (!ensure(PresetGraphHandle->GetGraphPresetOptions().bIsPreset))
-			{
-				return false;
-			}
-
-			FConstGraphHandle ReferencedGraphHandle = ReferencedDocument->GetRootGraph();
-			if (!ensure(ReferencedGraphHandle->IsValid()))
+			if (!InOutBuilderToTransform.IsPreset() || !ParentBuilder)
 			{
 				return false;
 			}
-
+	
+			const FMetasoundFrontendDocument& DocumentToTransform = InOutBuilderToTransform.GetConstDocumentChecked();
+			const FMetasoundFrontendDocument& ParentDocument = ParentBuilder->GetConstDocumentChecked();
+			
 			// Determine the inputs and outputs needed in the wrapping graph. Also
 			// cache any exiting literals that have been set on the wrapping graph.
 			TSet<FName> InputsInheritingDefault;
-			TArray<FMetasoundFrontendClassInput> ClassInputs = GenerateRequiredClassInputs(InDocument, PresetGraphHandle, InputsInheritingDefault);
-			TArray<FMetasoundFrontendClassOutput> ClassOutputs = GenerateRequiredClassOutputs(InDocument, PresetGraphHandle);
+			TArray<FMetasoundFrontendClassInput> ClassInputs = GenerateRequiredClassInputs(InOutBuilderToTransform, InputsInheritingDefault);
+			TArray<FMetasoundFrontendClassOutput> ClassOutputs = GenerateRequiredClassOutputs(InOutBuilderToTransform);
 
 #if WITH_EDITORONLY_DATA
 			// Cache off member metadata so it be can be readded if necessary after the graph is cleared 
-			const FMemberIDToMetadataMap CachedMemberMetadata = InDocument->GetMetadata()->MemberMetadata;
+			const FMemberIDToMetadataMap CachedMemberMetadata = DocumentToTransform.Metadata.MemberMetadata;
 #endif // WITH_EDITORONLY_DATA
 
 			FGuid PresetNodeID;
-			PresetGraphHandle->IterateConstNodes([InPresetNodeID = &PresetNodeID](FConstNodeHandle PresetNodeHandle)
+			InOutBuilderToTransform.IterateNodesByClassType([&](const FMetasoundFrontendClass&, const FMetasoundFrontendNode& Node)
 			{
-				*InPresetNodeID = PresetNodeHandle->GetID();
-			}, EMetasoundFrontendClassType::External); 
-			
+				PresetNodeID = Node.GetID();
+			}, EMetasoundFrontendClassType::External);
+
 			if (!PresetNodeID.IsValid())
 			{
 				// This ID was originally being set to FGuid::NewGuid. 
 				// If you were reliant on that ID, please resave the asset so it is serialized with a valid ID
-				PresetNodeID = InDocument->GetRootGraph()->GetClassID();
+				PresetNodeID = DocumentToTransform.RootGraph.ID;
 			}
 
 			// Clear the root graph so it can be rebuilt.
-			PresetGraphHandle->ClearGraph();
+			TSharedRef<FDocumentModifyDelegates> ModifyDelegates = MakeShared<FDocumentModifyDelegates>(InOutBuilderToTransform.GetConstDocumentChecked());
+			InOutBuilderToTransform.ClearDocument(ModifyDelegates);
+
+			InOutBuilderToTransform.SetPresetFlags();
 
 			// Ensure preset interfaces match those found in referenced graph.  Referenced graph is assumed to be
 			// well-formed (i.e. all inputs/outputs/environment variables declared by interfaces are present, and
 			// of proper name & data type).
-			const TSet<FMetasoundFrontendVersion>& RefInterfaceVersions = ReferencedDocument->GetInterfaceVersions();
+			const TSet<FMetasoundFrontendVersion>& RefInterfaceVersions = ParentDocument.Interfaces;
 			for (const FMetasoundFrontendVersion& Version : RefInterfaceVersions)
 			{
-				InDocument->AddInterfaceVersion(Version);
+				InOutBuilderToTransform.AddInterface(Version.Name, /*bAddUserModifiableInterfaceOnly = */false);
 			}
 
 			// Add referenced node
-			FMetasoundFrontendClassMetadata ReferencedClassMetadata = ReferencedGraphHandle->GetGraphMetadata();
-			// Swap type on look-up as it will be referenced as an externally defined class relative to the new Preset asset
-			ReferencedClassMetadata.SetType(EMetasoundFrontendClassType::External);
-
-			FNodeHandle ReferencedNodeHandle = PresetGraphHandle->AddNode(ReferencedClassMetadata, PresetNodeID);
-
+			const FMetasoundFrontendClassMetadata& ParentClassMetadata = ParentDocument.RootGraph.Metadata;
+			const FMetasoundFrontendNode* ReferencedNode = InOutBuilderToTransform.AddNodeByClassName(ParentClassMetadata.GetClassName(), ParentClassMetadata.GetVersion().Major, PresetNodeID);
+			check(ReferencedNode);
+			
 #if WITH_EDITOR
-			// Set node location.
-			FMetasoundFrontendNodeStyle RefNodeStyle;
-
-			// Offset to be to the right of input nodes
+			// Set node location, offset to be to the right of input nodes
 			const FGuid EdNodeGuid = FGuid::NewGuid(); // EdNodes are now never serialized and are transient, so just assign here
-			RefNodeStyle.Display.Locations.Add(EdNodeGuid, DisplayStyle::NodeLayout::DefaultOffsetX);
-			ReferencedNodeHandle->SetNodeStyle(RefNodeStyle);
+			InOutBuilderToTransform.SetNodeLocation(PresetNodeID, DisplayStyle::NodeLayout::DefaultOffsetX, &EdNodeGuid);
 #endif // WITH_EDITOR
 
 			// Connect parent graph to referenced graph
-			PresetGraphHandle->SetInputsInheritingDefault(MoveTemp(InputsInheritingDefault));
-
-			AddAndConnectInputs(ClassInputs, PresetGraphHandle, ReferencedNodeHandle);
-			AddAndConnectOutputs(ClassOutputs, PresetGraphHandle, ReferencedNodeHandle);
+			InOutBuilderToTransform.SetGraphInputsInheritingDefault(MoveTemp(InputsInheritingDefault));
+			AddAndConnectInputs(ClassInputs, InOutBuilderToTransform, PresetNodeID);
+			AddAndConnectOutputs(ClassOutputs, InOutBuilderToTransform, PresetNodeID);
 
 #if WITH_EDITORONLY_DATA
-			FMemberIDToMetadataMap& MemberMetadata = InDocument->GetMetadata()->MemberMetadata;
-			AddMemberMetadata(CachedMemberMetadata, PresetGraphHandle, MemberMetadata);
+			AddMemberMetadata(CachedMemberMetadata, InOutBuilderToTransform);
 #endif // WITH_EDITORONLY_DATA
 
 			return true;
 		}
-	
-		FRebuildPresetRootGraph::FRebuildPresetRootGraph(const FMetasoundFrontendDocument& InReferencedDocument)
-		{
-			// TODO: Swap implementation to not use access pointers/controllers
-			ReferencedDocument = IDocumentController::CreateDocumentHandle(InReferencedDocument);
-		}
-
-		bool FRebuildPresetRootGraph::Transform(FMetasoundFrontendDocument& InDocument) const
-		{
-			// TODO: Swap implementation to not use access pointers/controllers
-			return Transform(IDocumentController::CreateDocumentHandle(InDocument));
-		}
 
 #if WITH_EDITORONLY_DATA
-		void FRebuildPresetRootGraph::AddMemberMetadata(const FMemberIDToMetadataMap& InCachedMemberMetadata, FGraphHandle& InPresetGraphHandle, FMemberIDToMetadataMap& InOutMemberMetadata) const
+		void FRebuildPresetRootGraph::AddMemberMetadata(const FMemberIDToMetadataMap& InCachedMemberMetadata, FMetaSoundFrontendDocumentBuilder& InOutBuilderToTransform) const
 		{
 			// Add member metadata if a member with the corresponding node ID exists in the preset graph
 			if (!InCachedMemberMetadata.IsEmpty())
 			{
 				for (const TPair<FGuid, TObjectPtr<UMetaSoundFrontendMemberMetadata>>& MemberMetadataPair : InCachedMemberMetadata)
 				{
-					FConstNodeHandle FoundNodeHandle = InPresetGraphHandle->GetNodeWithID(MemberMetadataPair.Key);
-					if (FoundNodeHandle->IsValid())
+					if (InOutBuilderToTransform.FindNode(MemberMetadataPair.Key))
 					{
-						InOutMemberMetadata.Add(MemberMetadataPair.Key, MemberMetadataPair.Value);
+						InOutBuilderToTransform.SetMemberMetadata(*MemberMetadataPair.Value);
 					}
 				}
 			}
 		}
+
 #endif // WITH_EDITORONLY_DATA
 
-		void FRebuildPresetRootGraph::AddAndConnectInputs(const TArray<FMetasoundFrontendClassInput>& InClassInputs, FGraphHandle& InPresetGraphHandle, FNodeHandle& InReferencedNode) const
+		void FRebuildPresetRootGraph::AddAndConnectInputs(const TArray<FMetasoundFrontendClassInput>& InClassInputs, FMetaSoundFrontendDocumentBuilder& InOutBuilderToTransform, const FGuid& InReferencedNodeID) const
 		{
 			// Add inputs and space appropriately
-			FVector2D InputNodeLocation = FVector2D::ZeroVector;
-
-			FConstGraphHandle ReferencedGraphHandle = ReferencedDocument->GetRootGraph();
-
 			const INodeTemplate* InputTemplate = INodeTemplateRegistry::Get().FindTemplate(FInputNodeTemplate::ClassName);
 			check(InputTemplate);
-			TArray<FNodeHandle> NodeHandles;
+
+			TArray<const FMetasoundFrontendVertex*> ReferencedNodeInputVertices;
+			TArray<const FMetasoundFrontendNode*> InputTemplateNodes;
+
 			for (const FMetasoundFrontendClassInput& ClassInput : InClassInputs)
 			{
-				FNodeHandle InputNode = InPresetGraphHandle->AddInputVertex(ClassInput);
-				if (ensure(InputNode->IsValid()))
+				// Input may have already been added if interface member
+				const FMetasoundFrontendNode* InputNode = InOutBuilderToTransform.FindGraphInputNode(ClassInput.Name);
+				if (!InputNode)
 				{
-					// Connect input node to corresponding referencing node.
-					FOutputHandle OutputToConnect = InputNode->GetOutputWithVertexName(ClassInput.Name);
-					FInputHandle InputToConnect = InReferencedNode->GetInputWithVertexName(ClassInput.Name);
-					ensure(OutputToConnect->Connect(*InputToConnect));
-					NodeHandles.Add(MoveTemp(InputNode));
-
-					// template node takes on data type of concrete input node's output type
-					const FName DataType = InputNode->GetOutputs().Last()->GetDataType();
-
-					FNodeTemplateGenerateInterfaceParams Params { { }, { DataType }};
-					FNodeHandle TemplateNodeHandle = InPresetGraphHandle->AddTemplateNode(*InputTemplate, MoveTemp(Params));
-					TemplateNodeHandle->GetInputs().Last()->Connect(*OutputToConnect);
-					TemplateNodeHandle->GetOutputs().Last()->Connect(*InputToConnect);
+					InputNode = InOutBuilderToTransform.AddGraphInput(ClassInput);
 				}
+				const FGuid InputNodeID = InputNode->GetID();
+				check(InputNode);
+				const FMetasoundFrontendVertex* InputNodeOutputVertex = InOutBuilderToTransform.FindNodeOutput(InputNode->GetID(), ClassInput.Name);
+				check(InputNodeOutputVertex);
+				const FMetasoundFrontendVertex* ReferencedNodeInputVertex = InOutBuilderToTransform.FindNodeInput(InReferencedNodeID, ClassInput.Name);
+				check(ReferencedNodeInputVertex);
+				
+				ReferencedNodeInputVertices.Add(ReferencedNodeInputVertex);
+				// If not editor, just make connection directly between input and output node
+				#if !WITH_EDITORONLY_DATA
+				InOutBuilderToTransform.AddEdge(FMetasoundFrontendEdge
+				{
+					InputNodeID,
+					InputNodeOutputVertex->VertexID,
+					InReferencedNodeID,
+					ReferencedNodeInputVertex->VertexID
+				});
+				// If editor, add template node and connections
+				#else
+				// template node takes on data type of concrete input node's output type
+				const FName DataType = InputNode->Interface.Outputs.Last().TypeName;
+				
+				FNodeTemplateGenerateInterfaceParams Params{ { }, { DataType } };
+				const FMetasoundFrontendNode* TemplateNode = InOutBuilderToTransform.AddNodeByTemplate(*InputTemplate, MoveTemp(Params));
+
+				check(TemplateNode);
+				InputTemplateNodes.Add(TemplateNode);
+
+				const FGuid TemplateNodeInputVertexID = TemplateNode->Interface.Inputs.Last().VertexID;
+				const FGuid TemplateNodeOutputVertexID = TemplateNode->Interface.Outputs.Last().VertexID;
+				InOutBuilderToTransform.AddEdge(FMetasoundFrontendEdge
+				{
+					InputNodeID,
+					InputNodeOutputVertex->VertexID, 
+					TemplateNode->GetID(),
+					TemplateNodeInputVertexID
+				});
+				InOutBuilderToTransform.AddEdge(FMetasoundFrontendEdge
+				{
+					TemplateNode->GetID(),
+					TemplateNodeOutputVertexID,
+					InReferencedNodeID,
+					ReferencedNodeInputVertex->VertexID
+				});
+				#endif
 			}
 
 #if WITH_EDITOR
 			// Sort before adding nodes to graph layout & copy to preset (must be done after all
 			// inputs/outputs are added but before setting locations to propagate effectively)
-			FMetasoundFrontendInterfaceStyle Style = ReferencedGraphHandle->GetInputStyle();
-			InPresetGraphHandle->SetInputStyle(MoveTemp(Style));
+			const FMetasoundFrontendGraphClass& ParentRootGraph = ParentBuilder->GetConstDocumentChecked().RootGraph;
+			FMetasoundFrontendInterfaceStyle Style = ParentRootGraph.GetDefaultInterface().GetInputStyle();
 
-PRAGMA_DISABLE_DEPRECATION_WARNINGS
-			Style.SortDefaults(NodeHandles, DocumentTransform::GetNodeDisplayNameProjection());
-PRAGMA_ENABLE_DEPRECATION_WARNINGS
+			// Sort vertices on referenced node
+			// then use that order to order connected input nodes (which share the same names)
+			auto GetInputDisplayName = [&InOutBuilderToTransform, &InReferencedNodeID](const FMetasoundFrontendVertex& Vertex)
+			{
+				return InOutBuilderToTransform.GetNodeInputDisplayName(InReferencedNodeID, Vertex.Name);
+			};
+			Style.SortVertices(ReferencedNodeInputVertices, GetInputDisplayName);
+			InOutBuilderToTransform.SetInputStyle(MoveTemp(Style));
 
-			InputNodeLocation = FVector2D::ZeroVector;
-			for (const FNodeHandle& NodeHandle : NodeHandles)
+			// Set editor node locations
+			FVector2D InputNodeLocation = FVector2D::ZeroVector;
+			for (const FMetasoundFrontendNode* TemplateNode : InputTemplateNodes)
 			{
-				// Create input template node and set location
-				FMetasoundFrontendNodeStyle NodeStyle;
-				const FGuid EdNodeGuid = FGuid::NewGuid(); // EdNodes are now never serialized and are transient, so just assign here
-				NodeStyle.Display.Locations.Add(EdNodeGuid, InputNodeLocation);
-				FOutputHandle InputNodeOutputHandle = NodeHandle->GetOutputs().Last();
-				FInputHandle InputTemplateNodeInputHandle = InputNodeOutputHandle->GetConnectedInputs().Last();
-				FNodeHandle TemplateNodeHandle = InputTemplateNodeInputHandle->GetOwningNode();
-				TemplateNodeHandle->SetNodeStyle(NodeStyle);
+				InOutBuilderToTransform.SetNodeLocation(TemplateNode->GetID(), InputNodeLocation);
 				InputNodeLocation += DisplayStyle::NodeLayout::DefaultOffsetY;
 			}
 #endif // WITH_EDITOR
 		}
 
-		void FRebuildPresetRootGraph::AddAndConnectOutputs(const TArray<FMetasoundFrontendClassOutput>& InClassOutputs, FGraphHandle& InPresetGraphHandle, FNodeHandle& InReferencedNode) const
+		void FRebuildPresetRootGraph::AddAndConnectOutputs(const TArray<FMetasoundFrontendClassOutput>& InClassOutputs, FMetaSoundFrontendDocumentBuilder& InOutBuilderToTransform, const FGuid& InReferencedNodeID) const
 		{
 			// Add outputs and space appropriately
-			FVector2D OutputNodeLocation = (2 * DisplayStyle::NodeLayout::DefaultOffsetX);
+			TArray<const FMetasoundFrontendVertex*> ReferencedNodeOutputVertices;
 
-			FConstGraphHandle ReferencedGraphHandle = ReferencedDocument->GetRootGraph();
-
-			TArray<FNodeHandle> NodeHandles;
 			for (const FMetasoundFrontendClassOutput& ClassOutput : InClassOutputs)
 			{
-				FNodeHandle OutputNode = InPresetGraphHandle->AddOutputVertex(ClassOutput);
-				if (ensure(OutputNode->IsValid()))
+				// Output may have already been added if interface member
+				const FMetasoundFrontendNode* OutputNode = InOutBuilderToTransform.FindGraphOutputNode(ClassOutput.Name);
+				if (!OutputNode)
 				{
-					// Connect input node to corresponding referenced node. 
-					FInputHandle InputToConnect = OutputNode->GetInputWithVertexName(ClassOutput.Name);
-					FOutputHandle OutputToConnect = InReferencedNode->GetOutputWithVertexName(ClassOutput.Name);
-					ensure(InputToConnect->Connect(*OutputToConnect));
-					NodeHandles.Add(MoveTemp(OutputNode));
+					OutputNode = InOutBuilderToTransform.AddGraphOutput(ClassOutput);
 				}
+				check(OutputNode);
+				
+				// Connect output node input vertex to corresponding referenced node output vertex. 
+				const FMetasoundFrontendVertex* ReferencedNodeOutputVertex = InOutBuilderToTransform.FindNodeOutput(InReferencedNodeID, ClassOutput.Name);
+				check(ReferencedNodeOutputVertex);
+				const FMetasoundFrontendVertex* OutputNodeInputVertex = InOutBuilderToTransform.FindNodeInput(OutputNode->GetID(), ClassOutput.Name);
+				check(OutputNodeInputVertex);
+
+				InOutBuilderToTransform.AddEdge(FMetasoundFrontendEdge
+				{
+					InReferencedNodeID,
+					ReferencedNodeOutputVertex->VertexID,
+					OutputNode->GetID(),
+					OutputNodeInputVertex->VertexID
+				});
+				ReferencedNodeOutputVertices.Add(ReferencedNodeOutputVertex);
 			}
 
 #if WITH_EDITOR
 			// Sort before adding nodes to graph layout & copy to preset (must be done after all
 			// inputs/outputs are added but before setting locations to propagate effectively)
-			FMetasoundFrontendInterfaceStyle Style = ReferencedGraphHandle->GetOutputStyle();
-			InPresetGraphHandle->SetOutputStyle(MoveTemp(Style));
+			const FMetasoundFrontendGraphClass& ParentRootGraph = ParentBuilder->GetConstDocumentChecked().RootGraph;
+			FMetasoundFrontendInterfaceStyle Style = ParentRootGraph.GetDefaultInterface().GetOutputStyle();
+
+			// Sort vertices on referenced node
+			// then use that order to order connected output nodes (which share the same names)
+			auto GetOutputDisplayName = [&InOutBuilderToTransform, &InReferencedNodeID](const FMetasoundFrontendVertex& Vertex)
+			{
+				return InOutBuilderToTransform.GetNodeOutputDisplayName(InReferencedNodeID, Vertex.Name);
+			};
+			Style.SortVertices(ReferencedNodeOutputVertices, GetOutputDisplayName);
 
-PRAGMA_DISABLE_DEPRECATION_WARNINGS
-			Style.SortDefaults(NodeHandles, DocumentTransform::GetNodeDisplayNameProjection());
-PRAGMA_ENABLE_DEPRECATION_WARNINGS
+			InOutBuilderToTransform.SetOutputStyle(MoveTemp(Style));
 
-			// Set output node location
-			for (const FNodeHandle& OutputNode : NodeHandles)
+			// Set output node locations
+			FVector2D OutputNodeLocation = (2 * DisplayStyle::NodeLayout::DefaultOffsetX);
+			for (const FMetasoundFrontendVertex* OutputVertex : ReferencedNodeOutputVertices)
 			{
-				FMetasoundFrontendNodeStyle NodeStyle;
+				const FName NodeName = OutputVertex->Name;
+				const FGuid OutputNodeID = InOutBuilderToTransform.FindGraphOutputNode(NodeName)->GetID();
+
+				// Set editor node locations
 				const FGuid EdNodeGuid = FGuid::NewGuid(); // EdNodes are now never serialized and are transient, so just assign here
-				NodeStyle.Display.Locations.Add(EdNodeGuid, OutputNodeLocation);
-				OutputNode->SetNodeStyle(NodeStyle);
+				InOutBuilderToTransform.SetNodeLocation(OutputNodeID, OutputNodeLocation);
 				OutputNodeLocation += DisplayStyle::NodeLayout::DefaultOffsetY;
 			}
 #endif // WITH_EDITOR
 		}
 
-		TArray<FMetasoundFrontendClassInput> FRebuildPresetRootGraph::GenerateRequiredClassInputs(FDocumentHandle& InDocumentHandle, const FConstGraphHandle& InPresetGraph, TSet<FName>& OutInputsInheritingDefault) const
+		TArray<FMetasoundFrontendClassInput> FRebuildPresetRootGraph::GenerateRequiredClassInputs(FMetaSoundFrontendDocumentBuilder& InDocumentToTransformBuilder, TSet<FName>& OutInputsInheritingDefault) const
 		{
 			TArray<FMetasoundFrontendClassInput> ClassInputs;
+			check(ParentBuilder);
 
-			FConstGraphHandle ReferencedGraph = ReferencedDocument->GetRootGraph();
-			TSet<FName> InputsToInherit = InPresetGraph->GetInputsInheritingDefault();
+			const FMetasoundFrontendGraphClass& ParentRootGraph = ParentBuilder->GetConstDocumentChecked().RootGraph;
+			const FMetasoundFrontendDocument& ParentDocument = ParentBuilder->GetConstDocumentChecked();
+			const FMetasoundFrontendDocument& DocumentToTransform = InDocumentToTransformBuilder.GetConstDocumentChecked();
+			const TSet<FName>* ExistingInputsInheritingDefaultPtr = InDocumentToTransformBuilder.GetGraphInputsInheritingDefault();
+			check(ExistingInputsInheritingDefaultPtr);
+			TSet<FName> InputsInheritingDefault = *ExistingInputsInheritingDefaultPtr;
 
 			// Iterate through all input nodes of referenced graph
-			ReferencedGraph->IterateConstNodes([&](FConstNodeHandle InputNode)
+			for (const FMetasoundFrontendClassInput& ParentClassInput : ParentDocument.RootGraph.GetDefaultInterface().Inputs)
 			{
-				const FName NodeName = InputNode->GetNodeName();
-				FConstInputHandle Input = InputNode->GetConstInputWithVertexName(NodeName);
-				if (ensure(Input->IsValid()))
-				{
-					FMetasoundFrontendClassInput ClassInput;
-
-					ClassInput.Name = NodeName;
-					ClassInput.TypeName = Input->GetDataType();
-					ClassInput.AccessType = Input->GetVertexAccessType();
-
-#if WITH_EDITOR
-					ClassInput.Metadata.SetDescription(InputNode->GetDescription());
-					ClassInput.Metadata.SetDisplayName(Input->GetMetadata().GetDisplayName());
-#endif // WITH_EDITOR
-					FDocumentAccessPtr DocumentPtr = InDocumentHandle->GetDocumentPtr();
-					const FMetasoundFrontendDocument* Document = DocumentPtr.Get();
-					check(Document);
-					ClassInput.VertexID = FDocumentIDGenerator::Get().CreateVertexID(*Document);;
+				// Copy class input and reset defaults and id
+				const FName& NodeName = ParentClassInput.Name;
+				FMetasoundFrontendClassInput NewClassInput = ParentClassInput;
+				NewClassInput.VertexID = FDocumentIDGenerator::Get().CreateVertexID(DocumentToTransform);
+				NewClassInput.ResetDefaults(/*bInitializeDefaultPage=*/false);
 
-					if (const FMetasoundFrontendClassInput* ExistingClassInput = InPresetGraph->FindClassInputWithName(NodeName).Get())
+				if (const FMetasoundFrontendClassInput* ExistingClassInput = InDocumentToTransformBuilder.FindGraphInput(ParentClassInput.Name))
+				{
+					NewClassInput.NodeID = ExistingClassInput->NodeID;
+				}
+				
+				auto InheritDefaultsFromGraph = [&NewClassInput, &NodeName](const FMetaSoundFrontendDocumentBuilder& InBuilder)
+				{
+					if (const FMetasoundFrontendClassInput* GraphClassInput = InBuilder.FindGraphInput(NodeName))
 					{
-						ClassInput.NodeID = ExistingClassInput->NodeID;
+						GraphClassInput->IterateDefaults([&NewClassInput](const FGuid& PageID, const FMetasoundFrontendLiteral& Literal)
+						{
+							NewClassInput.AddDefault(PageID) = Literal;
+						});
 					}
-
-					auto InheritDefaultsFromGraph = [&ClassInput, &NodeName](const FConstGraphHandle& GraphHandle)
+					else
 					{
-						const FGuid GraphVertexID = GraphHandle->GetVertexIDForInputVertex(NodeName);
-						if (const FMetasoundFrontendClassInput* GraphClassInput = GraphHandle->FindInputDescriptionWithVertexID(GraphVertexID))
-						{
-							GraphClassInput->IterateDefaults([&ClassInput](const FGuid& PageID, const FMetasoundFrontendLiteral& Literal)
-							{
-								ClassInput.AddDefault(PageID) = Literal;
-							});
-						}
-						else
-						{
-							ClassInput.InitDefault();
-						}
-					};
+						NewClassInput.InitDefault();
+					}
+				};
 
-					if (InPresetGraph->ContainsInputVertex(NodeName, ClassInput.TypeName))
+				const FMetasoundFrontendClassInput* GraphInput = InDocumentToTransformBuilder.FindGraphInput(NodeName);
+				if (GraphInput && GraphInput->TypeName == ParentClassInput.TypeName)
+				{
+					// If the input vertex already exists in the parent graph,
+					// check if parent should be used or not from set of managed
+					// input names.
+					if (InputsInheritingDefault.Contains(NodeName))
 					{
-						// If the input vertex already exists in the parent graph,
-						// check if parent should be used or not from set of managed
-						// input names.
-						if (InPresetGraph->GetInputsInheritingDefault().Contains(NodeName))
-						{
-							InheritDefaultsFromGraph(ReferencedGraph);
-						}
-						else
-						{
-							InheritDefaultsFromGraph(InPresetGraph);
-						}
+						InheritDefaultsFromGraph(*ParentBuilder);
 					}
 					else
 					{
-						InheritDefaultsFromGraph(ReferencedGraph);
-
-						// Add this to the list of inheriting as the
-						// type no longer matches or it no longer exists
-						InputsToInherit.Add(NodeName);
+						// Use existing defaults 
+						InheritDefaultsFromGraph(InDocumentToTransformBuilder);
 					}
-
-					ClassInputs.Add(MoveTemp(ClassInput));
 				}
-			}, EMetasoundFrontendClassType::Input);
+				else
+				{
+					InheritDefaultsFromGraph(*ParentBuilder);
+					// Add this to the list of inheriting as the
+					// type no longer matches or it no longer exists
+					InputsInheritingDefault.Add(NodeName);
+				}
+				ClassInputs.Add(MoveTemp(NewClassInput));
+			}
 
-			// Fill new managed inputs set with names of all class inputs & if the old input was explicitly not
-			// marked as a managed input, then remove it from the new managed inputs if found.
 			OutInputsInheritingDefault.Reset();
 			Algo::TransformIf(ClassInputs, OutInputsInheritingDefault,
-				[&InputsToInherit](const FMetasoundFrontendClassInput& Input)
+				[&InputsInheritingDefault](const FMetasoundFrontendClassInput& Input)
 				{
-					return InputsToInherit.Contains(Input.Name);
+					return InputsInheritingDefault.Contains(Input.Name);
 				},
 				[](const FMetasoundFrontendClassInput& Input)
 				{
 					return Input.Name;
 				});
-
+			
 			return ClassInputs;
 		}
-
-		TArray<FMetasoundFrontendClassOutput> FRebuildPresetRootGraph::GenerateRequiredClassOutputs(FDocumentHandle& InDocumentHandle, const FConstGraphHandle& InPresetGraph) const
+		
+		TArray<FMetasoundFrontendClassOutput> FRebuildPresetRootGraph::GenerateRequiredClassOutputs(FMetaSoundFrontendDocumentBuilder& InDocumentToTransformBuilder) const
 		{
 			TArray<FMetasoundFrontendClassOutput> ClassOutputs;
+			check(ParentBuilder);
 
-			FConstGraphHandle ReferencedGraph = ReferencedDocument->GetRootGraph();
+			const FMetasoundFrontendDocument& ParentDocument = ParentBuilder->GetConstDocumentChecked();
+			const FMetasoundFrontendDocument& DocumentToTransform = InDocumentToTransformBuilder.GetConstDocumentChecked();
 
-			// Iterate over the referenced graph's output nodes.
-			ReferencedGraph->IterateConstNodes([&](FConstNodeHandle OutputNode)
+			for (const FMetasoundFrontendClassOutput& ParentClassOutput : ParentDocument.RootGraph.GetDefaultInterface().Outputs)
 			{
-				const FName NodeName = OutputNode->GetNodeName();
-				FConstOutputHandle Output = OutputNode->GetConstOutputWithVertexName(NodeName);
-				if (ensure(Output->IsValid()))
-				{
-					FMetasoundFrontendClassOutput ClassOutput;
-
-					ClassOutput.Name = NodeName;
-					ClassOutput.TypeName = Output->GetDataType();
-					ClassOutput.AccessType = Output->GetVertexAccessType();
+				// Copy class output and reset defaults and id
+				FMetasoundFrontendClassOutput NewClassOutput = ParentClassOutput;
+				NewClassOutput.VertexID = FDocumentIDGenerator::Get().CreateVertexID(DocumentToTransform);
 
-#if WITH_EDITOR
-					ClassOutput.Metadata.SetDescription(OutputNode->GetDescription());
-					ClassOutput.Metadata.SetDisplayName(Output->GetMetadata().GetDisplayName());
-#endif // WITH_EDITOR
-
-					FDocumentAccessPtr DocumentPtr = InDocumentHandle->GetDocumentPtr();
-					const FMetasoundFrontendDocument* Document = DocumentPtr.Get();
-					check(Document);
-					ClassOutput.VertexID = FDocumentIDGenerator::Get().CreateVertexID(*Document);
-
-					if (const FMetasoundFrontendClassOutput* ExistingClassOutput = InPresetGraph->FindClassOutputWithName(NodeName).Get())
-					{
-						ClassOutput.NodeID = ExistingClassOutput->NodeID;
-					}
-
-					ClassOutputs.Add(MoveTemp(ClassOutput));
+				if (const FMetasoundFrontendClassOutput* ExistingClassOutput = InDocumentToTransformBuilder.FindGraphOutput(ParentClassOutput.Name))
+				{
+					NewClassOutput.NodeID = ExistingClassOutput->NodeID;
 				}
-			}, EMetasoundFrontendClassType::Output);
-
+				ClassOutputs.Add(MoveTemp(NewClassOutput));
+			}
 			return ClassOutputs;
 		}
 
diff --git a/Engine/Plugins/Runtime/Metasound/Source/MetasoundFrontend/Public/MetasoundFrontendController.h b/Engine/Plugins/Runtime/Metasound/Source/MetasoundFrontend/Public/MetasoundFrontendController.h
--- a/Engine/Plugins/Runtime/Metasound/Source/MetasoundFrontend/Public/MetasoundFrontendController.h
+++ b/Engine/Plugins/Runtime/Metasound/Source/MetasoundFrontend/Public/MetasoundFrontendController.h
@@ -536,6 +536,7 @@ namespace Metasound
 			  * differences (ex. DisplayName only used in editor) but no runtime
 			  * behavior has been modified.
 			  */
+			UE_DEPRECATED(5.7, "Moved to FMetaSoundFrontendDocumentBuilder")
 			virtual bool CanAutoUpdate(FClassInterfaceUpdates& OutInterfaceUpdates) const = 0;
 #endif // WITH_EDITORONLY_DATA
 
diff --git a/Engine/Plugins/Runtime/Metasound/Source/MetasoundFrontend/Public/MetasoundFrontendDocument.h b/Engine/Plugins/Runtime/Metasound/Source/MetasoundFrontend/Public/MetasoundFrontendDocument.h
--- a/Engine/Plugins/Runtime/Metasound/Source/MetasoundFrontend/Public/MetasoundFrontendDocument.h
+++ b/Engine/Plugins/Runtime/Metasound/Source/MetasoundFrontend/Public/MetasoundFrontendDocument.h
@@ -4,6 +4,7 @@
 #include "Algo/Transform.h"
 #include "Containers/Array.h"
 #include "Containers/Map.h"
+#include "HAL/IConsoleManager.h"
 #include "IAudioParameterInterfaceRegistry.h"
 #include "Internationalization/Text.h"
 #include "MetasoundAccessPtr.h"
@@ -47,6 +48,8 @@ namespace Metasound
 
 	namespace Frontend
 	{
+		extern int32 MetaSoundAutoUpdateNativeClassesOfEqualVersionCVar;
+
 		constexpr FGuid DefaultPageID(0, 0, 0, 0);
 		constexpr TCHAR DefaultPageName[] = TEXT("Default");
 
@@ -1042,7 +1045,7 @@ public:
 	UE_API void IterateDefaults(TFunctionRef<void(const FGuid&, FMetasoundFrontendLiteral&)> IterFunc);
 	UE_API void IterateDefaults(TFunctionRef<void(const FGuid&, const FMetasoundFrontendLiteral&)> IterFunc) const;
 	UE_API bool RemoveDefault(const FGuid& InPageID);
-	UE_API void ResetDefaults();
+	UE_API void ResetDefaults(bool bInitializeDefaultPage = true);
 	UE_API void SetDefaults(TArray<FMetasoundFrontendClassInputDefault> InputDefaults);
 };
 
diff --git a/Engine/Plugins/Runtime/Metasound/Source/MetasoundFrontend/Public/MetasoundFrontendDocumentBuilder.h b/Engine/Plugins/Runtime/Metasound/Source/MetasoundFrontend/Public/MetasoundFrontendDocumentBuilder.h
--- a/Engine/Plugins/Runtime/Metasound/Source/MetasoundFrontend/Public/MetasoundFrontendDocumentBuilder.h
+++ b/Engine/Plugins/Runtime/Metasound/Source/MetasoundFrontend/Public/MetasoundFrontendDocumentBuilder.h
@@ -30,7 +30,9 @@ namespace Metasound::Frontend
 	// Forward Declarations
 	class INodeTemplate;
 
+	using FNodePredicateFunctionRef = TFunctionRef<bool(const FMetasoundFrontendNode&)>;
 	using FConstClassAndNodeFunctionRef = TFunctionRef<void(const FMetasoundFrontendClass&, const FMetasoundFrontendNode&)>;
+	using FConstClassAndNodeAndPageIDFunctionRef = TFunctionRef<void(const FMetasoundFrontendClass&, const FMetasoundFrontendNode&, const FGuid&)>;
 	using FFinalizeNodeFunctionRef = TFunctionRef<void(FMetasoundFrontendNode&, const Metasound::Frontend::FNodeRegistryKey&)>;
 
 	enum class EInvalidEdgeReason : uint8
@@ -197,9 +199,14 @@ public:
 	// If none provided, returns node pointer to node for the builder's currently set build page ID (see 'GetBuildPageID').
 	UE_API const FMetasoundFrontendNode* AddGraphOutput(FMetasoundFrontendClassOutput ClassOutput, const FGuid* InPageID = nullptr);
 
-	UE_API bool AddInterface(FName InterfaceName);
+	// Adds an interface and associated inputs and outputs
+	// If bAddUserModifiableInterfaceOnly is true, only user modifiable interfaces can be added (see FMetasoundFrontendInterfaceUClassOptions)
+	UE_API bool AddInterface(FName InterfaceName, bool bAddUserModifiableInterfaceOnly = true);
 
 	UE_API const FMetasoundFrontendNode* AddNodeByClassName(const FMetasoundFrontendClassName& InClassName, int32 InMajorVersion = 1, FGuid InNodeID = FGuid::NewGuid(), const FGuid* InPageID = nullptr);
+
+	// Add node with specific major and minor version. The public version of AddNodeByClassName adds the highest minor version for a given major version
+	UE_INTERNAL const FMetasoundFrontendNode* AddNodeByClassName(const FMetasoundFrontendClassName& InClassName, int32 InMajorVersion, int32 InMinorVersion, FGuid InNodeID = FGuid::NewGuid(), const FGuid* InPageID = nullptr);
 	
 	UE_API const FMetasoundFrontendNode* AddNodeByTemplate(const Metasound::Frontend::INodeTemplate& InTemplate, FNodeTemplateGenerateInterfaceParams Params, FGuid InNodeID = FGuid::NewGuid(), const FGuid* InPageID = nullptr);
 
@@ -218,6 +225,17 @@ public:
 	// is not already connected and the edge is valid (see function 'IsValidEdge').
 	UE_API bool CanAddEdge(const FMetasoundFrontendEdge& InEdge, const FGuid* InPageID = nullptr) const;
 
+#if WITH_EDITORONLY_DATA
+	// Returns whether the node is eligible for auto-updating (i.e.
+	// has undergone minor revision or the interface has changed, but
+	// no higher major revision is available). Provides interface updates
+	// to populate with any information regarding interface updates.
+	// Can return true if the interface has changed but with only cosmetic
+	// differences (ex. DisplayName only used in editor) but no runtime
+	// behavior has been modified.
+	UE_API bool CanAutoUpdate(const FGuid& InNodeID, Metasound::Frontend::FClassInterfaceUpdates& OutInterfaceUpdates, const FGuid* InPageID = nullptr);
+#endif // WITH_EDITORONLY_DATA
+
 	// Clears document completely of all graph page data (nodes, edges, & member metadata), dependencies,
 	// interfaces, member metadata, preset state, etc. Leaves ClassMetadata intact. Reloads the builder state,
 	// so external delegates must be relinked if desired.
@@ -235,15 +253,19 @@ public:
 	UE_API bool ContainsNode(const FGuid& InNodeID, const FGuid* InPageID = nullptr) const;
 
 	UE_API bool ConvertFromPreset();
+	UE_DEPRECATED(5.7, "Use the ConvertToPreset which takes in a builder")
 	UE_API bool ConvertToPreset(const FMetasoundFrontendDocument& InReferencedDocument, TSharedPtr<Metasound::Frontend::FDocumentModifyDelegates> ModifyDelegates = { });
+	UE_API bool ConvertToPreset(const FMetaSoundFrontendDocumentBuilder& InReferencedDocumentBuilder, TSharedPtr<Metasound::Frontend::FDocumentModifyDelegates> ModifyDelegates = { });
+	
 
 	static UE_API TOptional<Metasound::FAnyDataReference> CreateDataReference(const Metasound::FOperatorSettings& InOperatorSettings, FName DataType, const Metasound::FLiteral& InLiteral, Metasound::EDataReferenceAccessType AccessType);
 
+	
 #if WITH_EDITORONLY_DATA
 	// Fills out the provided ClassInterfaceUpdate struct with the differences between
 	// the registry's version of the class interface and that of the node.
 	// Returns whether or not the interface was found in the registry.
-	UE_API bool DiffAgainstRegistryInterface(const FGuid& InNodeID, const FGuid& InPageID, bool bInUseHighestMinorVersion, Metasound::Frontend::FClassInterfaceUpdates& OutInterfaceUpdates) const;
+	UE_API bool DiffAgainstRegistryInterface(const FGuid& InNodeID, const FGuid& InPageID, bool bInUseHighestMinorVersion, Metasound::Frontend::FClassInterfaceUpdates& OutInterfaceUpdates, bool bInForceRegenerateClassInterfaceOverride = true) const;
 #endif // WITH_EDITORONLY_DATA
 
 	UE_API const FMetasoundFrontendClassInput* DuplicateGraphInput(FName ExistingName, FName NewName);
@@ -256,6 +278,7 @@ public:
 	UE_DEPRECATED(5.6, "Use the duplicate overload that supplies existing name and new name and returns output")
 	UE_API const FMetasoundFrontendNode* DuplicateGraphOutput(const FMetasoundFrontendClassOutput& InClassOutput, const FName InName, const FGuid* InPageID = nullptr);
 
+	TConstStructView<FMetasoundFrontendClassInterface> FindClassInterfaceOverride(const FGuid& InNodeID, const FGuid* InPageID = nullptr) const;
 
 #if WITH_EDITORONLY_DATA
 	UE_API const FMetasoundFrontendEdgeStyle* FindConstEdgeStyle(const FGuid& InNodeID, FName OutputName, const FGuid* InPageID = nullptr) const;
@@ -272,6 +295,8 @@ public:
 
 	UE_API const FMetasoundFrontendClass* FindDependency(const FGuid& InClassID) const;
 	UE_API const FMetasoundFrontendClass* FindDependency(const FMetasoundFrontendClassMetadata& InMetadata) const;
+	UE_API const FMetasoundFrontendClass* FindDependency(const Metasound::Frontend::FNodeRegistryKey& InRegistryKey) const;
+
 	UE_API TArray<const FMetasoundFrontendEdge*> FindEdges(const FGuid& InNodeID, const FGuid& InVertexID, const FGuid* InPageID = nullptr) const;
 
 	UE_API const FMetasoundFrontendClassInput* FindGraphInput(FName InputName) const;
@@ -450,6 +475,9 @@ public:
 	// Iterates all nodes for the given page (uses build page ID if not provided).
 	UE_API void IterateNodes(Metasound::Frontend::FConstClassAndNodeFunctionRef Func, const FGuid* InPageID = nullptr) const;
 
+	// Iterates nodes, filtered by the given predicate. If bIterateAllPages is true, nodes on all pages will be considered and Page ID is ignored.
+	UE_API void IterateNodesByPredicate(Metasound::Frontend::FConstClassAndNodeAndPageIDFunctionRef Func, Metasound::Frontend::FNodePredicateFunctionRef PredicateFunc, const FGuid* InPageID = nullptr, bool bIterateAllPages = false) const;
+
 	// Iterates nodes that are filtered by only subscribing to a class with the given type (asserts if provided invalid class type).
 	UE_API void IterateNodesByClassType(Metasound::Frontend::FConstClassAndNodeFunctionRef Func, EMetasoundFrontendClassType ClassType, const FGuid* InPageID = nullptr) const;
 
@@ -491,6 +519,8 @@ public:
 
 	UE_API bool RemoveInterface(FName Name);
 	UE_API bool RemoveNamedEdges(const TSet<Metasound::Frontend::FNamedEdge>& InNamedEdgesToRemove, TArray<FMetasoundFrontendEdge>* OutRemovedEdges = nullptr, const FGuid* InPageID = nullptr);
+	
+	// Removes node and connected edges from document. Does not remove dependency. Use RemoveUnusedDependencies
 	UE_API bool RemoveNode(const FGuid& InNodeID, const FGuid* InPageID = nullptr);
 
 #if WITH_EDITORONLY_DATA
@@ -508,6 +538,12 @@ public:
 
 	UE_DEPRECATED(5.5, "Use GenerateNewClassName instead")
 	UE_API bool RenameRootGraphClass(const FMetasoundFrontendClassName& InName);
+	
+	using FVertexNameAndType = TTuple<FName, FName>;
+	// Replaces existing dependency with a new one. Removes all nodes using this class, and readds nodes with new class, 
+	// trying to maintain existing connections (based on name/data type) and data. 
+	// Optionally report disconnected inputs and outputs
+	UE_API bool ReplaceDependency(const Metasound::Frontend::FNodeClassRegistryKey& OldClass, const Metasound::Frontend::FNodeClassRegistryKey& NewClass, TArray<FVertexNameAndType>* OutDisconnectedInputs = nullptr, TArray<FVertexNameAndType>* OutDisconnectedOutputs = nullptr);
 
 #if WITH_EDITORONLY_DATA
 	UE_API bool ResetGraphInputDefault(FName InputName);
@@ -526,7 +562,7 @@ public:
 	// Returns true if the builder is already targeting the given ID or if it successfully
 	// found a page implementation with the given ID and was able to switch to it, false if not.
 	// Swapping the targeted build graph ID clears the local cache, so swapping frequently can
-	// induce cash thrashing. BroadcastDelegate should always be true unless dealing with the controller
+	// induce cache thrashing. BroadcastDelegate should always be true unless dealing with the controller
 	// API (exposed as a mechanism for mutating via controllers while deprecating.  Option will be removed
 	// in a future build).
 	UE_API bool SetBuildPageID(const FGuid& InBuildPageID, bool bBroadcastDelegate = true);
@@ -561,6 +597,7 @@ public:
 	// Optionally, if bForceUpdate is set, updates inheritance even if not a preset (primarily for clearing flag
 	// if non-preset has unnecessary data).
 	UE_API bool SetGraphInputInheritsDefault(FName InName, bool bInputInheritsDefault, bool bForceUpdate = false);
+	UE_API bool SetGraphInputsInheritingDefault(TSet<FName>&& InNames);
 
 	// Sets a given graph input's name to a new name. Succeeds if the graph output exists and the new name is set (or is the same as the old name).
 	UE_API bool SetGraphInputName(FName InputName, FName InName);
@@ -620,6 +657,8 @@ public:
 
 	UE_DEPRECATED(5.7, "Use AccessFlags API calls with deprecation flag instead")
 	UE_API void SetIsDeprecated(const bool bInIsDeprecated);
+	UE_API void SetInputStyle(FMetasoundFrontendInterfaceStyle&& Style);
+	UE_API void SetOutputStyle(FMetasoundFrontendInterfaceStyle&& Style);
 
 	UE_API void SetMemberMetadata(UMetaSoundFrontendMemberMetadata& NewMetadata);
 
@@ -644,6 +683,8 @@ public:
 	UE_API bool SetNodeConfiguration(const FGuid& InNodeID, TInstancedStruct<FMetaSoundFrontendNodeConfiguration> InNodeConfiguration, const FGuid* InPageID = nullptr);
 
 	UE_API bool SetNodeInputDefault(const FGuid& InNodeID, const FGuid& InVertexID, const FMetasoundFrontendLiteral& InLiteral, const FGuid* InPageID = nullptr);
+	
+	UE_INTERNAL void SetPresetFlags(bool bIsPreset=true);
 
 	// Sets the document's version number.  Should only be called by document versioning.
 	UE_API void SetVersionNumber(const FMetasoundFrontendVersionNumber& InDocumentVersionNumber);
@@ -652,6 +693,11 @@ public:
 	UE_API bool SwapGraphOutput(const FMetasoundFrontendClassVertex& InExistingOutputVertex, const FMetasoundFrontendClassVertex& NewOutputVertex);
 
 #if WITH_EDITOR
+	// Synchronizes all dependency Metadata in document with that found in the registry. If not found
+	// in the registry, no action taken. Returns false if any dependencies could not be checked for updates.
+	// If provided array, pointers to classes that were modified will be added to it.
+	UE_API bool SynchronizeDependencyMetadata(TArray<const FMetasoundFrontendClass*>* InOutModifiedClasses = nullptr);
+
 	UE_API bool UpdateDependencyRegistryData(const TMap<Metasound::Frontend::FNodeRegistryKey, Metasound::Frontend::FNodeRegistryKey>& OldToNewClassKeys);
 #endif // WITH_EDITOR
 
@@ -687,6 +733,48 @@ private:
 
 	FMetasoundFrontendNode* AddNodeInternal(const FMetasoundFrontendClassMetadata& InClassMetadata, Metasound::Frontend::FFinalizeNodeFunctionRef FinalizeNode, const FGuid& InPageID, FGuid InNodeID = FGuid::NewGuid(), int32* NewNodeIndex = nullptr);
 	FMetasoundFrontendNode* AddNodeInternal(const Metasound::Frontend::FNodeRegistryKey& InClassKey, FFinalizeNodeFunctionRef FinalizeNode, const FGuid& InPageID, FGuid InNodeID = FGuid::NewGuid(), int32* NewNodeIndex = nullptr);
+
+	struct FInputConnectionInfo
+	{
+		// Handle of connected output vertex
+		FMetasoundFrontendVertexHandle ConnectedOutput;
+		// Name and data type of node input vertex
+		FName Name;
+		FName DataType;
+		FMetasoundFrontendLiteral DefaultValue;
+		bool bLiteralSet = false;
+	};
+
+	struct FOutputConnectionInfo
+	{
+		// Handles of connected input vertices
+		TArray<FMetasoundFrontendVertexHandle> ConnectedInputs;
+		// Name and data type of node output vertex
+		FName Name;
+		FName DataType;
+	};
+
+	// Map of input name/type to input info
+	using FInputConnections = TMap<FVertexNameAndType, FInputConnectionInfo>;
+	// Map of output name/type to connected inputs
+	using FOutputConnections = TMap<FVertexNameAndType, FOutputConnectionInfo>;
+	struct FNodeInstanceReplacementData
+	{
+#if WITH_EDITOR
+		FMetasoundFrontendNodeStyle Style;
+#endif // WITH_EDITOR
+
+		TInstancedStruct<FMetaSoundFrontendNodeConfiguration> Configuration;
+		TInstancedStruct<FMetasoundFrontendClassInterface> ClassInterfaceOverride;
+		FInputConnections InputConnections;
+		FOutputConnections OutputConnections;
+		FGuid NodeID;
+		FGuid PageID;
+	};
+
+	void ApplyNodeInstanceReplacementData(FMetasoundFrontendNode& InReplacementNode, FNodeInstanceReplacementData&& InInstanceData, TArray<FVertexNameAndType>* OutDisconnectedInputs=nullptr, TArray<FVertexNameAndType>* OutDisconnectedOutputs=nullptr);
+	FNodeInstanceReplacementData CaptureNodeInstanceReplacementData(FMetasoundFrontendNode& InOriginalNode, const FGuid* InPageID = nullptr);
+
 	void BeginBuilding(TSharedPtr<Metasound::Frontend::FDocumentModifyDelegates> Delegates = {}, bool bPrimeCache = false);
 
 	// Conforms GraphOutput node's ClassID, Access & Data Type with the GraphOutput.
@@ -725,6 +813,19 @@ private:
 	bool SpliceVariableNodeFromStack(const FGuid& InNodeID, const FGuid& InPageID);
 	bool UnlinkVariableNode(const FGuid& InNodeID, const FGuid& InPageID);
 
+	struct FNodeConfigurationUpdateData
+	{
+		TConstStructView<FMetaSoundFrontendNodeConfiguration> ExistingConfig;
+		TConstStructView<FMetasoundFrontendClassInterface> ExistingClassInterfaceOverride;
+		TInstancedStruct<FMetaSoundFrontendNodeConfiguration> RegisteredConfig;
+		TInstancedStruct<FMetasoundFrontendClassInterface> RegeneratedClassInterfaceOverride;
+		bool bDidUpdateClassInterfaceOverride = false;
+	};
+
+	bool ShouldReplaceExistingNodeConfig(TConstStructView<FMetaSoundFrontendNodeConfiguration> InRegisteredNodeConfig, TConstStructView<FMetaSoundFrontendNodeConfiguration> InExistingConfig) const;
+	const FMetasoundFrontendClassInterface& GetApplicableRegistryInterface(const FMetasoundFrontendClass& InRegisteredClass, const FNodeConfigurationUpdateData& InNodeConfigurationUpdates) const;
+	void FindNodeConfigurationUpdates(const FGuid& InNodeID, const FGuid& InPageID, const FMetasoundFrontendClass& InRegisteredClass, FNodeConfigurationUpdateData& OutNodeConfigUpdates, bool bInForceRegenerateClassInterfaceOverride) const;
+
 	UPROPERTY(Transient)
 	TScriptInterface<IMetaSoundDocumentInterface> DocumentInterface;
 
diff --git a/Engine/Plugins/Runtime/Metasound/Source/MetasoundFrontend/Public/MetasoundFrontendTransform.h b/Engine/Plugins/Runtime/Metasound/Source/MetasoundFrontend/Public/MetasoundFrontendTransform.h
--- a/Engine/Plugins/Runtime/Metasound/Source/MetasoundFrontend/Public/MetasoundFrontendTransform.h
+++ b/Engine/Plugins/Runtime/Metasound/Source/MetasoundFrontend/Public/MetasoundFrontendTransform.h
@@ -55,9 +55,6 @@ namespace Metasound::Frontend
 	public:
 		virtual ~IGraphTransform() = default;
 
-		UE_DEPRECATED(5.3, "Deprecated: unused function which encouraged breaking const behavior of transform state.")
-		virtual FMetasoundFrontendDocument& GetOwningDocument() const { static FMetasoundFrontendDocument BaseDoc; return BaseDoc; };
-
 		/** Return true if the graph was modified, false otherwise. */
 		virtual bool Transform(FMetasoundFrontendGraph& InOutGraph) const = 0;
 	};
@@ -70,9 +67,6 @@ namespace Metasound::Frontend
 
 		/** Return true if the node was modified, false otherwise. */
 		UE_API virtual bool Transform(const FGuid& InNodeID, FMetaSoundFrontendDocumentBuilder& OutBuilder) const;
-
-		UE_DEPRECATED(5.4, "Use transform overload with node ID and builder parameters to apply node transform to builder's underlying document")
-		UE_API virtual bool Transform(FMetasoundFrontendNode& InOutNode) const;
 	};
 
 	/** Adds or swaps document members (inputs, outputs) and removing any document members where necessary and adding those missing. */
@@ -146,47 +140,59 @@ namespace Metasound::Frontend
 		virtual bool Transform(FDocumentHandle InDocument) const override { return false; }
 	};
 
-	/** Completely rebuilds the graph connecting a preset's inputs to the reference
+	/** Completely rebuilds the graph connecting a preset's inputs to the reference (parent)
 		* document's root graph. It maintains previously set input values entered upon 
-		* the presets wrapping graph. */
+		* the preset's wrapping graph. */
 	class FRebuildPresetRootGraph : public IDocumentTransform
 	{
 	public:
 		/** Create transform.
 			* @param InReferenceDocument - The document containing the wrapped MetaSound graph.
 			*/
+		UE_DEPRECATED(5.7, "Use the constructor which takes a FMetaSoundFrontendDocumentBuilder")
 		FRebuildPresetRootGraph(FConstDocumentHandle InReferencedDocument)
-			: ReferencedDocument(InReferencedDocument)
+			: ParentBuilder(nullptr)
+		{
+		}
+
+		FRebuildPresetRootGraph(const FMetaSoundFrontendDocumentBuilder& InParentDocumentBuilder)
+			: ParentBuilder(&InParentDocumentBuilder)
 		{
 		}
 
+		UE_DEPRECATED(5.7, "Use Transform which takes in a document builder.")
 		UE_API FRebuildPresetRootGraph(const FMetasoundFrontendDocument& InReferencedDocument);
 
+		UE_DEPRECATED(5.7, "Use Transform which takes in a document builder.")
 		UE_API virtual bool Transform(FDocumentHandle InDocument) const override;
+
+		UE_DEPRECATED(5.7, "Use Transform which takes in a document builder.")
 		UE_API virtual bool Transform(FMetasoundFrontendDocument& InOutDocument) const override;
 
+		UE_API bool Transform(FMetaSoundFrontendDocumentBuilder& InOutBuilderToTransform) const;
+
 	private:
 
 		// Get the class inputs needed for this preset. Input literals set on 
 		// the preset graph will be used if they are set and are marked as inheriting
 		// the default from the referenced graph.
-		UE_API TArray<FMetasoundFrontendClassInput> GenerateRequiredClassInputs(FDocumentHandle& InDocumentHandle, const FConstGraphHandle& InParentGraph, TSet<FName>& OutInputsInheritingDefault) const;
+		TArray<FMetasoundFrontendClassInput> GenerateRequiredClassInputs(FMetaSoundFrontendDocumentBuilder& InDocumentToTransformBuilder, TSet<FName>& OutInputsInheritingDefault) const;
 
 		// Get the class Outputs needed for this preset.
-		UE_API TArray<FMetasoundFrontendClassOutput> GenerateRequiredClassOutputs(FDocumentHandle& InDocumentHandle, const FConstGraphHandle& InParentGraph) const;
-
+		TArray<FMetasoundFrontendClassOutput> GenerateRequiredClassOutputs(FMetaSoundFrontendDocumentBuilder& InDocumentToTransformBuilder) const;
+		
 		// Add inputs to parent graph and connect to wrapped graph node.
-		UE_API void AddAndConnectInputs(const TArray<FMetasoundFrontendClassInput>& InClassInputs, FGraphHandle& InParentGraphHandle, FNodeHandle& InReferencedNode) const;
+		void AddAndConnectInputs(const TArray<FMetasoundFrontendClassInput>& InClassInputs, FMetaSoundFrontendDocumentBuilder& InOutBuilderToTransform, const FGuid& InReferencedNodeID) const;
 
 		// Add outputs to parent graph and connect to wrapped graph node.
-		UE_API void AddAndConnectOutputs(const TArray<FMetasoundFrontendClassOutput>& InClassOutputs, FGraphHandle& InParentGraphHandle, FNodeHandle& InReferencedNode) const;
+		void AddAndConnectOutputs(const TArray<FMetasoundFrontendClassOutput>& InClassOutputs, FMetaSoundFrontendDocumentBuilder& InOutBuilderToTransform, const FGuid& InReferencedNodeID) const;
 
 #if WITH_EDITORONLY_DATA
 		using FMemberIDToMetadataMap = TMap<FGuid, TObjectPtr<UMetaSoundFrontendMemberMetadata>>;
-		UE_API void AddMemberMetadata(const FMemberIDToMetadataMap& InCachedMemberMetadata, FGraphHandle& InPresetGraphHandle, FMemberIDToMetadataMap& InOutMemberMetadata) const;
+		void AddMemberMetadata(const FMemberIDToMetadataMap& InCachedMemberMetadata, FMetaSoundFrontendDocumentBuilder& InOutBuilderToTransform) const;
 #endif // WITH_EDITORONLY_DATA
 
-		FConstDocumentHandle ReferencedDocument = IDocumentController::GetInvalidHandle();
+		const FMetaSoundFrontendDocumentBuilder* ParentBuilder = nullptr;
 	};
 
 #if WITH_EDITORONLY_DATA
@@ -207,7 +213,10 @@ namespace Metasound::Frontend
 		{
 		}
 
+		UE_DEPRECATED(5.7, "Use Transform which takes in a document builder.")
 		UE_API bool Transform(FDocumentHandle InDocument);
+		
+		UE_API bool Transform(FMetaSoundFrontendDocumentBuilder& InOutBuilderToTransform);
 
 	private:
 		// Keeps track of classes already updated so node check can be avoided.
