Build fix for MetaSound autoupdate edge connection

diff --git a/Engine/Plugins/Runtime/Metasound/Source/MetasoundFrontend/Private/MetasoundFrontendDocumentBuilder.cpp b/Engine/Plugins/Runtime/Metasound/Source/MetasoundFrontend/Private/MetasoundFrontendDocumentBuilder.cpp
--- a/Engine/Plugins/Runtime/Metasound/Source/MetasoundFrontend/Private/MetasoundFrontendDocumentBuilder.cpp
+++ b/Engine/Plugins/Runtime/Metasound/Source/MetasoundFrontend/Private/MetasoundFrontendDocumentBuilder.cpp
@@ -4456,7 +4456,10 @@ void FMetaSoundFrontendDocumentBuilder::IterateNodesByPredicate(Metasound::Front
 		{
 			if (const FMetasoundFrontendClass* Class = FindDependency(Node.ClassID))
 			{
-				Func(*Class, Node, Graph.PageID);
+				if (PredicateFunc(Node))
+				{
+					Func(*Class, Node, Graph.PageID);
+				}
 			}
 		}
 	}
@@ -5395,7 +5398,6 @@ bool FMetaSoundFrontendDocumentBuilder::RemoveUnusedDependencies()
 	const FMetasoundFrontendGraphClass& RootGraph = Document.RootGraph;
 	const TArray<FMetasoundFrontendClass>& Dependencies = Document.Dependencies;
 
-	// Remove duplicate entries, keeping the highest version dependencies. 
 	bool bDidEdit = false;
 
 	// Remove unused dependencies
@@ -5580,7 +5582,9 @@ void FMetaSoundFrontendDocumentBuilder::ApplyNodeInstanceReplacementData(FMetaso
 					};
 
 					const bool bIsValidEdge = IsValidEdge(NewEdge, PageID) == Metasound::Frontend::EInvalidEdgeReason::None;
-					if (bIsValidEdge)
+					// It is possible the edge already exists if the edge is between 
+					// this and another node of the same dependency that just added it when it was replaced 
+					if (bIsValidEdge && !ContainsEdge(NewEdge, PageID))
 					{
 						AddEdge(MoveTemp(NewEdge), PageID);
 					}
