[MetaSounds][WaveAsset] - Create new asset actions are displayed for unsupported assets

diff --git a/Engine/Plugins/Runtime/Metasound/Source/MetasoundEditor/Private/MetasoundNodeDetailCustomization.cpp b/Engine/Plugins/Runtime/Metasound/Source/MetasoundEditor/Private/MetasoundNodeDetailCustomization.cpp
--- a/Engine/Plugins/Runtime/Metasound/Source/MetasoundEditor/Private/MetasoundNodeDetailCustomization.cpp
+++ b/Engine/Plugins/Runtime/Metasound/Source/MetasoundEditor/Private/MetasoundNodeDetailCustomization.cpp
@@ -22,6 +22,7 @@
 #include "MetasoundFrontendRegistries.h"
 #include "MetasoundEditorSettings.h"
 #include "MetasoundNodeConfigurationCustomization.h"
+#include "MetasoundWave.h"
 #include "Modules/ModuleManager.h"
 #include "PropertyCustomizationHelpers.h"
 #include "PropertyEditorDelegates.h"
@@ -738,8 +739,11 @@ PRAGMA_ENABLE_DEPRECATION_WARNINGS
 				return FString();
 			};
 
+			const bool bAllowCreate = DataTypeInfo.DataTypeName != GetMetasoundDataTypeName<FWaveAsset>();
+
 			return SNew(SObjectPropertyEntryBox)
 				.AllowClear(true)
+				.AllowCreate(bAllowCreate)
 				.AllowedClass(DataTypeInfo.ProxyGeneratorClass)
 				.DisplayBrowse(true)
 				.DisplayThumbnail(true)
