Fix MetaSound autoupdate dependency update missing page ID

diff --git a/Engine/Plugins/Runtime/Metasound/Source/MetasoundFrontend/Private/MetasoundFrontendTransform.cpp b/Engine/Plugins/Runtime/Metasound/Source/MetasoundFrontend/Private/MetasoundFrontendTransform.cpp
--- a/Engine/Plugins/Runtime/Metasound/Source/MetasoundFrontend/Private/MetasoundFrontendTransform.cpp
+++ b/Engine/Plugins/Runtime/Metasound/Source/MetasoundFrontend/Private/MetasoundFrontendTransform.cpp
@@ -537,7 +537,7 @@ namespace Metasound
 
 					FClassInterfaceUpdates InterfaceUpdates;
 					const bool bHasUpdated = UpdatedClasses.Contains(Class.ID);
-					const bool bCanAutoUpdate = InOutBuilderToTransform.CanAutoUpdate(Node.GetID(), InterfaceUpdates);
+					const bool bCanAutoUpdate = InOutBuilderToTransform.CanAutoUpdate(Node.GetID(), InterfaceUpdates, &PageID);
 					if (bHasUpdated || !bCanAutoUpdate)
 					{
 						return;
