Fix crash in metasound array random get nodes
#rnx

diff --git a/Engine/Plugins/Runtime/Metasound/Source/MetasoundFrontend/Public/MetasoundArrayRandomNode.h b/Engine/Plugins/Runtime/Metasound/Source/MetasoundFrontend/Public/MetasoundArrayRandomNode.h
--- a/Engine/Plugins/Runtime/Metasound/Source/MetasoundFrontend/Public/MetasoundArrayRandomNode.h
+++ b/Engine/Plugins/Runtime/Metasound/Source/MetasoundFrontend/Public/MetasoundArrayRandomNode.h
@@ -318,7 +318,14 @@ namespace Metasound
 				DebugSharedStateBehaviorString = EnumName.GetValue().ToString();
 			}
 #endif // WITH_METASOUND_DEBUG_ENVIRONMENT
-			bIsPreviewSound = InParams.Environment.GetValue<bool>(SourceInterface::Environment::IsPreview);
+	   		if (InParams.Environment.Contains<bool>(SourceInterface::Environment::IsPreview))
+			{
+				bIsPreviewSound = InParams.Environment.GetValue<bool>(SourceInterface::Environment::IsPreview);
+			}
+			else
+			{
+				bIsPreviewSound = false;
+			}
 			
 			*OutValue = TDataTypeFactory<ElementType>::CreateAny(InParams.OperatorSettings);
 			*OutIndex = INDEX_NONE;
