Metasound code size fixes.

FCreateOperatorFunction -> raw function ptr instead of TFunction saves 154k of code on one target
(TFunction instance creation involves a lot of boilerplate per node)

Factoring out some common code into a dedicated FOutputNode ctor saves another 77k.

diff --git a/Engine/Plugins/Runtime/Metasound/Source/MetasoundFrontend/Private/MetasoundOutputNode.cpp b/Engine/Plugins/Runtime/Metasound/Source/MetasoundFrontend/Private/MetasoundOutputNode.cpp
--- a/Engine/Plugins/Runtime/Metasound/Source/MetasoundFrontend/Private/MetasoundOutputNode.cpp
+++ b/Engine/Plugins/Runtime/Metasound/Source/MetasoundFrontend/Private/MetasoundOutputNode.cpp
@@ -150,6 +150,13 @@ namespace Metasound
 	{
 	}
 
+	FOutputNode::FOutputNode(const FVertexName& InInstanceName, const FGuid& InInstanceID, const FVertexName& InVertexName, const FName& InDataTypeName, EVertexAccessType InVertexAccessType)
+	: FOutputNode(InVertexName,
+		FNodeData{InInstanceName, InInstanceID, CreateVertexInterface(InVertexName, InDataTypeName, InVertexAccessType)},
+		MakeShared<FNodeClassMetadata>(CreateNodeClassMetadata(InVertexName, InDataTypeName, InVertexAccessType)))
+	{
+	}
+
 	TSharedRef<IOperatorFactory, ESPMode::ThreadSafe> FOutputNode::GetDefaultOperatorFactory() const
 	{
 		return Factory;
diff --git a/Engine/Plugins/Runtime/Metasound/Source/MetasoundFrontend/Public/MetasoundOutputNode.h b/Engine/Plugins/Runtime/Metasound/Source/MetasoundFrontend/Public/MetasoundOutputNode.h
--- a/Engine/Plugins/Runtime/Metasound/Source/MetasoundFrontend/Public/MetasoundOutputNode.h
+++ b/Engine/Plugins/Runtime/Metasound/Source/MetasoundFrontend/Public/MetasoundOutputNode.h
@@ -64,6 +64,9 @@ namespace Metasound
 
 		UE_API virtual TSharedRef<IOperatorFactory, ESPMode::ThreadSafe> GetDefaultOperatorFactory() const override;
 
+	protected:
+		UE_API FOutputNode(const FVertexName& InInstanceName, const FGuid& InInstanceID, const FVertexName& InVertexName, const FName& InDataTypeName, EVertexAccessType InVertexAccessType);
+
 	private:
 		TSharedRef<OutputNodePrivate::FOutputOperatorFactory, ESPMode::ThreadSafe> Factory;
 	};
@@ -74,7 +77,7 @@ namespace Metasound
 	{
 	public:
 		TOutputNode(const FVertexName& InInstanceName, const FGuid& InInstanceID, const FVertexName& InVertexName)
-		:	TOutputNode(InVertexName, FNodeData{InInstanceName, InInstanceID, FOutputNode::CreateVertexInterface(InVertexName, GetMetasoundDataTypeName<DataType>(), VertexAccess)}, MakeShared<FNodeClassMetadata>(FOutputNode::CreateNodeClassMetadata(InVertexName, GetMetasoundDataTypeName<DataType>(), VertexAccess)))
+		:	FOutputNode(InVertexName, InInstanceID, InVertexName, GetMetasoundDataTypeName<DataType>(), VertexAccess)
 		{
 		}
 
diff --git a/Engine/Plugins/Runtime/Metasound/Source/MetasoundGraphCore/Public/MetasoundFacade.h b/Engine/Plugins/Runtime/Metasound/Source/MetasoundGraphCore/Public/MetasoundFacade.h
--- a/Engine/Plugins/Runtime/Metasound/Source/MetasoundGraphCore/Public/MetasoundFacade.h
+++ b/Engine/Plugins/Runtime/Metasound/Source/MetasoundGraphCore/Public/MetasoundFacade.h
@@ -94,7 +94,7 @@ namespace Metasound
 		// the signature of the CreateOperator function.
 		class FFactory : public IOperatorFactory
 		{
-			using FCreateOperatorFunction = TFunction<TUniquePtr<IOperator>(const FBuildOperatorParams& InParams, FBuildResults& OutResults)>;
+			using FCreateOperatorFunction = TUniquePtr<IOperator>(*)(const FBuildOperatorParams& InParams, FBuildResults& OutResults);
 			friend class FNodeFacade;
 
 			public:
