Fix MetaSound inputs not updating in editor preview

#rnx

diff --git a/Engine/Plugins/Runtime/Metasound/Source/MetasoundEngine/Private/MetasoundBuilderBase.cpp b/Engine/Plugins/Runtime/Metasound/Source/MetasoundEngine/Private/MetasoundBuilderBase.cpp
--- a/Engine/Plugins/Runtime/Metasound/Source/MetasoundEngine/Private/MetasoundBuilderBase.cpp
+++ b/Engine/Plugins/Runtime/Metasound/Source/MetasoundEngine/Private/MetasoundBuilderBase.cpp
@@ -1026,6 +1026,22 @@ void UMetaSoundBuilderBase::Initialize()
 	InitDelegates(*DocumentDelegates);
 }
 
+void UMetaSoundBuilderBase::Initialize(TScriptInterface<IMetaSoundDocumentInterface>& DocInterface)
+{
+	using namespace Metasound::Frontend;
+	
+	TSharedRef<FDocumentModifyDelegates> DocumentDelegates = MakeShared<FDocumentModifyDelegates>(DocInterface->GetConstDocument());
+	Builder = FMetaSoundFrontendDocumentBuilder(DocInterface, DocumentDelegates);
+	const FMetasoundFrontendDocument& Document = DocInterface->GetConstDocument();
+	const FMetasoundFrontendClassName& DocClassName = Document.RootGraph.Metadata.GetClassName();
+	if (!DocClassName.IsValid())
+	{
+		Builder.InitDocument();
+	}
+
+	InitDelegates(*DocumentDelegates);
+}
+
 void UMetaSoundBuilderBase::InitNodeLocations()
 {
 	Builder.InitNodeLocations();
diff --git a/Engine/Plugins/Runtime/Metasound/Source/MetasoundEngine/Private/MetasoundBuilderSubsystem.cpp b/Engine/Plugins/Runtime/Metasound/Source/MetasoundEngine/Private/MetasoundBuilderSubsystem.cpp
--- a/Engine/Plugins/Runtime/Metasound/Source/MetasoundEngine/Private/MetasoundBuilderSubsystem.cpp
+++ b/Engine/Plugins/Runtime/Metasound/Source/MetasoundEngine/Private/MetasoundBuilderSubsystem.cpp
@@ -506,19 +506,21 @@ void UMetaSoundSourceBuilder::OnEdgeAdded(int32 EdgeIndex) const
 void UMetaSoundSourceBuilder::OnInputAdded(int32 InputIndex)
 {
 	using namespace Metasound::DynamicGraph;
+	using namespace Metasound::Frontend;
+
+	const FMetasoundFrontendDocument& Doc = Builder.GetConstDocumentChecked();
+	const FMetasoundFrontendGraphClass& GraphClass = Doc.RootGraph;
+	const FMetasoundFrontendClassInput& NewInput = GraphClass.GetDefaultInterface().Inputs[InputIndex];
+
+	constexpr bool bCreateUObjectProxies = true;
+	UMetaSoundSource& Source = GetMetaSoundSource();
+	Source.RuntimeInputData.InputMap.Add(NewInput.Name, UMetaSoundSource::CreateRuntimeInput(IDataTypeRegistry::Get(), NewInput, bCreateUObjectProxies, UMetaSoundSettings::GetPageOrder()));
 
-	ExecuteAuditionableTransaction([this, InputIndex](FDynamicOperatorTransactor& Transactor)
+	ExecuteAuditionableTransaction([this, InputIndex, &NewInput](FDynamicOperatorTransactor& Transactor)
 	{
 		using namespace Metasound;
 		using namespace Metasound::Frontend;
 
-		const FMetasoundFrontendDocument& Doc = Builder.GetConstDocumentChecked();
-		const FMetasoundFrontendGraphClass& GraphClass = Doc.RootGraph;
-		const FMetasoundFrontendClassInput& NewInput = GraphClass.GetDefaultInterface().Inputs[InputIndex];
-
-		constexpr bool bCreateUObjectProxies = true;
-		UMetaSoundSource& Source = GetMetaSoundSource();
-		Source.RuntimeInputData.InputMap.Add(NewInput.Name, UMetaSoundSource::CreateRuntimeInput(IDataTypeRegistry::Get(), NewInput, bCreateUObjectProxies, UMetaSoundSettings::GetPageOrder()));
 
 		for (uint64 AudioComponentID : LiveComponentIDs)
 		{
@@ -846,18 +848,18 @@ void UMetaSoundSourceBuilder::OnRemovingInput(int32 InputIndex)
 {
 	using namespace Metasound::DynamicGraph;
 
-	ExecuteAuditionableTransaction([this, InputIndex](FDynamicOperatorTransactor& Transactor)
+	const FMetasoundFrontendDocument& Doc = Builder.GetConstDocumentChecked();
+	const FMetasoundFrontendGraphClass& GraphClass = Doc.RootGraph;
+	const FMetasoundFrontendClassInput& InputBeingRemoved = GraphClass.GetDefaultInterface().Inputs[InputIndex];
+
+	UMetaSoundSource& Source = GetMetaSoundSource();
+	Source.RuntimeInputData.InputMap.Remove(InputBeingRemoved.Name);
+
+	ExecuteAuditionableTransaction([this, InputIndex, &InputBeingRemoved](FDynamicOperatorTransactor& Transactor)
 	{
 		using namespace Metasound;
 		using namespace Metasound::Frontend;
 
-		const FMetasoundFrontendDocument& Doc = Builder.GetConstDocumentChecked();
-		const FMetasoundFrontendGraphClass& GraphClass = Doc.RootGraph;
-		const FMetasoundFrontendClassInput& InputBeingRemoved = GraphClass.GetDefaultInterface().Inputs[InputIndex];
-
-		UMetaSoundSource& Source = GetMetaSoundSource();
-		Source.RuntimeInputData.InputMap.Remove(InputBeingRemoved.Name);
-
 		Transactor.RemoveInputDataDestination(InputBeingRemoved.Name);
 
 		for (uint64 AudioComponentID : LiveComponentIDs)
diff --git a/Engine/Plugins/Runtime/Metasound/Source/MetasoundEngine/Public/MetasoundBuilderBase.h b/Engine/Plugins/Runtime/Metasound/Source/MetasoundEngine/Public/MetasoundBuilderBase.h
--- a/Engine/Plugins/Runtime/Metasound/Source/MetasoundEngine/Public/MetasoundBuilderBase.h
+++ b/Engine/Plugins/Runtime/Metasound/Source/MetasoundEngine/Public/MetasoundBuilderBase.h
@@ -574,6 +574,9 @@ public:
 
 	// Resets FrontendBuilder instance, creating a transient MetaSound document that is managed by this UObject Builder.
 	UE_API void Initialize();
+	// Initializes builder given an existing document interface 
+	UE_API void Initialize(TScriptInterface<IMetaSoundDocumentInterface>& DocInterface);
+
 
 protected:
 	virtual void BuildAndOverwriteMetaSoundInternal(TScriptInterface<IMetaSoundDocumentInterface> ExistingMetaSound, bool bForceUniqueClassName) const PURE_VIRTUAL(UMetaSoundBuilderBase::BuildAndOverwriteMetaSoundInternal, ;);
diff --git a/Engine/Plugins/Runtime/Metasound/Source/MetasoundEngine/Public/MetasoundDocumentBuilderRegistry.h b/Engine/Plugins/Runtime/Metasound/Source/MetasoundEngine/Public/MetasoundDocumentBuilderRegistry.h
--- a/Engine/Plugins/Runtime/Metasound/Source/MetasoundEngine/Public/MetasoundDocumentBuilderRegistry.h
+++ b/Engine/Plugins/Runtime/Metasound/Source/MetasoundEngine/Public/MetasoundDocumentBuilderRegistry.h
@@ -89,14 +89,9 @@ namespace Metasound::Engine
 			TObjectPtr<UMetaSoundBuilderBase> NewBuilder = 
 				CastChecked<UMetaSoundBuilderBase>(NewObject<UObject>(GetTransientPackage(), &BuilderUClass, NewName, RF_Transactional | RF_Transient));
 
-			FMetaSoundFrontendDocumentBuilder& BuilderRef = NewBuilder->GetBuilder();
-			BuilderRef = FMetaSoundFrontendDocumentBuilder(DocInterface);
+			NewBuilder->Initialize(DocInterface);
 			const FMetasoundFrontendDocument& Document = DocInterface->GetConstDocument();
 			const FMetasoundFrontendClassName& ClassName = Document.RootGraph.Metadata.GetClassName();
-			if (!ClassName.IsValid())
-			{
-				BuilderRef.InitDocument();
-			}
 
 			checkf(ClassName.IsValid(), TEXT("Document initialization must result in a valid class name being generated"));
 			AddBuilderInternal(ClassName, NewBuilder);
