Fix editor gc crash with MetaSound object literal defaults

diff --git a/Engine/Plugins/Runtime/Metasound/Source/MetasoundEditor/Private/MetasoundEditor.cpp b/Engine/Plugins/Runtime/Metasound/Source/MetasoundEditor/Private/MetasoundEditor.cpp
--- a/Engine/Plugins/Runtime/Metasound/Source/MetasoundEditor/Private/MetasoundEditor.cpp
+++ b/Engine/Plugins/Runtime/Metasound/Source/MetasoundEditor/Private/MetasoundEditor.cpp
@@ -5229,7 +5229,10 @@ PRAGMA_ENABLE_DEPRECATION_WARNINGS
 				const FName InputName = GraphInput.Name;
 				UMetasoundEditorGraph& MetaSoundGraph = ParentPtr->GetMetaSoundGraphChecked();
 				UMetasoundEditorGraphInput* EditorGraphInput = MetaSoundGraph.FindInput(InputName);
-				if (EditorGraphInput)
+				
+				// Literal does not need to be updated if the data type doesn't match,
+				// as it should be getting overwritten by the new literal type after this delegate call
+				if (EditorGraphInput && GraphInput.TypeName == EditorGraphInput->GetDataType())
 				{
 					UMetasoundEditorGraphMemberDefaultLiteral* Literal = EditorGraphInput->GetLiteral();
 					if (Literal)
diff --git a/Engine/Plugins/Runtime/Metasound/Source/MetasoundEditor/Private/MetasoundEditorGraphMemberDefaults.cpp b/Engine/Plugins/Runtime/Metasound/Source/MetasoundEditor/Private/MetasoundEditorGraphMemberDefaults.cpp
--- a/Engine/Plugins/Runtime/Metasound/Source/MetasoundEditor/Private/MetasoundEditorGraphMemberDefaults.cpp
+++ b/Engine/Plugins/Runtime/Metasound/Source/MetasoundEditor/Private/MetasoundEditorGraphMemberDefaults.cpp
@@ -288,7 +288,7 @@ namespace Metasound::Editor
 								{
 									bModified = true;
 
-									LiteralType NewValue;
+									LiteralType NewValue {};
 									if (!TestLiteral.TryGet(NewValue))
 									{
 										UE_LOG(LogMetasoundEditor, Verbose, TEXT("Synchronizing Page Default: Setting member '%s' (type '%s') to literal value '%s'. Type has changed or literal could not be set and will be set to type's default constructed value."),
