Improve metasound node registration association with modules
Allow unregistration of C++ nodes
Fixup deprecated code paths related to new registration mechanism

diff --git a/Engine/Plugins/Experimental/MetasoundExperimental/Source/MetasoundExperimentalEngineRuntime/MetasoundExperimentalEngineRuntime.Build.cs b/Engine/Plugins/Experimental/MetasoundExperimental/Source/MetasoundExperimentalEngineRuntime/MetasoundExperimentalEngineRuntime.Build.cs
--- a/Engine/Plugins/Experimental/MetasoundExperimental/Source/MetasoundExperimentalEngineRuntime/MetasoundExperimentalEngineRuntime.Build.cs
+++ b/Engine/Plugins/Experimental/MetasoundExperimental/Source/MetasoundExperimentalEngineRuntime/MetasoundExperimentalEngineRuntime.Build.cs
@@ -37,5 +37,13 @@ public class MetasoundExperimentalEngineRuntime : ModuleRules
 				"SignalProcessing"
 			}
 		);
+
+		PrivateDefinitions.AddRange(
+			new string[]
+			{
+				"METASOUND_PLUGIN=MetasoundExperimental",
+				"METASOUND_MODULE=MetasoundExperimentalEngineRuntime"
+			}
+		);
 	}
 }
diff --git a/Engine/Plugins/Experimental/MetasoundExperimental/Source/MetasoundExperimentalEngineRuntime/Private/MetasoundExperimentalEngineRuntimeModule.cpp b/Engine/Plugins/Experimental/MetasoundExperimental/Source/MetasoundExperimentalEngineRuntime/Private/MetasoundExperimentalEngineRuntimeModule.cpp
--- a/Engine/Plugins/Experimental/MetasoundExperimental/Source/MetasoundExperimentalEngineRuntime/Private/MetasoundExperimentalEngineRuntimeModule.cpp
+++ b/Engine/Plugins/Experimental/MetasoundExperimental/Source/MetasoundExperimentalEngineRuntime/Private/MetasoundExperimentalEngineRuntimeModule.cpp
@@ -2,21 +2,20 @@
 
 #include "MetasoundExperimentalEngineRuntimeModule.h"
 
-#include "MetasoundDataTypeRegistrationMacro.h"
-#include "Misc/Paths.h"
+#include "MetasoundFrontendModuleRegistrationMacros.h"
 #include "Modules/ModuleManager.h"
 
 void FMetasoundExperimentalEngineRuntimeModule::StartupModule()
 {
 	FModuleManager::Get().LoadModule("AudioExperimentalRuntime");
-	FMetasoundFrontendRegistryContainer::Get()->RegisterPendingNodes();
+	METASOUND_REGISTER_ITEMS_IN_MODULE
 }
 
 void FMetasoundExperimentalEngineRuntimeModule::ShutdownModule()
 {
+	METASOUND_UNREGISTER_ITEMS_IN_MODULE
 }
+
+METASOUND_IMPLEMENT_MODULE_REGISTRATION_LIST
 IMPLEMENT_MODULE(FMetasoundExperimentalEngineRuntimeModule, MetasoundExperimentalEngineRuntime);
 
-//#if WITH_ENGINE
-//#error ("Being linked with Engine")
-//#endif //
diff --git a/Engine/Plugins/Experimental/MetasoundExperimental/Source/MetasoundExperimentalRuntime/MetasoundExperimentalRuntime.Build.cs b/Engine/Plugins/Experimental/MetasoundExperimental/Source/MetasoundExperimentalRuntime/MetasoundExperimentalRuntime.Build.cs
--- a/Engine/Plugins/Experimental/MetasoundExperimental/Source/MetasoundExperimentalRuntime/MetasoundExperimentalRuntime.Build.cs
+++ b/Engine/Plugins/Experimental/MetasoundExperimental/Source/MetasoundExperimentalRuntime/MetasoundExperimentalRuntime.Build.cs
@@ -42,5 +42,12 @@ public class MetasoundExperimentalRuntime : ModuleRules
 				"SignalProcessing"
 			}
 		);
+
+		PrivateDefinitions.AddRange(
+			new string[]
+			{
+				"METASOUND_PLUGIN=MetasoundExperimental",
+				"METASOUND_MODULE=MetasoundExperimentalRuntime"
+			});
 	}
 }
diff --git a/Engine/Plugins/Experimental/MetasoundExperimental/Source/MetasoundExperimentalRuntime/Private/MetasoundExperimentalRuntimeModule.cpp b/Engine/Plugins/Experimental/MetasoundExperimental/Source/MetasoundExperimentalRuntime/Private/MetasoundExperimentalRuntimeModule.cpp
--- a/Engine/Plugins/Experimental/MetasoundExperimental/Source/MetasoundExperimentalRuntime/Private/MetasoundExperimentalRuntimeModule.cpp
+++ b/Engine/Plugins/Experimental/MetasoundExperimental/Source/MetasoundExperimentalRuntime/Private/MetasoundExperimentalRuntimeModule.cpp
@@ -2,26 +2,20 @@
 
 #include "MetasoundExperimentalRuntimeModule.h"
 
-#include "MetasoundCatMixerNode.h"
-#include "MetasoundChannelAgnosticType.h"
-#include "MetasoundDataTypeRegistrationMacro.h"
-#include "Misc/Paths.h"
+#include "MetasoundFrontendModuleRegistrationMacros.h"
 #include "Modules/ModuleManager.h"
 
 void FMetasoundExperimentalRuntimeModule::StartupModule()
 {
 	FModuleManager::Get().LoadModule("AudioExperimentalRuntime");
-	FMetasoundFrontendRegistryContainer::Get()->RegisterPendingNodes();
-
-	// Datatypes.
-	//Metasound::Frontend::RegisterDataType<Metasound::FChannelAgnosticType>();
+	METASOUND_REGISTER_ITEMS_IN_MODULE
 }
 
 void FMetasoundExperimentalRuntimeModule::ShutdownModule()
 {
+	METASOUND_UNREGISTER_ITEMS_IN_MODULE
 }
+
+METASOUND_IMPLEMENT_MODULE_REGISTRATION_LIST
 IMPLEMENT_MODULE(FMetasoundExperimentalRuntimeModule, MetasoundExperimentalRuntime);
 
-//#if WITH_ENGINE
-//#error ("Being linked with Engine")
-//#endif //
diff --git a/Engine/Plugins/Runtime/AudioModulation/Source/AudioModulation/AudioModulation.Build.cs b/Engine/Plugins/Runtime/AudioModulation/Source/AudioModulation/AudioModulation.Build.cs
--- a/Engine/Plugins/Runtime/AudioModulation/Source/AudioModulation/AudioModulation.Build.cs
+++ b/Engine/Plugins/Runtime/AudioModulation/Source/AudioModulation/AudioModulation.Build.cs
@@ -66,6 +66,13 @@ namespace UnrealBuildTool.Rules
 			{
 				PublicDefinitions.Add("WITH_AUDIOMODULATION_METASOUND_SUPPORT=0");
 			}
+
+			PrivateDefinitions.AddRange(
+				new string[]
+				{
+					"METASOUND_PLUGIN=AudioModulation",
+					"METASOUND_MODULE=AudioModulation"
+				});
 		}
 	}
 }
diff --git a/Engine/Plugins/Runtime/AudioModulation/Source/AudioModulation/Private/AudioModulation.cpp b/Engine/Plugins/Runtime/AudioModulation/Source/AudioModulation/Private/AudioModulation.cpp
--- a/Engine/Plugins/Runtime/AudioModulation/Source/AudioModulation/Private/AudioModulation.cpp
+++ b/Engine/Plugins/Runtime/AudioModulation/Source/AudioModulation/Private/AudioModulation.cpp
@@ -18,6 +18,7 @@
 
 #if WITH_AUDIOMODULATION_METASOUND_SUPPORT
 #include "MetasoundDataTypeRegistrationMacro.h"
+#include "MetasoundFrontendModuleRegistrationMacros.h"
 #include "MetasoundFrontendRegistries.h"
 
 REGISTER_METASOUND_DATATYPE(AudioModulation::FSoundModulatorAsset, "Modulator", Metasound::ELiteralType::UObjectProxy, USoundModulatorBase);
@@ -329,7 +330,7 @@ void FAudioModulationModule::StartupModule()
 	// so check that the MetaSoundEngine is loaded prior to pending Modulation defined classes
 	FModuleManager::Get().LoadModuleChecked("MetasoundEngine");
 
-	FMetasoundFrontendRegistryContainer::Get()->RegisterPendingNodes();
+	METASOUND_REGISTER_ITEMS_IN_MODULE
 #endif // WITH_AUDIOMODULATION_METASOUND_SUPPORT
 
 	UE_LOG(LogAudioModulation, Log, TEXT("Audio Modulation Initialized"));
@@ -339,8 +340,16 @@ void FAudioModulationModule::ShutdownModule()
 {
 	LLM_SCOPE(ELLMTag::AudioMixerPlugins);
 	IModularFeatures::Get().UnregisterModularFeature(FAudioModulationPluginFactory::GetModularFeatureName(), &ModulationPluginFactory);
+
+#if WITH_AUDIOMODULATION_METASOUND_SUPPORT
+	UE_LOG(LogAudioModulation, Log, TEXT("Unregistering Modulation MetaSound Nodes..."));
+	METASOUND_UNREGISTER_ITEMS_IN_MODULE
+#endif // WITH_AUDIOMODULATION_METASOUND_SUPPORT
+
 	UE_LOG(LogAudioModulation, Log, TEXT("Audio Modulation Shutdown"));
 }
 
-
+#if WITH_AUDIOMODULATION_METASOUND_SUPPORT
+METASOUND_IMPLEMENT_MODULE_REGISTRATION_LIST
+#endif
 IMPLEMENT_MODULE(FAudioModulationModule, AudioModulation);
diff --git a/Engine/Plugins/Runtime/Harmonix/Source/HarmonixMetasound/HarmonixMetasound.Build.cs b/Engine/Plugins/Runtime/Harmonix/Source/HarmonixMetasound/HarmonixMetasound.Build.cs
--- a/Engine/Plugins/Runtime/Harmonix/Source/HarmonixMetasound/HarmonixMetasound.Build.cs
+++ b/Engine/Plugins/Runtime/Harmonix/Source/HarmonixMetasound/HarmonixMetasound.Build.cs
@@ -44,5 +44,13 @@ public class HarmonixMetasound : ModuleRules
 			PrivateDependencyModuleNames.Add("AssetRegistry");
 			PrivateDependencyModuleNames.Add("UnrealEd");
 		}
+
+		PrivateDefinitions.AddRange(
+			new string[]
+			{
+				"METASOUND_PLUGIN=Harmonix",
+				"METASOUND_MODULE=HarmonixMetasound"
+			}
+		);
 	}
 }
diff --git a/Engine/Plugins/Runtime/Harmonix/Source/HarmonixMetasound/Private/HarmonixMetasoundModule.cpp b/Engine/Plugins/Runtime/Harmonix/Source/HarmonixMetasound/Private/HarmonixMetasoundModule.cpp
--- a/Engine/Plugins/Runtime/Harmonix/Source/HarmonixMetasound/Private/HarmonixMetasoundModule.cpp
+++ b/Engine/Plugins/Runtime/Harmonix/Source/HarmonixMetasound/Private/HarmonixMetasoundModule.cpp
@@ -3,6 +3,7 @@
 #include "HarmonixMetasoundModule.h"
 
 #include "MetasoundDataTypeRegistrationMacro.h"
+#include "MetasoundFrontendModuleRegistrationMacros.h"
 #include "MetasoundGeneratorHandle.h"
 #include "MusicEnvironmentSubsystem.h"
 #include "Analysis/MetasoundFrontendAnalyzerRegistry.h"
@@ -29,7 +30,7 @@ void FHarmonixMetasoundModule::StartupModule()
 {
 	using namespace Metasound;
 
-	FMetasoundFrontendRegistryContainer::Get()->RegisterPendingNodes();
+	METASOUND_REGISTER_ITEMS_IN_MODULE
 
 	// Register passthrough analyzers for output watching
 	UMetasoundGeneratorHandle::RegisterPassthroughAnalyzerForType(
@@ -76,6 +77,8 @@ void FHarmonixMetasoundModule::StartupModule()
 
 void FHarmonixMetasoundModule::ShutdownModule()
 {
+	METASOUND_UNREGISTER_ITEMS_IN_MODULE
 }
 
+METASOUND_IMPLEMENT_MODULE_REGISTRATION_LIST
 IMPLEMENT_MODULE(FHarmonixMetasoundModule, HarmonixMetasound);
diff --git a/Engine/Plugins/Runtime/Metasound/Source/MetasoundEngine/MetasoundEngine.Build.cs b/Engine/Plugins/Runtime/Metasound/Source/MetasoundEngine/MetasoundEngine.Build.cs
--- a/Engine/Plugins/Runtime/Metasound/Source/MetasoundEngine/MetasoundEngine.Build.cs
+++ b/Engine/Plugins/Runtime/Metasound/Source/MetasoundEngine/MetasoundEngine.Build.cs
@@ -45,6 +45,9 @@ namespace UnrealBuildTool.Rules
 					"WaveTable"
 				}
 			);
+
+			PrivateDefinitions.Add("METASOUND_PLUGIN=Metasound");
+			PrivateDefinitions.Add("METASOUND_MODULE=MetasoundEngine");
 		}
 	}
 }
diff --git a/Engine/Plugins/Runtime/Metasound/Source/MetasoundEngine/Private/MetasoundEngineModule.cpp b/Engine/Plugins/Runtime/Metasound/Source/MetasoundEngine/Private/MetasoundEngineModule.cpp
--- a/Engine/Plugins/Runtime/Metasound/Source/MetasoundEngine/Private/MetasoundEngineModule.cpp
+++ b/Engine/Plugins/Runtime/Metasound/Source/MetasoundEngine/Private/MetasoundEngineModule.cpp
@@ -13,6 +13,7 @@
 #include "MetasoundDocumentBuilderRegistry.h"
 #include "MetasoundDocumentInterface.h"
 #include "MetasoundFrontendDocumentBuilder.h"
+#include "MetasoundFrontendModuleRegistrationMacros.h"
 #include "MetasoundGlobals.h"
 #include "MetasoundLog.h"
 #include "MetasoundOutputSubsystem.h"
@@ -113,8 +114,8 @@ namespace Metasound::Engine
 		RegisterInterfaces();
 		RegisterInternalInterfaceBindings();
 
-		// Flush node registration queue
-		FMetasoundFrontendRegistryContainer::Get()->RegisterPendingNodes();
+		// Register objects statically associated with this module.
+		METASOUND_REGISTER_ITEMS_IN_MODULE
 
 		// Register Analyzers
 		// TODO: Determine if we can move this registration to Frontend where it likely belongs
@@ -170,6 +171,7 @@ namespace Metasound::Engine
 
 	void FModule::FModule::ShutdownModule()
 	{
+		METASOUND_UNREGISTER_ITEMS_IN_MODULE
 #if WITH_EDITOR
 		ShutdownAssetClassRegistry();
 #endif // WITH_EDITOR
@@ -444,4 +446,5 @@ namespace Metasound::Engine
 #endif // WITH_EDITOR
 } // namespace Metasound::Engine
 
+METASOUND_IMPLEMENT_MODULE_REGISTRATION_LIST
 IMPLEMENT_MODULE(Metasound::Engine::FModule, MetasoundEngine);
diff --git a/Engine/Plugins/Runtime/Metasound/Source/MetasoundFrontend/MetasoundFrontend.Build.cs b/Engine/Plugins/Runtime/Metasound/Source/MetasoundFrontend/MetasoundFrontend.Build.cs
--- a/Engine/Plugins/Runtime/Metasound/Source/MetasoundFrontend/MetasoundFrontend.Build.cs
+++ b/Engine/Plugins/Runtime/Metasound/Source/MetasoundFrontend/MetasoundFrontend.Build.cs
@@ -37,8 +37,6 @@ namespace UnrealBuildTool.Rules
 				);
 			}
 
-			PublicDefinitions.Add("WITH_METASOUND_FRONTEND=1");
-
 			// In UE 5.6, registered node are expected to support the constructor 
 			// signature Constructor(FNodeData, TSharedRef<const FNodeClassMetadata>)
 			// Because there are many existing nodes, it may take time to update 
@@ -47,6 +45,8 @@ namespace UnrealBuildTool.Rules
 			// do not drown out other compiler errors and warnings.
 			PublicDefinitions.Add("UE_METASOUND_DISABLE_5_6_NODE_REGISTRATION_DEPRECATION_WARNINGS=0");
 
+			PrivateDefinitions.Add("METASOUND_PLUGIN=Metasound");
+			PrivateDefinitions.Add("METASOUND_MODULE=MetasoundFrontend");
 		}
 	}
 }
diff --git a/Engine/Plugins/Runtime/Metasound/Source/MetasoundFrontend/Private/MetasoundFrontend.cpp b/Engine/Plugins/Runtime/Metasound/Source/MetasoundFrontend/Private/MetasoundFrontend.cpp
--- a/Engine/Plugins/Runtime/Metasound/Source/MetasoundFrontend/Private/MetasoundFrontend.cpp
+++ b/Engine/Plugins/Runtime/Metasound/Source/MetasoundFrontend/Private/MetasoundFrontend.cpp
@@ -192,22 +192,19 @@ class FMetasoundFrontendModule : public IModuleInterface
 		RegisterNodeTemplate(MakeUnique<FAudioAnalyzerNodeTemplate>());
 		RegisterNodeTemplate(MakeUnique<FInputNodeTemplate>());
 		RegisterNodeTemplate(MakeUnique<FRerouteNodeTemplate>());
-
-		FMetasoundFrontendRegistryContainer* Registry = FMetasoundFrontendRegistryContainer::Get();
-		if (ensure(nullptr != Registry))
-		{
-			Registry->RegisterPendingNodes();
-		}
+		METASOUND_REGISTER_ITEMS_IN_MODULE
 	}
 
 	virtual void ShutdownModule() override
 	{
 		using namespace Metasound::Frontend;
 
+		METASOUND_UNREGISTER_ITEMS_IN_MODULE
 		UnregisterNodeTemplate(FAudioAnalyzerNodeTemplate::ClassName, FAudioAnalyzerNodeTemplate::VersionNumber);
 		UnregisterNodeTemplate(FInputNodeTemplate::ClassName, FInputNodeTemplate::VersionNumber);
 		UnregisterNodeTemplate(FRerouteNodeTemplate::ClassName, FRerouteNodeTemplate::VersionNumber);
 	}
 };
 
+METASOUND_IMPLEMENT_MODULE_REGISTRATION_LIST
 IMPLEMENT_MODULE(FMetasoundFrontendModule, MetasoundFrontend);
diff --git a/Engine/Plugins/Runtime/Metasound/Source/MetasoundFrontend/Private/MetasoundFrontendModuleRegistrationMacros.cpp b/Engine/Plugins/Runtime/Metasound/Source/MetasoundFrontend/Private/MetasoundFrontendModuleRegistrationMacros.cpp
new file mode 100644
--- /dev/null
+++ b/Engine/Plugins/Runtime/Metasound/Source/MetasoundFrontend/Private/MetasoundFrontendModuleRegistrationMacros.cpp
@@ -0,0 +1,95 @@
+// Copyright Epic Games, Inc. All Rights Reserved.
+
+#include "MetasoundFrontendModuleRegistrationMacros.h"
+#include "MetasoundLog.h"
+
+namespace Metasound::Frontend::RegistrationPrivate
+{
+	FRegistrationAction::FRegistrationAction(FRegistrationAction*& InOutHead, FRegistrationFunction InRegister, FRegistrationFunction InUnregister)
+	: Register(InRegister)
+	, Unregister(InUnregister)
+	{
+		Next = InOutHead;
+		InOutHead = this;
+	}
+
+	void FRegistrationListBase::RegisterAll(const Metasound::Frontend::FModuleInfo& InOwningModuleInfo, FRegistrationAction* Head)
+	{
+#if WITH_EDITORONLY_DATA
+		UE_LOG(LogMetaSound, Verbose, TEXT("Running registration actions for module %s in plugin %s"), *InOwningModuleInfo.ModuleName.ToString(), *InOwningModuleInfo.PluginName.ToString());
+#endif // if WITH_EDITORONLY_DATA
+		for (FRegistrationAction* Action = Head; Action != nullptr; Action = Action->Next)
+		{
+#if WITH_EDITORONLY_DATA
+			UE_LOG(LogMetaSound, VeryVerbose, TEXT("Running Register:%s for module %s in plugin %s"), Action->Name, *InOwningModuleInfo.ModuleName.ToString(), *InOwningModuleInfo.PluginName.ToString());
+#endif
+			if (Action->Register)
+			{
+				Action->Register(InOwningModuleInfo);
+			}
+		}
+	}
+
+	void FRegistrationListBase::UnregisterAll(const FModuleInfo& InOwningModuleInfo, FRegistrationAction* Head)
+	{
+#if WITH_EDITORONLY_DATA
+		UE_LOG(LogMetaSound, Verbose, TEXT("Running unregistration actions for module %s in plugin %s"), *InOwningModuleInfo.ModuleName.ToString(), *InOwningModuleInfo.PluginName.ToString());
+#endif // if WITH_EDITORONLY_DATA
+		for (FRegistrationAction* Action = Head; Action != nullptr; Action = Action->Next)
+		{
+#if WITH_EDITORONLY_DATA
+			UE_LOG(LogMetaSound, VeryVerbose, TEXT("Running Unregister:%s for module %s in plugin %s"), Action->Name, *InOwningModuleInfo.ModuleName.ToString(), *InOwningModuleInfo.PluginName.ToString());
+#endif
+			if (Action->Unregister)
+			{
+				Action->Unregister(InOwningModuleInfo);
+			}
+		}
+	}
+
+	FRegistrationAction*& FGlobalRegistrationList::GetHeadAction()
+	{
+		static FRegistrationAction* Head = nullptr;
+		return Head;
+	}
+
+	void FGlobalRegistrationList::RegisterAll(const FModuleInfo& InOwningModuleInfo, FRegistrationAction* Head)
+	{
+#if WITH_EDITORONLY_DATA
+		UE_LOG(LogMetaSound, Verbose, TEXT("Running pending registration actions for global registration list"));
+#endif // if WITH_EDITORONLY_DATA
+	  
+		// Global registration needs to handle the scenario where it is called
+		// multiple times from different modules. It has to protect against double
+		// registration of nodes. 
+		
+		FRegistrationAction* End = GetRegisterTailAction();
+		
+		for (FRegistrationAction* Action = Head; (Action != End) && (Action != nullptr); Action = Action->Next)
+		{
+#if WITH_EDITORONLY_DATA
+			UE_LOG(LogMetaSound, VeryVerbose, TEXT("Running Register:%s for global registration list"), Action->Name);
+#endif
+			if (Action->Register)
+			{
+				Action->Register(InOwningModuleInfo);
+			}
+		}
+
+		// Track nodes that have been registered
+		GetRegisterTailAction() = Head;
+	}
+
+	void FGlobalRegistrationList::UnregisterAll(const FModuleInfo& InOwningModuleInfo, FRegistrationAction* Head)
+	{
+		// Unregister on the global registration list is not supported. 
+		checkNoEntry();
+	}
+
+	FRegistrationAction*& FGlobalRegistrationList::GetRegisterTailAction()
+	{
+		static FRegistrationAction* Tail = nullptr;
+		return Tail;
+	}
+}
+
diff --git a/Engine/Plugins/Runtime/Metasound/Source/MetasoundFrontend/Private/MetasoundFrontendNodeClassRegistry.cpp b/Engine/Plugins/Runtime/Metasound/Source/MetasoundFrontend/Private/MetasoundFrontendNodeClassRegistry.cpp
--- a/Engine/Plugins/Runtime/Metasound/Source/MetasoundFrontend/Private/MetasoundFrontendNodeClassRegistry.cpp
+++ b/Engine/Plugins/Runtime/Metasound/Source/MetasoundFrontend/Private/MetasoundFrontendNodeClassRegistry.cpp
@@ -23,6 +23,7 @@
 #include "MetasoundFrontendDataTypeRegistry.h"
 #include "MetasoundFrontendDocumentBuilder.h"
 #include "MetasoundFrontendGraph.h"
+#include "MetasoundFrontendModuleRegistrationMacros.h"
 #include "MetasoundFrontendPages.h"
 #include "MetasoundFrontendProxyDataCache.h"
 #include "MetasoundFrontendRegistryTransaction.h"
@@ -232,6 +233,20 @@ PRAGMA_ENABLE_DEPRECATION_WARNINGS
 	}
 #endif
 
+#if WITH_EDITORONLY_DATA
+		FName INodeClassRegistryEntry::GetPluginName() const
+		{
+			// Default implementation returns default FName
+			return FName();
+		}
+
+		FName INodeClassRegistryEntry::GetModuleName() const
+		{
+			// Default implementation returns default FName
+			return FName();
+		}
+#endif
+
 	FNodeClassRegistryTransaction::FNodeClassRegistryTransaction(ETransactionType InType, const FNodeClassInfo& InNodeClassInfo, FNodeClassRegistryTransaction::FTimeType InTimestamp)
 	: Type(InType)
 	, NodeClassInfo(InNodeClassInfo)
@@ -340,6 +355,7 @@ PRAGMA_ENABLE_DEPRECATION_WARNINGS
 			return FNodeClassRegistryKey(InClassInfo);
 		}
 	} // namespace NodeClassRegistryKey
+	  
 
 	void FNodeClassRegistry::BuildAndRegisterGraphFromDocument(const FMetasoundFrontendDocument& InDocument, const FProxyDataCache& InProxyDataCache, FNodeClassInfo&& InNodeClassInfo, const FTopLevelAssetPath& AssetPath, TArrayView<const FGuid> InPageOrder)
 	{
@@ -425,6 +441,13 @@ PRAGMA_ENABLE_DEPRECATION_WARNINGS
 			LazyInitCommands.Empty();
 		}
 
+		// For backwards compatibility, also register global registration list
+		// to mimic behavior before the introduction of the module specific registration
+		// lists.
+PRAGMA_DISABLE_DEPRECATION_WARNINGS
+		RegistrationPrivate::FGlobalRegistrationList::RegisterAll(FModuleInfo{}, RegistrationPrivate::FGlobalRegistrationList::GetHeadAction());
+PRAGMA_ENABLE_DEPRECATION_WARNINGS
+
 		if (!IsRunningCommandlet())
 		{
 			// Prime search engine after bulk registration.
@@ -905,6 +928,77 @@ PRAGMA_ENABLE_DEPRECATION_WARNINGS
 		return Key;
 	}
 
+#if WITH_EDITORONLY_DATA
+	bool FNodeClassRegistry::RegisterNodeMigration(const FNodeMigrationInfo& InMigrationInfo)
+	{
+		FNodeClassInfo ClassInfo;
+		ClassInfo.ClassName = InMigrationInfo.ClassName;
+		ClassInfo.Version = FMetasoundFrontendVersionNumber{ .Major=InMigrationInfo.MajorVersion, .Minor=InMigrationInfo.MinorVersion };
+		ClassInfo.Type = EMetasoundFrontendClassType::External;
+
+		const FNodeClassRegistryKey Key{ClassInfo};
+
+		// Attempts to register the exact same migration info twice will result in bSuccess=false
+		bool bSuccess = false;
+
+		// Add to map of migrations
+		{
+			FScopeLock Lock(&RegistryMapsCriticalSection);
+			if (nullptr == NodeMigrations.FindPair(Key, InMigrationInfo))
+			{
+				NodeMigrations.Add(Key, InMigrationInfo);
+				bSuccess = true;
+			}
+		}
+
+		if (bSuccess)
+		{
+			// Add to transactions so that the search engine knows it exists. 
+			const FNodeClassRegistryTransaction::FTimeType Timestamp = FPlatformTime::Cycles64();
+			TransactionBuffer->AddTransaction(FNodeClassRegistryTransaction(FNodeClassRegistryTransaction::ETransactionType::NodeMigrationRegistration, ClassInfo, Timestamp));
+		}
+
+		UE_CLOG(!bSuccess, LogMetaSound, Error, TEXT("Invalid attempt to register node migration more than once: %s. Please ensure node migration info is unique."), *InMigrationInfo.ToString());
+
+
+		return bSuccess;
+	}
+
+	bool FNodeClassRegistry::UnregisterNodeMigration(const FNodeMigrationInfo& InMigrationInfo)
+	{
+		FNodeClassInfo ClassInfo;
+		ClassInfo.ClassName = InMigrationInfo.ClassName;
+		ClassInfo.Version = FMetasoundFrontendVersionNumber{ .Major=InMigrationInfo.MajorVersion, .Minor=InMigrationInfo.MinorVersion };
+		ClassInfo.Type = EMetasoundFrontendClassType::External;
+
+		const FNodeClassRegistryKey Key{ClassInfo};
+
+		// Attempts to unregister a migration that does not exist will result in bSuccess=false
+		bool bSuccess = false;
+
+		// remove from map of migrations
+		{
+			FScopeLock Lock(&RegistryMapsCriticalSection);
+			int32 NumRemoved = NodeMigrations.Remove(Key, InMigrationInfo);
+			if (NumRemoved > 0)
+			{
+				bSuccess = true;
+			}
+		}
+
+		if (bSuccess)
+		{
+			// Add to transactions so that the search engine knows it exists. 
+			const FNodeClassRegistryTransaction::FTimeType Timestamp = FPlatformTime::Cycles64();
+			TransactionBuffer->AddTransaction(FNodeClassRegistryTransaction(FNodeClassRegistryTransaction::ETransactionType::NodeMigrationUnregistration, ClassInfo, Timestamp));
+		}
+
+		UE_CLOG(!bSuccess, LogMetaSound, Error, TEXT("Invalid attempt to unregister node migration more than once: %s. Please ensure node migration info is unique."), *InMigrationInfo.ToString());
+
+		return bSuccess;
+	}
+#endif // WITH_EDITORONLY_DATA
+
 	FNodeClassRegistryKey FNodeClassRegistry::RegisterNodeTemplate(TUniquePtr<INodeTemplateRegistryEntry>&& InEntry)
 	{
 		METASOUND_LLM_SCOPE;
@@ -1237,6 +1331,17 @@ PRAGMA_ENABLE_DEPRECATION_WARNINGS
 				InFunc(*(*Entry));
 				return true;
 			}
+#if WITH_EDITORONLY_DATA
+			else if (const FNodeMigrationInfo* Migration = NodeMigrations.Find(InKey))
+			{
+				TArray<const FNodeMigrationInfo*> Migrations;
+				NodeMigrations.MultiFindPointer(InKey, Migrations);
+				for (const FNodeMigrationInfo* Info : Migrations)
+				{
+					UE_LOG(LogMetaSound, Error, TEXT("Node cannot be found because it has been migrated. %s. Please update your MetaSound Graph or add plugin %s as a depedency."), *Info->ToString(), *Info->ToPlugin.ToString());
+				}
+			}
+#endif // if WITH_EDITORONLY_DATA
 			return false;
 		};
 
@@ -1306,6 +1411,11 @@ PRAGMA_ENABLE_DEPRECATION_WARNINGS
 		return &Metasound::Frontend::FNodeClassRegistry::Get();
 	}
 
+	INodeClassRegistry& INodeClassRegistry::GetChecked()
+	{
+		return Metasound::Frontend::FNodeClassRegistry::Get();
+	}
+
 	void INodeClassRegistry::ShutdownMetasoundFrontend()
 	{
 		Metasound::Frontend::FNodeClassRegistry::Shutdown();
diff --git a/Engine/Plugins/Runtime/Metasound/Source/MetasoundFrontend/Private/MetasoundFrontendNodeClassRegistryPrivate.h b/Engine/Plugins/Runtime/Metasound/Source/MetasoundFrontend/Private/MetasoundFrontendNodeClassRegistryPrivate.h
--- a/Engine/Plugins/Runtime/Metasound/Source/MetasoundFrontend/Private/MetasoundFrontendNodeClassRegistryPrivate.h
+++ b/Engine/Plugins/Runtime/Metasound/Source/MetasoundFrontend/Private/MetasoundFrontendNodeClassRegistryPrivate.h
@@ -84,8 +84,13 @@ namespace Metasound::Frontend
 		 * @return True on success.
 		 */
 		virtual FNodeClassRegistryKey RegisterNode(TUniquePtr<Metasound::Frontend::INodeClassRegistryEntry>&& InEntry) override;
-
 		virtual bool UnregisterNode(const FNodeClassRegistryKey& InKey) override;
+
+#if WITH_EDITORONLY_DATA
+		virtual bool RegisterNodeMigration(const FNodeMigrationInfo& InMigrationInfo) override;
+		virtual bool UnregisterNodeMigration(const FNodeMigrationInfo& InMigrationInfo) override;
+#endif // if WITH_EDITORONLY_DATA
+
 		virtual bool IsNodeRegistered(const FNodeClassRegistryKey& InKey) const override;
 		virtual bool IsGraphRegistered(const FGraphRegistryKey& InKey) const override;
 
@@ -182,6 +187,10 @@ namespace Metasound::Frontend
 		// Registry in which we keep all information about nodes implemented in C++.
 		TMultiMap<FNodeClassRegistryKey, TSharedRef<INodeClassRegistryEntry, ESPMode::ThreadSafe>> RegisteredNodes;
 
+#if WITH_EDITORONLY_DATA
+		TMultiMap<FNodeClassRegistryKey, FNodeMigrationInfo> NodeMigrations;
+#endif
+
 		// Registry in which we keep all information about dynamically-generated templated nodes via in C++.
 		TMap<FNodeClassRegistryKey, TSharedRef<INodeTemplateRegistryEntry, ESPMode::ThreadSafe>> RegisteredNodeTemplates;
 
diff --git a/Engine/Plugins/Runtime/Metasound/Source/MetasoundFrontend/Private/MetasoundFrontendNodeMigration.cpp b/Engine/Plugins/Runtime/Metasound/Source/MetasoundFrontend/Private/MetasoundFrontendNodeMigration.cpp
new file mode 100644
--- /dev/null
+++ b/Engine/Plugins/Runtime/Metasound/Source/MetasoundFrontend/Private/MetasoundFrontendNodeMigration.cpp
@@ -0,0 +1,44 @@
+// Copyright Epic Games, Inc. All Rights Reserved.
+
+#include "MetasoundFrontendNodeMigration.h"
+
+namespace Metasound::Frontend
+{
+
+#if WITH_EDITORONLY_DATA
+	bool operator==(const FNodeMigrationInfo& InLHS, const FNodeMigrationInfo& InRHS)
+	{
+		return (InLHS.ClassName == InRHS.ClassName) &&
+			(InLHS.UEVersion == InRHS.UEVersion) &&
+			(InLHS.MajorVersion == InRHS.MajorVersion) &&
+			(InLHS.MinorVersion == InRHS.MinorVersion ) &&
+			(InLHS.FromPlugin == InRHS.FromPlugin) &&
+			(InLHS.FromModule == InRHS.FromModule) &&
+			(InLHS.ToPlugin == InRHS.ToPlugin) &&
+			(InLHS.ToModule == InRHS.ToModule);
+	}
+
+	FString FNodeMigrationInfo::ToString() const
+	{
+		TStringBuilder<256> StringBuilder;
+		return StringBuilder.Append(TEXT("Migration: "))
+			.Append(ClassName.ToString())
+			.AppendChar(' ')
+			.Append(FString::FromInt(MajorVersion))
+			.AppendChar('.')
+			.Append(FString::FromInt(MinorVersion))
+			.Append(TEXT(" From "))
+			.Append(FromPlugin.ToString())
+			.AppendChar('/')
+			.Append(FromModule.ToString())
+			.Append(TEXT(" To "))
+			.Append(ToPlugin.ToString())
+			.AppendChar('/')
+			.Append(ToModule.ToString())
+			.Append(TEXT(" in UE"))
+			.Append(UEVersion.ToString())
+			.ToString();
+	}
+#endif // if WITH_EDITORONLY_DATA
+
+} // namespace Metasound::Frontend
diff --git a/Engine/Plugins/Runtime/Metasound/Source/MetasoundFrontend/Private/MetasoundFrontendQuerySteps.cpp b/Engine/Plugins/Runtime/Metasound/Source/MetasoundFrontend/Private/MetasoundFrontendQuerySteps.cpp
--- a/Engine/Plugins/Runtime/Metasound/Source/MetasoundFrontend/Private/MetasoundFrontendQuerySteps.cpp
+++ b/Engine/Plugins/Runtime/Metasound/Source/MetasoundFrontend/Private/MetasoundFrontendQuerySteps.cpp
@@ -68,48 +68,91 @@ namespace Metasound
 	{
 		using namespace Frontend;
 
-		InOutEntries.Sort([](const FFrontendQueryEntry& A, const FFrontendQueryEntry& B)
+		// Track number of times each type of entry has been seen
+		struct FTransactionStats
 		{
-			return GetTransactionTimestamp(A) < GetTransactionTimestamp(B);
-		});
-
-		// Registration - Unregistration pairs result in no net change, 
-		// so reduce to most recent transaction that is not a pair
-		int32 MostRecentNonPairedIndex = INDEX_NONE;
-		int32 Index = 0;
-		while (Index < InOutEntries.Num())
-		{
-			// If last entry, cannot be a pair 
-			if (Index == InOutEntries.Num() - 1)
-			{
-				MostRecentNonPairedIndex = Index;
-				break;
-			}
-			// Skip if valid pair
-			if (InOutEntries[Index].Value.Get<FNodeRegistryTransaction>().GetTransactionType() == FNodeRegistryTransaction::ETransactionType::NodeRegistration &&
-				InOutEntries[Index + 1].Value.Get<FNodeRegistryTransaction>().GetTransactionType() == FNodeRegistryTransaction::ETransactionType::NodeUnregistration)
+			int32 Num = 0;
+			const FFrontendQueryEntry* MostRecent = nullptr;
+
+			void Update(const FFrontendQueryEntry& InEntry)
 			{
-				Index += 2;
+				Num++;
+				if (nullptr == MostRecent)
+				{
+					MostRecent = &InEntry;
+				}
+				else
+				{
+					if (FReduceRegistrationEventsToCurrentStatus::GetTransactionTimestamp(InEntry) > FReduceRegistrationEventsToCurrentStatus::GetTransactionTimestamp(*MostRecent))
+					{
+						MostRecent = &InEntry;
+					}
+				}
 			}
-			// If not valid pair, this is the current most recent non paired index
-			else
+		};
+
+		// Accumulate the counts of each transaction type
+		FTransactionStats NodeRegistrationStats;
+		FTransactionStats NodeUnregistrationStats;
+		FTransactionStats MigrationRegistrationStats;
+		FTransactionStats MigrationUnregistrationStats;
+
+		for (const FFrontendQueryEntry& Entry : InOutEntries)
+		{
+			const FNodeClassRegistryTransaction& Transaction = Entry.Value.Get<FNodeRegistryTransaction>();
+			switch (Transaction.GetTransactionType())
 			{
-				MostRecentNonPairedIndex = Index;
-				Index++;
+				case FNodeClassRegistryTransaction::ETransactionType::NodeRegistration:
+					NodeRegistrationStats.Update(Entry);
+					break;
+
+				case FNodeClassRegistryTransaction::ETransactionType::NodeUnregistration:
+					NodeUnregistrationStats.Update(Entry);
+					break;
+
+				case FNodeClassRegistryTransaction::ETransactionType::NodeMigrationRegistration:
+					MigrationRegistrationStats.Update(Entry);
+					break;
+
+				case FNodeClassRegistryTransaction::ETransactionType::NodeMigrationUnregistration:
+					MigrationUnregistrationStats.Update(Entry);
+					break;
+
+				default:
+					{
+						checkNoEntry();
+					}
 			}
 		}
 
-		// Empty or entries were all pairs 
-		if (MostRecentNonPairedIndex == INDEX_NONE)
+		// Use the final number of each transaction type to determine the final state
+		FFrontendQueryPartition RemainingEntries;
+		if (NodeRegistrationStats.Num > NodeUnregistrationStats.Num)
 		{
-			InOutEntries.Reset();
+			// Registration trumps migration, so we do not consider migration at all here.
+			RemainingEntries.Add(*NodeRegistrationStats.MostRecent);
 		}
-		else
+		else 
 		{
-			const FFrontendQueryEntry Entry = InOutEntries[MostRecentNonPairedIndex];
-			InOutEntries.Reset();
-			InOutEntries.Add(Entry);
+			// Make sure to propagate unregistration actions since reduce steps may get called
+			// on subpartitions. 
+			if (NodeUnregistrationStats.Num > NodeRegistrationStats.Num)
+			{
+				RemainingEntries.Add(*NodeUnregistrationStats.MostRecent);
+			}
+			
+			// If  not registered, check if migrated
+			if (MigrationRegistrationStats.Num > MigrationUnregistrationStats.Num)
+			{
+				RemainingEntries.Add(*MigrationRegistrationStats.MostRecent);
+			}
+			else if (MigrationUnregistrationStats.Num > MigrationRegistrationStats.Num)
+			{
+				RemainingEntries.Add(*MigrationUnregistrationStats.MostRecent);
+			}
 		}
+
+		InOutEntries = MoveTemp(RemainingEntries);
 	}
 
 	FReduceRegistrationEventsToCurrentStatus::FTimeType FReduceRegistrationEventsToCurrentStatus::GetTransactionTimestamp(const FFrontendQueryEntry& InEntry)
diff --git a/Engine/Plugins/Runtime/Metasound/Source/MetasoundFrontend/Private/MetasoundNodeRegistrationMacro.cpp b/Engine/Plugins/Runtime/Metasound/Source/MetasoundFrontend/Private/MetasoundNodeRegistrationMacro.cpp
--- a/Engine/Plugins/Runtime/Metasound/Source/MetasoundFrontend/Private/MetasoundNodeRegistrationMacro.cpp
+++ b/Engine/Plugins/Runtime/Metasound/Source/MetasoundFrontend/Private/MetasoundNodeRegistrationMacro.cpp
@@ -10,10 +10,14 @@ namespace Metasound::Frontend
 {
 	namespace NodeRegistrationPrivate
 	{
-		FNodeRegistryEntryBase::FNodeRegistryEntryBase(const FNodeClassMetadata& InMetadata)
+		FNodeRegistryEntryBase::FNodeRegistryEntryBase(const Metasound::FNodeClassMetadata& InMetadata, const FModuleInfo& InOwningModuleInfo)
 		: ClassMetadata(MakeShared<FNodeClassMetadata>(InMetadata))
 		, ClassInfo(FMetasoundFrontendClassMetadata::GenerateClassMetadata(InMetadata, EMetasoundFrontendClassType::External))
 		, FrontendClass(GenerateClass(InMetadata))
+#if WITH_EDITORONLY_DATA
+		, PluginName(InOwningModuleInfo.PluginName)
+		, ModuleName(InOwningModuleInfo.ModuleName)
+#endif
 		{
 		}
 
@@ -51,6 +55,17 @@ namespace Metasound::Frontend
 		{
 			return ClassMetadata;
 		}
+#if WITH_EDITORONLY_DATA
+		FName FNodeRegistryEntryBase::GetPluginName() const
+		{
+			return PluginName;
+		}
+
+		FName FNodeRegistryEntryBase::GetModuleName() const
+		{
+			return ModuleName;
+		}
+#endif
 
 		// This adapter class forwards the correct FBuilderOperatorParams
 		// to the node's operator creation method. Many operator creation
@@ -105,21 +120,9 @@ namespace Metasound::Frontend
 			return bSuccessfullyRegisteredNode;
 		}
 
-		FORCENOINLINE bool RegisterNodeUsingMacroCheck(bool& bAlreadyRegistered)
-		{
-			if (bAlreadyRegistered)
-			{
-				UE_LOG(LogMetaSound, Display, TEXT("Tried to call METASOUND_REGISTER_NODE twice with the same class. ignoring the second call. Likely because METASOUND_REGISTER_NODE is in a header that's used in multiple modules. Consider moving it to a private header or cpp file."));
-				return false;
-			}
-
-			bAlreadyRegistered = true;
-			return true;
-		}
-
-		FORCENOINLINE bool RegisterNodeDeferred(void (*InRegisterFunc)())
+		bool UnregisterNodeInternal(const FNodeClassMetadata& InMetadata)
 		{
-			return INodeClassRegistry::Get()->EnqueueInitCommand(InRegisterFunc);
+			return INodeClassRegistry::Get()->UnregisterNode(FNodeClassRegistryKey(InMetadata));
 		}
 	}
 }
diff --git a/Engine/Plugins/Runtime/Metasound/Source/MetasoundFrontend/Public/MetasoundArrayNodesRegistration.h b/Engine/Plugins/Runtime/Metasound/Source/MetasoundFrontend/Public/MetasoundArrayNodesRegistration.h
--- a/Engine/Plugins/Runtime/Metasound/Source/MetasoundFrontend/Public/MetasoundArrayNodesRegistration.h
+++ b/Engine/Plugins/Runtime/Metasound/Source/MetasoundFrontend/Public/MetasoundArrayNodesRegistration.h
@@ -61,119 +61,119 @@ namespace Metasound
 		};
 
 		template<typename ArrayType, typename std::enable_if<TArrayNodeSupport<ArrayType>::bIsArrayGetSupported, bool>::type = true>
-		bool RegisterArrayGetNode()
+		bool RegisterArrayGetNode(const Frontend::FModuleInfo& InModuleInfo)
 		{
 			using FNodeType = typename Metasound::TArrayGetNode<ArrayType>;
-			return Frontend::RegisterNode<FNodeType>();
+			return Frontend::RegisterNode<FNodeType>(InModuleInfo);
 		}
 
 		template<typename ArrayType, typename std::enable_if<!TArrayNodeSupport<ArrayType>::bIsArrayGetSupported, bool>::type = true>
-		bool RegisterArrayGetNode()
+		bool RegisterArrayGetNode(const Frontend::FModuleInfo& InModuleInfo)
 		{
 			// No op if not supported
 			return true;
 		}
 
 		template<typename ArrayType, typename std::enable_if<TArrayNodeSupport<ArrayType>::bIsArraySetSupported, bool>::type = true>
-		bool RegisterArraySetNode()
+		bool RegisterArraySetNode(const Frontend::FModuleInfo& InModuleInfo)
 		{
 			using FNodeType = typename Metasound::TArraySetNode<ArrayType>;
 
 			static_assert(TArrayNodeSupport<ArrayType>::bIsArraySetSupported, "TArraySetNode<> is not supported by array type");
 
-			return Frontend::RegisterNode<FNodeType>();
+			return Frontend::RegisterNode<FNodeType>(InModuleInfo);
 		}
 
 		template<typename ArrayType, typename std::enable_if<!TArrayNodeSupport<ArrayType>::bIsArraySetSupported, bool>::type = true>
-		bool RegisterArraySetNode()
+		bool RegisterArraySetNode(const Frontend::FModuleInfo& InModuleInfo)
 		{
 			// No op if not supported
 			return true;
 		}
 
 		template<typename ArrayType, typename std::enable_if<TArrayNodeSupport<ArrayType>::bIsArraySubsetSupported, bool>::type = true>
-		bool RegisterArraySubsetNode()
+		bool RegisterArraySubsetNode(const Frontend::FModuleInfo& InModuleInfo)
 		{
 			using FNodeType = typename Metasound::TArraySubsetNode<ArrayType>;
 
 			static_assert(TArrayNodeSupport<ArrayType>::bIsArraySubsetSupported, "TArraySubsetNode<> is not supported by array type");
 
-			return Frontend::RegisterNode<FNodeType>();
+			return Frontend::RegisterNode<FNodeType>(InModuleInfo);
 		}
 
 		template<typename ArrayType, typename std::enable_if<!TArrayNodeSupport<ArrayType>::bIsArraySubsetSupported, bool>::type = true>
-		bool RegisterArraySubsetNode()
+		bool RegisterArraySubsetNode(const Frontend::FModuleInfo& InModuleInfo)
 		{
 			// No op if not supported
 			return true;
 		}
 
 		template<typename ArrayType, typename std::enable_if<TArrayNodeSupport<ArrayType>::bIsArrayConcatSupported, bool>::type = true>
-		bool RegisterArrayConcatNode()
+		bool RegisterArrayConcatNode(const Frontend::FModuleInfo& InModuleInfo)
 		{
 			using FNodeType = typename Metasound::TArrayConcatNode<ArrayType>;
 
 			static_assert(TArrayNodeSupport<ArrayType>::bIsArrayConcatSupported, "TArrayConcatNode<> is not supported by array type");
 
-			return Frontend::RegisterNode<FNodeType>();
+			return Frontend::RegisterNode<FNodeType>(InModuleInfo);
 		}
 
 		template<typename ArrayType, typename std::enable_if<!TArrayNodeSupport<ArrayType>::bIsArrayConcatSupported, bool>::type = true>
-		bool RegisterArrayConcatNode()
+		bool RegisterArrayConcatNode(const Frontend::FModuleInfo& InModuleInfo)
 		{
 			// No op if not supported
 			return true;
 		}
 
 		template<typename ArrayType, typename std::enable_if<TArrayNodeSupport<ArrayType>::bIsArrayNumSupported, bool>::type = true>
-		bool RegisterArrayNumNode()
+		bool RegisterArrayNumNode(const Frontend::FModuleInfo& InModuleInfo)
 		{
-			return Frontend::RegisterNode<Metasound::TArrayNumNode<ArrayType>>();
+			return Frontend::RegisterNode<Metasound::TArrayNumNode<ArrayType>>(InModuleInfo);
 		}
 
 		template<typename ArrayType, typename std::enable_if<!TArrayNodeSupport<ArrayType>::bIsArrayNumSupported, bool>::type = true>
-		bool RegisterArrayNumNode()
+		bool RegisterArrayNumNode(const Frontend::FModuleInfo& InModuleInfo)
 		{
 			// No op if not supported
 			return true;
 		}
 
 		template<typename ArrayType, typename std::enable_if<TArrayNodeSupport<ArrayType>::bIsArrayShuffleSupported, bool>::type = true>
-		bool RegisterArrayShuffleNode()
+		bool RegisterArrayShuffleNode(const Frontend::FModuleInfo& InModuleInfo)
 		{
 			using FNodeType = typename Metasound::TArrayShuffleNode<ArrayType>;
-			return Frontend::RegisterNode<FNodeType>();
+			return Frontend::RegisterNode<FNodeType>(InModuleInfo);
 		}
 
 		template<typename ArrayType, typename std::enable_if<!TArrayNodeSupport<ArrayType>::bIsArrayShuffleSupported, bool>::type = true>
-		bool RegisterArrayShuffleNode()
+		bool RegisterArrayShuffleNode(const Frontend::FModuleInfo& InModuleInfo)
 		{
 			// No op if not supported
 			return true;
 		}
 
 		template<typename ArrayType, typename std::enable_if<TArrayNodeSupport<ArrayType>::bIsArrayRandomGetSupported, bool>::type = true>
-		bool RegisterArrayRandomGetNode()
+		bool RegisterArrayRandomGetNode(const Frontend::FModuleInfo& InModuleInfo)
 		{
 			using FNodeType = typename Metasound::TArrayRandomGetNode<ArrayType>;
-			return Frontend::RegisterNode<FNodeType>();
+			return Frontend::RegisterNode<FNodeType>(InModuleInfo);
 		}
 
 		template<typename ArrayType, typename std::enable_if<!TArrayNodeSupport<ArrayType>::bIsArrayRandomGetSupported, bool>::type = true>
-		bool RegisterArrayRandomGetNode()
+		bool RegisterArrayRandomGetNode(const Frontend::FModuleInfo& InModuleInfo)
 		{
 			// No op if not supported
 			return true;
 		}
 
 		template<typename ArrayType, typename std::enable_if<TArrayNodeSupport<ArrayType>::bIsArrayNumSupported, bool>::type = true>
-		bool RegisterArrayLastIndexNode()
+		bool RegisterArrayLastIndexNode(const Frontend::FModuleInfo& InModuleInfo)
 		{
-			return Frontend::RegisterNode<Metasound::TArrayLastIndexNode<ArrayType>>();
+			return Frontend::RegisterNode<Metasound::TArrayLastIndexNode<ArrayType>>(InModuleInfo);
 		}
 
 		template<typename ArrayType, typename std::enable_if<!TArrayNodeSupport<ArrayType>::bIsArrayNumSupported, bool>::type = true>
-		bool RegisterArrayLastIndexNode()
+		bool RegisterArrayLastIndexNode(const Frontend::FModuleInfo& InModuleInfo)
 		{
 			// No op if not supported
 			return true;
@@ -187,18 +187,18 @@ namespace Metasound
 		 * array elements.
 		 */
 		template<typename ArrayType>
-		bool RegisterArrayNodes()
+		bool RegisterArrayNodes(const FModuleInfo& InModuleInfo)
 		{
 			using namespace MetasoundArrayNodesPrivate;
 
-			bool bSuccess = RegisterArrayNumNode<ArrayType>();
-			bSuccess = bSuccess && RegisterArrayGetNode<ArrayType>();
-			bSuccess = bSuccess && RegisterArraySetNode<ArrayType>();
-			bSuccess = bSuccess && RegisterArraySubsetNode<ArrayType>();
-			bSuccess = bSuccess && RegisterArrayConcatNode<ArrayType>();
-			bSuccess = bSuccess && RegisterArrayShuffleNode<ArrayType>();
-			bSuccess = bSuccess && RegisterArrayRandomGetNode<ArrayType>();
-			bSuccess = bSuccess && RegisterArrayLastIndexNode<ArrayType>();
+			bool bSuccess = RegisterArrayNumNode<ArrayType>(InModuleInfo);
+			bSuccess = bSuccess && RegisterArrayGetNode<ArrayType>(InModuleInfo);
+			bSuccess = bSuccess && RegisterArraySetNode<ArrayType>(InModuleInfo);
+			bSuccess = bSuccess && RegisterArraySubsetNode<ArrayType>(InModuleInfo);
+			bSuccess = bSuccess && RegisterArrayConcatNode<ArrayType>(InModuleInfo);
+			bSuccess = bSuccess && RegisterArrayShuffleNode<ArrayType>(InModuleInfo);
+			bSuccess = bSuccess && RegisterArrayRandomGetNode<ArrayType>(InModuleInfo);
+			bSuccess = bSuccess && RegisterArrayLastIndexNode<ArrayType>(InModuleInfo);
 			return bSuccess;
 		}
 	}
diff --git a/Engine/Plugins/Runtime/Metasound/Source/MetasoundFrontend/Public/MetasoundDataTypeRegistrationMacro.h b/Engine/Plugins/Runtime/Metasound/Source/MetasoundFrontend/Public/MetasoundDataTypeRegistrationMacro.h
--- a/Engine/Plugins/Runtime/Metasound/Source/MetasoundFrontend/Public/MetasoundDataTypeRegistrationMacro.h
+++ b/Engine/Plugins/Runtime/Metasound/Source/MetasoundFrontend/Public/MetasoundDataTypeRegistrationMacro.h
@@ -3,6 +3,7 @@
 #pragma once
 
 #include "IAudioProxyInitializer.h"
+#include "HAL/PreprocessorHelpers.h"
 #include "MetasoundArrayNodesRegistration.h"
 #include "MetasoundAutoConverterNode.h"
 #include "MetasoundConverterNodeRegistrationMacro.h"
@@ -12,6 +13,7 @@
 #include "MetasoundEnum.h"
 #include "MetasoundFrontendDataTypeRegistry.h"
 #include "MetasoundFrontendDataTypeTraits.h"
+#include "MetasoundFrontendModuleRegistrationMacros.h"
 #include "MetasoundFrontendRegistries.h"
 #include "MetasoundInputNode.h"
 #include "MetasoundLiteral.h"
@@ -118,17 +120,17 @@ namespace Metasound::Frontend
 
 		// This utility function can be used to optionally check to see if we can transmit a data type, and autogenerate send and receive nodes for that datatype.
 		template<typename TDataType, typename TEnableIf<TIsTransmittable<TDataType>::Value, bool>::Type = true>
-		void AttemptToRegisterSendAndReceiveNodes()
+		void AttemptToRegisterSendAndReceiveNodes(const FModuleInfo& InModuleInfo)
 		{
 			if (TEnableTransmissionNodeRegistration<TDataType>::Value)
 			{
-				ensureAlways(RegisterNode<Metasound::TSendNode<TDataType>>());
-				ensureAlways(RegisterNode<Metasound::TReceiveNode<TDataType>>());
+				RegisterNode<Metasound::TSendNode<TDataType>>(InModuleInfo);
+				RegisterNode<Metasound::TReceiveNode<TDataType>>(InModuleInfo);
 			}
 		}
 
 		template<typename TDataType, typename TEnableIf<!TIsTransmittable<TDataType>::Value, bool>::Type = true>
-		void AttemptToRegisterSendAndReceiveNodes()
+		void AttemptToRegisterSendAndReceiveNodes(const FModuleInfo& InModuleInfo)
 		{
 			// This implementation intentionally noops, because Metasound::TIsTransmittable is false for this datatype.
 			// This is either because the datatype is not trivially copyable, and thus can't be buffered between threads,
@@ -137,7 +139,7 @@ namespace Metasound::Frontend
 
 		// This utility function can be used to check to see if we can static cast between two types, and autogenerate a node for that static cast.
 		template<typename TFromDataType, typename TToDataType, typename std::enable_if<TIsAutoConvertible<TFromDataType, TToDataType>::Value, bool>::type = true>
-		void AttemptToRegisterConverter()
+		void AttemptToRegisterConverter(const FModuleInfo& InModuleInfo)
 		{
 			using FConverterNode = Metasound::TAutoConverterNode<TFromDataType, TToDataType>;
 
@@ -148,35 +150,34 @@ namespace Metasound::Frontend
 
 				if (!std::is_same<TFromDataType, TToDataType>::value && !FMetasoundFrontendRegistryContainer::Get()->IsNodeRegistered(Key))
 				{
-					ensureAlways(RegisterNode<FConverterNode>(Metadata));
+					RegisterNode<FConverterNode>(Metadata, InModuleInfo);
 					
-					bool bSucessfullyRegisteredConversionNode = RegisterConversionNode<FConverterNode, TFromDataType, TToDataType>(FConverterNode::GetInputName(), FConverterNode::GetOutputName(), Metadata);
-					ensureAlways(bSucessfullyRegisteredConversionNode);
+					RegisterConversionNode<FConverterNode, TFromDataType, TToDataType>(FConverterNode::GetInputName(), FConverterNode::GetOutputName(), Metadata);
 				}
 			}
 		}
 
 		template<typename TFromDataType, typename TToDataType, typename std::enable_if<!TIsAutoConvertible<TFromDataType, TToDataType>::Value, int>::type = 0>
-		void AttemptToRegisterConverter()
+		void AttemptToRegisterConverter(const FModuleInfo& InModuleInfo)
 		{
 			// This implementation intentionally noops, because static_cast<TFromDataType>(TToDataType&) is invalid.
 		}
 
 		// Here we attempt to infer and autogenerate conversions for basic datatypes.
 		template<typename TDataType>
-		void RegisterConverterNodes()
+		void RegisterConverterNodes(const FModuleInfo& InModuleInfo)
 		{
 			// Conversions to this data type:
-			AttemptToRegisterConverter<bool, TDataType>();
-			AttemptToRegisterConverter<int32, TDataType>();
-			AttemptToRegisterConverter<float, TDataType>();
-			AttemptToRegisterConverter<FString, TDataType>();
+			AttemptToRegisterConverter<bool, TDataType>(InModuleInfo);
+			AttemptToRegisterConverter<int32, TDataType>(InModuleInfo);
+			AttemptToRegisterConverter<float, TDataType>(InModuleInfo);
+			AttemptToRegisterConverter<FString, TDataType>(InModuleInfo);
 
 			// Conversions from this data type:
-			AttemptToRegisterConverter<TDataType, bool>();
-			AttemptToRegisterConverter<TDataType, int32>();
-			AttemptToRegisterConverter<TDataType, float>();
-			AttemptToRegisterConverter<TDataType, FString>();
+			AttemptToRegisterConverter<TDataType, bool>(InModuleInfo);
+			AttemptToRegisterConverter<TDataType, int32>(InModuleInfo);
+			AttemptToRegisterConverter<TDataType, float>(InModuleInfo);
+			AttemptToRegisterConverter<TDataType, FString>(InModuleInfo);
 		}
 
 		/** Creates the FDataTypeRegistryInfo for a data type.
@@ -387,7 +388,7 @@ namespace Metasound::Frontend
 		 * @return True on success, false on failure.
 		 */
 		template<typename TDataType, ELiteralType PreferredArgType = ELiteralType::None, typename UClassToUse = void>
-		bool RegisterDataTypeWithFrontendInternal()
+		bool RegisterDataTypeWithFrontendInternal(const FModuleInfo& InModuleInfo)
 		{
 			static constexpr bool bIsParsable = TLiteralTraits<TDataType>::bIsParsableFromAnyLiteralType;
 			static constexpr bool bIsConstructorType = TIsConstructorVertexSupported<TDataType>::Value;
@@ -466,7 +467,6 @@ namespace Metasound::Frontend
 				void InitNodeClasses()
 				{
 					// Create class info using prototype node
-					// TODO: register nodes with static class info instead of prototype instance.
 
 					if constexpr (bIsParsable)
 					{
@@ -779,12 +779,11 @@ namespace Metasound::Frontend
 			};
 
 			bool bSucceeded = Frontend::IDataTypeRegistry::Get().RegisterDataType(MakeUnique<FDataTypeRegistryEntry>());
-			ensureAlwaysMsgf(bSucceeded, TEXT("Failed to register data type %s in the node registry!"), *GetMetasoundDataTypeString<TDataType>());
 
 			if (bSucceeded)
 			{
-				RegisterConverterNodes<TDataType>();
-				AttemptToRegisterSendAndReceiveNodes<TDataType>();
+				RegisterConverterNodes<TDataType>(InModuleInfo);
+				AttemptToRegisterSendAndReceiveNodes<TDataType>(InModuleInfo);
 			}
 			
 			return bSucceeded;
@@ -802,7 +801,7 @@ namespace Metasound::Frontend
 		 * @return True on success, false on failure.
 		 */
 		template<typename TDataType, ELiteralType PreferredArgType>
-		bool RegisterDataTypeArrayWithFrontend()
+		bool RegisterDataTypeArrayWithFrontend(const FModuleInfo& InModuleInfo)
 		{
 			using namespace MetasoundDataTypeRegistrationPrivate;
 			using TArrayType = TArray<TDataType>;
@@ -810,9 +809,9 @@ namespace Metasound::Frontend
 			if constexpr (TEnableAutoArrayTypeRegistration<TDataType>::Value)
 			{
 				constexpr bool bIsArrayType = true;
-				bool bSuccess = RegisterDataTypeWithFrontendInternal<TArrayType, TLiteralArrayEnum<PreferredArgType>::Value>();
-				bSuccess = bSuccess && RegisterArrayNodes<TArrayType>();
-				bSuccess = bSuccess && RegisterDataTypeWithFrontendInternal<TVariable<TArrayType>>();
+				bool bSuccess = RegisterDataTypeWithFrontendInternal<TArrayType, TLiteralArrayEnum<PreferredArgType>::Value>(InModuleInfo);
+				bSuccess = bSuccess && RegisterArrayNodes<TArrayType>(InModuleInfo);
+				bSuccess = bSuccess && RegisterDataTypeWithFrontendInternal<TVariable<TArrayType>>(InModuleInfo);
 				return bSuccess;
 			}
 			else
@@ -820,7 +819,7 @@ namespace Metasound::Frontend
 				return true;
 			}
 		}
-	}
+	} // namespace MetasoundDataTypeRegistrationPrivate
 	
 	/** Registers a data type with the MetaSound Frontend. This allows the data type 
 	 * to be used in Input, Output, Send and Receive  nodes by informing the 
@@ -834,23 +833,38 @@ namespace Metasound::Frontend
 	 * @return True on success, false on failure.
 	 */
 	template<typename TDataType, ELiteralType PreferredArgType = ELiteralType::None, typename UClassToUse = void>
-	bool RegisterDataType()
+	bool RegisterDataType(const FModuleInfo& InModuleInfo)
 	{
 		using namespace MetasoundDataTypeRegistrationPrivate;
 
 		// Register TDataType as a metasound data type.
-		bool bSuccess = RegisterDataTypeWithFrontendInternal<TDataType, PreferredArgType, UClassToUse>();
-		ensure(bSuccess);
-		bSuccess = bSuccess && RegisterDataTypeWithFrontendInternal<TVariable<TDataType>>();
-		ensure(bSuccess);
+		bool bSuccess = RegisterDataTypeWithFrontendInternal<TDataType, PreferredArgType, UClassToUse>(InModuleInfo);
+		bSuccess = bSuccess && RegisterDataTypeWithFrontendInternal<TVariable<TDataType>>(InModuleInfo);
 
 		// Register TArray<TDataType> as a metasound data type.
-		bSuccess = bSuccess && RegisterDataTypeArrayWithFrontend<TDataType, PreferredArgType>();
-		ensure(bSuccess);
+		bSuccess = bSuccess && RegisterDataTypeArrayWithFrontend<TDataType, PreferredArgType>(InModuleInfo);
 
 		return bSuccess;
 	}
 
+	/** Registers a data type with the MetaSound Frontend. This allows the data type 
+	 * to be used in Input, Output, Send and Receive  nodes by informing the 
+	 * Frontend how to instantiate an instance. 
+	 *
+	 * @tparam TDataType - The data type to register.
+	 * @tparam PreferredArgType - The preferred constructor argument type to use when creating an instance of the data type.
+	 * @tparam UClassToUse - The preferred UObject class to use when constructing from an Audio::IProxyDataPtr. If the type is not
+	 *                       constructible with an Audio::IProxyDataPtr, then this should be void.
+	 *
+	 * @return True on success, false on failure.
+	 */
+	template<typename TDataType, ELiteralType PreferredArgType = ELiteralType::None, typename UClassToUse = void>
+	UE_DEPRECATED(5.7, "Use RegisterDataType(const FModuleInfo&) instead")
+	bool RegisterDataType()
+	{
+		return RegisterDataType<TDataType, PreferredArgType, UClassToUse>(FModuleInfo{});
+	}
+
 
 	/** Registration info for a data type.
 	 *
@@ -902,7 +916,7 @@ namespace Metasound
 
 #define ENQUEUE_METASOUND_DATATYPE_REGISTRATION_COMMAND(DataType, DataTypeName, ...) \
 	static_assert(::Metasound::Frontend::TMetasoundDataTypeRegistration<DataType>::bCanRegister, CANNOT_REGISTER_METASOUND_DATA_TYPE_ASSERT_STRING(DataType)); \
-	template<> const bool ::Metasound::Frontend::TMetasoundDataTypeRegistration<DataType>::bSuccessfullyRegistered = ::Metasound::Frontend::INodeClassRegistry::Get()->EnqueueInitCommand([](){ ::Metasound::Frontend::RegisterDataType<DataType, ##__VA_ARGS__>(); }); // This static bool is useful for debugging, but also is the only way the compiler will let us call this function outside of an expression.
+	METASOUND_IMPLEMENT_REGISTRATION_ACTION(UE_JOIN(_DataTypeRegistration,__COUNTER__), (::Metasound::Frontend::RegisterDataType<DataType, ##__VA_ARGS__>), nullptr)
 
 // This should be used to expose a datatype as a potential input or output for a metasound graph.
 // The first argument to the macro is the class to expose.
diff --git a/Engine/Plugins/Runtime/Metasound/Source/MetasoundFrontend/Public/MetasoundFrontendModuleRegistrationMacros.h b/Engine/Plugins/Runtime/Metasound/Source/MetasoundFrontend/Public/MetasoundFrontendModuleRegistrationMacros.h
new file mode 100644
--- /dev/null
+++ b/Engine/Plugins/Runtime/Metasound/Source/MetasoundFrontend/Public/MetasoundFrontendModuleRegistrationMacros.h
@@ -0,0 +1,272 @@
+// Copyright Epic Games, Inc. All Rights Reserved.
+#pragma once
+
+#include "HAL/PreprocessorHelpers.h"
+#include "UObject/NameTypes.h"
+
+#define UE_API METASOUNDFRONTEND_API
+
+/** MetaSound Frontend Registration Actions
+ *
+ * MetaSound Frontend Registration Actions are used to register and unregister
+ * MetaSound nodes and data types at the module level. 
+ *
+ * Usage:
+ *
+ * Node Developers can utilize these actions by implementing the following:
+ * - Add private definitions to the module's Build.cs file for METASOUND_PLUGIN and METASOUND_MODULE
+ *   	PrivateDefinitions.Add("METASOUND_PLUGIN=MyPlugin")
+ *   	PrivateDefinitions.Add("METASOUND_MODULE=MyModule")
+ *   These definitions must be private to the module. 
+ *
+ * - Implement a module registration group using the provided macros.
+ *   In MyModule.cpp
+ *   	METASOUND_IMPLEMENT_MODULE_REGISTRATION_LIST
+ *
+ * - Use METASOUND_REGISTER_NODE(...) and METASOUND_REGISTER_NODE_AND_CONFIGURATION(...) 
+ *   in the cpp files to setup registration actions for their nodes.
+ *   
+ * - Trigger registration in their module's StartupModule and ShutdownModule calls
+ *   void FMyModule::StartupModule()
+ *   {
+ * 		METASOUND_REGISTER_ITEMS_IN_MODULE
+ *   }
+ *   void FMyModule::ShutdownModule()
+ *   {
+ * 		METASOUND_UNREGISTER_ITEMS_IN_MODULE
+ *   }
+ *
+ *
+ * Implementation:
+ *
+ * This system uses creates a linked list of FRegistrationActions that is constructed
+ * during static initialization. Each module contains it's own linked list of FRegistrationActions
+ * and thus only performs registration and unregistration actions associated with 
+ * the module. This allows registration actions to be declared locally within the
+ * cpp file containing the node or data type, while maintaining the association 
+ * with the owning plugin and module.
+ *
+ * Each module's linked list contains a unique class name and does not export it's 
+ * API. This is in order to protect against modules accidentally adding their registration
+ * actions to the wrong linked list. The unique name for the module's list is created
+ * by combining the METASOUND_PLUGIN and METASOUND_MODULE private macro definitions.
+ * If those definitions are missing, a global linked list serves as a backup. 
+ *
+ * Thread safety for this system is provided by two mechanisms.
+ * - Initialization of the registration actions linked list is thread safe by the
+ *   fact that static initialization in C++ is single threaded. The list structure
+ *   (e.g. Head & Next pointers) should not be modified outside of static initialization.
+ *
+ * - RegisterAll and UnregisterAll should only be called during StartupModule and 
+ *   ShutdownModule because module startup and shutdown are also run single threaded. 
+ */
+
+namespace Metasound::Frontend
+{
+	// FModuleInfo captures information about the module which registers an item.
+	struct FModuleInfo
+	{
+#if WITH_EDITORONLY_DATA
+		FLazyName PluginName;
+		FLazyName ModuleName;
+#endif
+	};
+
+	namespace RegistrationPrivate
+	{
+		// Function signature of a registration function
+		typedef bool (*FRegistrationFunction)(const FModuleInfo&);
+		
+		// Registration actions build a linked list during static initialization.
+		struct FRegistrationAction
+		{
+			UE_API FRegistrationAction(FRegistrationAction*& InOutHead, FRegistrationFunction InRegister, FRegistrationFunction InUnregister);
+
+#if WITH_EDITORONLY_DATA
+		// Include debug name info in editoronly_data builds. 
+		template<size_t N>
+		FRegistrationAction(FRegistrationAction*& InOutHead, FRegistrationFunction InRegister, FRegistrationFunction InUnregister, const TCHAR(&NameLiteral)[N])
+		: FRegistrationAction(InOutHead, InRegister, InUnregister)
+		{
+			Name = NameLiteral;
+		}
+#endif
+
+		FRegistrationAction* Next = nullptr;
+		FRegistrationFunction Register = nullptr;
+		FRegistrationFunction Unregister = nullptr;
+
+#if WITH_EDITORONLY_DATA
+		const TCHAR* Name = nullptr; // Warning: Expected to be only used with string literals which provide static storage duration. 
+#endif
+	};
+
+		// FRegistrationListBase implements common functionality used by all registration
+		// list classes
+		class FRegistrationListBase
+		{
+		public:
+			// Calls "Register" all FRegistrationActions in the linked list beginning with the provided Head. 
+			UE_API static void RegisterAll(const FModuleInfo& InOwningModuleInfo, FRegistrationAction* Head);
+			// Calls "Unregister" all FRegistrationActions in the linked list beginning with the provided Head. 
+			UE_API static void UnregisterAll(const FModuleInfo& InOwningModuleInfo, FRegistrationAction* Head);
+		};
+
+		// FGlobalRegistrationList is a fallback for when the module has not defined
+		// the proper macros to create it's own registration list. This class should not
+		// be used because it cannot appropriately execute unregister actions for a 
+		// single plugin or module unloading.
+		class FGlobalRegistrationList 
+		{
+		public: 
+			static FModuleInfo ModuleInfo;
+
+			/* If you've hit this deprecation warning it is because you are using
+			 * a backwards compatible MetaSound node registration code path for your 
+			 * module. In order to better support loading and unloading of MetaSound, 
+			 * perform the following updates. This will remove this deprecation 
+			 * warning.
+			 *
+			 * - Add private definitions to the module's Build.cs file for METASOUND_PLUGIN 
+			 *   and METAOUND_MODULE
+			 *   	PrivateDefinitions.Add("METASOUND_PLUGIN=MyPlugin")
+			 *   	PrivateDefinitions.Add("METASOUND_MODULE=MyModule")
+			 *   These definitions must be private to the module. 
+			 *
+			 * - Implement a module registration group using the provided macros.
+			 *   In MyModule.cpp
+			 *   	METASOUND_IMPLEMENT_MODULE_REGISTRATION_LIST
+			 *
+			 * - Trigger registration in your module's StartupModule and ShutdownModule calls
+			 *
+			 *   void FMyModule::StartupModule()
+			 *   {
+			 * 		METASOUND_REGISTER_ITEMS_IN_MODULE
+			 *   }
+			 *   void FMyModule::ShutdownModule()
+			 *   {
+			 * 		METASOUND_UNREGISTER_ITEMS_IN_MODULE
+			 *   }
+			 *
+			 */
+			UE_DEPRECATED(5.7, "Macro based registration has been updated. See MetasoundFrontendModuleRegistrationMacros.h for a guide on defining module level registration.")
+			UE_API static FRegistrationAction*& GetHeadAction();
+
+			// Used internally to do incremental registration on a global list. 
+			static FRegistrationAction*& GetRegisterTailAction();
+
+			UE_API static void RegisterAll(const FModuleInfo& InOwningModuleInfo, FRegistrationAction* Head);
+			UE_API static void UnregisterAll(const FModuleInfo& InOwningModuleInfo, FRegistrationAction* Head);
+		};
+	} // namespace RegistrationPrivate
+} // namespace Metasound::Frontend
+
+// These macros are required in order to create module local registration actions.
+#if defined(METASOUND_MODULE) && defined(METASOUND_PLUGIN)
+
+// By default, always unregister on module shutdown. The only time we do not unregister
+// is when things cannot be associated with a module. This currently only happens
+// when using the deprecated global registration list. 
+#define METASOUND_UNREGISTER_ON_MODULE_SHUTDOWN 1
+
+// Module registration list class name
+#define METASOUND_MODULE_REGISTRATION_LIST_CLASS UE_JOIN(F, UE_JOIN(METASOUND_PLUGIN, UE_JOIN(_, UE_JOIN(METASOUND_MODULE, _Registration))))
+
+// Declare a module registration list. This needs to be visible wherever METASOUND_REGISTER_NODE
+// is called. Generally this can be a header file within the module's Private directory. 
+#define METASOUND_DEFINE_MODULE_REGISTRATION_LIST \
+namespace Metasound::Frontend::RegistrationPrivate \
+{\
+	/* Methods in this class are purposely not exported because they should only be accessed within the module. */ \
+	class METASOUND_MODULE_REGISTRATION_LIST_CLASS : public FRegistrationListBase \
+	{\
+		static_assert(sizeof(UE_STRINGIZE(METASOUND_PLUGIN)) > 0, "\"METASOUND_PLUGIN\" macro definition cannot be defined and empty"); \
+		static_assert(sizeof(UE_STRINGIZE(METASOUND_MODULE)) > 0, "\"METASOUND_MODULE\" macro definition cannot be defined and empty"); \
+		\
+	public: \
+		static const Metasound::Frontend::FModuleInfo ModuleInfo; \
+		static FRegistrationAction*& GetHeadAction(); \
+	};\
+}
+
+
+// Only allow module registration lists to be defined once.
+#ifndef METASOUND_MODULE_REGISTRATION_LIST_DEFINED
+#define METASOUND_MODULE_REGISTRATION_LIST_DEFINED
+METASOUND_DEFINE_MODULE_REGISTRATION_LIST
+#endif
+
+// ModuleRegistrationInfo only exists WITH_EDITORONLY_DATA
+#if WITH_EDITORONLY_DATA
+#define METASOUND_IMPLEMENT_MODULE_REGISTRATION_LIST_INFO \
+	const FModuleInfo METASOUND_MODULE_REGISTRATION_LIST_CLASS::ModuleInfo\
+	{\
+		UE_STRINGIZE(METASOUND_PLUGIN), \
+		UE_STRINGIZE(METASOUND_MODULE) \
+	};
+
+#else
+#define METASOUND_IMPLEMENT_MODULE_REGISTRATION_LIST_INFO const FModuleInfo METASOUND_MODULE_REGISTRATION_LIST_CLASS::ModuleInfo = {};
+#endif // if WITH_EDITORONLY_DATA
+
+// Implement a module registration list. This should be placed within a CPP file.
+#define METASOUND_IMPLEMENT_MODULE_REGISTRATION_LIST \
+namespace Metasound::Frontend::RegistrationPrivate \
+{ \
+	METASOUND_IMPLEMENT_MODULE_REGISTRATION_LIST_INFO \
+	FRegistrationAction*& METASOUND_MODULE_REGISTRATION_LIST_CLASS::GetHeadAction() \
+	{ \
+		static FRegistrationAction* Head = nullptr;  \
+		return Head; \
+	} \
+}
+
+#else // if defined(METASOUND_MODULE) && defined(METASOUND_PLUGIN)
+
+// If METASOUND_MODULE and METASOUND_PLUGIN are not both defined, we rely on the
+// global registration list class as a fallback for ensuring the registration 
+// actions still execute.
+#define METASOUND_MODULE_REGISTRATION_LIST_CLASS FGlobalRegistrationList
+
+// Disable unregistering on module shutdown because the actions are not associated with
+// a module specific registration list. If this were enabled, one module would incidentally
+// run unregister actions from another module.
+#define METASOUND_UNREGISTER_ON_MODULE_SHUTDOWN 0
+
+// The module registration lists cannot be defined or implemented. The global 
+// registration list is used instead.
+#define METASOUND_IMPLEMENT_MODULE_REGISTRATION_LIST
+
+#endif
+
+// The registration list class including global namespace access
+#define METASOUND_MODULE_REGISTRATION_LIST_CLASS_NS ::Metasound::Frontend::RegistrationPrivate::METASOUND_MODULE_REGISTRATION_LIST_CLASS
+
+// Retrieve the head action from the registration list class
+#define METASOUND_GET_HEAD_REGISTRATION_ACTION METASOUND_MODULE_REGISTRATION_LIST_CLASS_NS::GetHeadAction()
+
+// Retrieve the module info from the registration list class
+#define METASOUND_GET_MODULE_LIST_INFO METASOUND_MODULE_REGISTRATION_LIST_CLASS_NS::ModuleInfo
+
+// Register all nodes in the registration list class
+#define METASOUND_REGISTER_ITEMS_IN_MODULE METASOUND_MODULE_REGISTRATION_LIST_CLASS_NS::RegisterAll(METASOUND_GET_MODULE_LIST_INFO, METASOUND_GET_HEAD_REGISTRATION_ACTION);
+
+// Only run unregister actions if it is safe to do s.
+#if METASOUND_UNREGISTER_ON_MODULE_SHUTDOWN
+#define METASOUND_UNREGISTER_ITEMS_IN_MODULE METASOUND_MODULE_REGISTRATION_LIST_CLASS_NS::UnregisterAll(METASOUND_GET_MODULE_LIST_INFO, METASOUND_GET_HEAD_REGISTRATION_ACTION);
+#else // if METASOUND_UNREGISTER_ON_MODULE_SHUTDOWN
+#define METASOUND_UNREGISTER_ITEMS_IN_MODULE 
+#endif
+
+// Implement a registration action by declaring statically initialized FRegistrationAction
+#if WITH_EDITORONLY_DATA
+
+// With editoronly data includes additional debug logging info
+#define METASOUND_IMPLEMENT_REGISTRATION_ACTION(UniqueName, RegisterFunction, UnregisterFunction) static ::Metasound::Frontend::RegistrationPrivate::FRegistrationAction UE_JOIN(RegistrationAction,UniqueName)(METASOUND_GET_HEAD_REGISTRATION_ACTION, RegisterFunction, UnregisterFunction, TEXT(UE_STRINGIZE(UniqueName))); 
+
+#else
+
+#define METASOUND_IMPLEMENT_REGISTRATION_ACTION(UniqueName, RegisterFunction, UnregisterFunction) static ::Metasound::Frontend::RegistrationPrivate::FRegistrationAction UE_JOIN(RegistrationAction,UniqueName)(METASOUND_GET_HEAD_REGISTRATION_ACTION, RegisterFunction, UnregisterFunction); 
+#endif
+
+#undef UE_API
diff --git a/Engine/Plugins/Runtime/Metasound/Source/MetasoundFrontend/Public/MetasoundFrontendNodeClassRegistry.h b/Engine/Plugins/Runtime/Metasound/Source/MetasoundFrontend/Public/MetasoundFrontendNodeClassRegistry.h
--- a/Engine/Plugins/Runtime/Metasound/Source/MetasoundFrontend/Public/MetasoundFrontendNodeClassRegistry.h
+++ b/Engine/Plugins/Runtime/Metasound/Source/MetasoundFrontend/Public/MetasoundFrontendNodeClassRegistry.h
@@ -2,6 +2,7 @@
 #pragma once
 
 #include "MetasoundFrontendDocument.h"
+#include "MetasoundFrontendNodeMigration.h"
 #include "MetasoundFrontendRegistryKey.h"
 #include "MetasoundNodeInterface.h"
 #include "MetasoundNodeConstructorParams.h"
@@ -112,6 +113,14 @@ namespace Metasound::Frontend
 #else
 		UE_API virtual TInstancedStruct<FMetaSoundFrontendNodeConfiguration> CreateFrontendNodeConfiguration() const;
 #endif
+
+#if WITH_EDITORONLY_DATA
+		/** Name of the plugin registering this node. */
+		UE_API virtual FName GetPluginName() const;
+		
+		/** Name of the module registering this node. */
+		UE_API virtual FName GetModuleName() const;
+#endif
 	};
 
 	using INodeRegistryEntry UE_DEPRECATED(5.6, "Use INodeClassRegistryEntry instead") = INodeClassRegistryEntry;
@@ -173,6 +182,8 @@ namespace Metasound::Frontend
 		{
 			NodeRegistration,     //< Something was added to the registry.
 			NodeUnregistration,  //< Something was removed from the registry.
+			NodeMigrationRegistration, //< Something migrated.
+			NodeMigrationUnregistration, //< Module doing migration is unloaded.
 			Invalid
 		};
 
@@ -192,6 +203,16 @@ namespace Metasound::Frontend
 					return TEXT("Node Unregistration");
 				}
 
+				case ETransactionType::NodeMigrationRegistration:
+				{
+					return TEXT("Node Migration Registration");
+				}
+
+				case ETransactionType::NodeMigrationUnregistration:
+				{
+					return TEXT("Node Migration Unregistration");
+				}
+
 				case ETransactionType::Invalid:
 				{
 					return TEXT("Invalid");
@@ -254,6 +275,7 @@ namespace Metasound::Frontend
 
 
 		static UE_API INodeClassRegistry* Get();
+		static UE_API INodeClassRegistry& GetChecked();
 
 		UE_DEPRECATED(5.6, "This function is no longer supported")
 		static UE_API void ShutdownMetasoundFrontend();
@@ -275,11 +297,13 @@ namespace Metasound::Frontend
 
 		// Enqueue and init command for registering a node or data type.
 		// The command queue will be processed on module init or when calling `RegisterPendingNodes()`
+		UE_DEPRECATED(5.7, "Enqueuing init commands is no longer support. See MetasoundFrontendModuleRegistrationMacros.h for information on queuing up registeration actions.")
 		virtual bool EnqueueInitCommand(TUniqueFunction<void()>&& InFunc) = 0;
 
 		virtual void SetObjectReferencer(TUniquePtr<IObjectReferencer> InReferencer) = 0;
 
 		// This is called on module startup. This invokes any registration commands enqueued by our registration macros.
+		UE_DEPRECATED(5.7, "Replace with METASOUND_REGISTER_ITEMS_IN_MODULE. See MetasoundFrontendModuleRegistrationMacros.h for more info on triggering registration of nodes and data types.")
 		virtual void RegisterPendingNodes() = 0;
 
 		// Wait for async graph registration to complete for a specific graph
@@ -308,6 +332,14 @@ namespace Metasound::Frontend
 		 */
 		virtual bool UnregisterNode(const FNodeClassRegistryKey& InKey) = 0;
 
+#if WITH_EDITORONLY_DATA
+		/** Register a node migration from one plugin and/or module to another. */
+		virtual bool RegisterNodeMigration(const FNodeMigrationInfo& InMigrationInfo) = 0;
+		
+		/** Unregister a node migration from one plugin and/or module to another. This should be used during module shutdown. */
+		virtual bool UnregisterNodeMigration(const FNodeMigrationInfo& InMigrationInfo) = 0;
+#endif
+
 		/** Returns true if the provided registry key corresponds to a valid registered node. */
 		virtual bool IsNodeRegistered(const FNodeClassRegistryKey& InKey) const = 0;
 
diff --git a/Engine/Plugins/Runtime/Metasound/Source/MetasoundFrontend/Public/MetasoundFrontendNodeMigration.h b/Engine/Plugins/Runtime/Metasound/Source/MetasoundFrontend/Public/MetasoundFrontendNodeMigration.h
new file mode 100644
--- /dev/null
+++ b/Engine/Plugins/Runtime/Metasound/Source/MetasoundFrontend/Public/MetasoundFrontendNodeMigration.h
@@ -0,0 +1,40 @@
+// Copyright Epic Games, Inc. All Rights Reserved.
+#pragma once
+
+#include "UObject/NameTypes.h"
+#include "MetasoundNodeInterface.h"
+
+#define UE_API METASOUNDFRONTEND_API
+
+namespace Metasound::Frontend
+{
+#if WITH_EDITORONLY_DATA
+	struct FNodeMigrationInfo
+	{
+		// UE Version when migration occurred (e.g. "5.6")
+		FName UEVersion;
+
+		// ClassName of node which was migrated
+		FNodeClassName ClassName;
+
+		// Major version of node when migration happened
+		int32 MajorVersion = -1;
+
+		// Minor version of node when migration happened
+		int32 MinorVersion = -1;
+
+		// Prior plugin and module where node was registered. 
+		FName FromPlugin;
+		FName FromModule;
+
+		// New plugin and module where node is registered. 
+		FName ToPlugin;
+		FName ToModule;
+
+		UE_API FString ToString() const;
+		friend bool operator==(const FNodeMigrationInfo& InLHS, const FNodeMigrationInfo& InRHS);
+	};
+#endif // if WITH_EDITORONLY_DATA
+
+} // namespace Metasound::Frontend
+#undef UE_API
diff --git a/Engine/Plugins/Runtime/Metasound/Source/MetasoundFrontend/Public/MetasoundNodeRegistrationMacro.h b/Engine/Plugins/Runtime/Metasound/Source/MetasoundFrontend/Public/MetasoundNodeRegistrationMacro.h
--- a/Engine/Plugins/Runtime/Metasound/Source/MetasoundFrontend/Public/MetasoundNodeRegistrationMacro.h
+++ b/Engine/Plugins/Runtime/Metasound/Source/MetasoundFrontend/Public/MetasoundNodeRegistrationMacro.h
@@ -8,6 +8,7 @@
 #include "MetasoundFrontendDocument.h"
 #include "MetasoundFrontendNodeClassRegistry.h"
 #include "MetasoundLog.h"
+#include "MetasoundFrontendModuleRegistrationMacros.h"
 #include "MetasoundNodeInterface.h"
 
 #include "Traits/MetasoundNodeConstructorTraits.h"
@@ -24,13 +25,14 @@
 
 namespace Metasound::Frontend
 {
+
 	namespace NodeRegistrationPrivate
 	{
 		// Utilize base class to reduce template bloat in TNodeRegistryEntry
 		class FNodeRegistryEntryBase : public INodeClassRegistryEntry
 		{
 		public:
-			UE_API FNodeRegistryEntryBase(const Metasound::FNodeClassMetadata& InMetadata);
+			UE_API FNodeRegistryEntryBase(const Metasound::FNodeClassMetadata& InMetadata, const FModuleInfo& InOwningModuleInfo);
 
 			virtual ~FNodeRegistryEntryBase() = default;
 
@@ -50,14 +52,25 @@ namespace Metasound::Frontend
 
 			UE_API virtual TInstancedStruct<FMetaSoundFrontendNodeConfiguration> CreateFrontendNodeConfiguration() const override;
 
+#if WITH_EDITORONLY_DATA
+			UE_API virtual FName GetPluginName() const override;
+
+			UE_API virtual FName GetModuleName() const override;
+#endif
+
 		protected:
 
 			UE_API TSharedRef<const FNodeClassMetadata> GetNodeClassMetadata() const;
-		private:
 
+		private:
+			
 			TSharedRef<FNodeClassMetadata> ClassMetadata;
 			FNodeClassInfo ClassInfo;
 			FMetasoundFrontendClass FrontendClass;
+#if WITH_EDITORONLY_DATA
+			FName PluginName;
+			FName ModuleName;
+#endif
 		};
 
 		template<typename TNodeType>
@@ -173,31 +186,40 @@ namespace Metasound::Frontend
 		};
 
 		UE_API bool RegisterNodeInternal(TUniquePtr<INodeClassRegistryEntry>&& InEntry);
+		UE_API bool UnregisterNodeInternal(const FNodeClassMetadata& InMetadata);
 	} // namespace NodeRegistrationPrivate
 
-
-	template <typename TNodeType, typename ConfigurationType=void>
-	bool RegisterNode(const FNodeClassMetadata& InMetadata)
+	template<typename TNodeType, typename ConfigurationType=void>
+	bool RegisterNode(const FNodeClassMetadata& InMetadata, const FModuleInfo& InOwningModuleInfo)
 	{
 		using namespace NodeRegistrationPrivate;
 		if constexpr (TIsOnlyDeprecatedNodeConstructorProvided<TNodeType>::Value)
 		{
-			return RegisterNode<TDeprecatedNodeAPIAdapter<TNodeType>, ConfigurationType>(InMetadata);
+			return RegisterNode<TDeprecatedNodeAPIAdapter<TNodeType>, ConfigurationType>(InMetadata, InOwningModuleInfo);
 		}
 		else
 		{
-			return NodeRegistrationPrivate::RegisterNodeInternal(MakeUnique<TNodeRegistryEntry<TNodeType, ConfigurationType>>(InMetadata));
+			return NodeRegistrationPrivate::RegisterNodeInternal(MakeUnique<TNodeRegistryEntry<TNodeType, ConfigurationType>>(InMetadata, InOwningModuleInfo));
 		}
+
+	}
+ 
+
+	template <typename TNodeType, typename ConfigurationType=void>
+	UE_DEPRECATED(5.7, "Use RegisterNode(...) which provides FModuleInfo")
+	bool RegisterNode(const FNodeClassMetadata& InMetadata)
+	{
+		return RegisterNode<TNodeType, ConfigurationType>(InMetadata, FModuleInfo{});
 	} 
 
 	template <typename TNodeType, typename ConfigurationType=void>
-	bool RegisterNode()
+	bool RegisterNode(const FModuleInfo& InOwningModuleInfo)
 	{
 		static_assert(::Metasound::TIsNodeConstructorSupported<TNodeType>::Value, "In order to be registered as a MetaSound node, the node needs to implement the following public constructor: Constructor(Metasound::FNodeData InNodeData) or Construct(Metasound::FNodeData InNodeData, TSharedRef<const FNodeClassMetadata> InMetadata)"); 
 
 		if constexpr(TIsCreateNodeClassMetadataDeclared<TNodeType>::Value)
 		{
-			return RegisterNode<TNodeType, ConfigurationType>(TNodeType::CreateNodeClassMetadata());
+			return RegisterNode<TNodeType, ConfigurationType>(TNodeType::CreateNodeClassMetadata(), InOwningModuleInfo);
 		}
 		else
 		{
@@ -206,17 +228,25 @@ namespace Metasound::Frontend
 			// Register a node using a prototype node.
 			FNodeInitData InitData;
 			TUniquePtr<Metasound::INode> Node = MakeUnique<TNodeType>(InitData);
-			return RegisterNode<TNodeType, ConfigurationType>(Node->GetMetadata());
+			return RegisterNode<TNodeType, ConfigurationType>(Node->GetMetadata(), InOwningModuleInfo);
 		}
 	}
 
+	template <typename TNodeType, typename ConfigurationType=void>
+	UE_DEPRECATED(5.7, "Use RegisterNode(...) which provides FModuleInfo")
+	bool RegisterNode()
+	{
+		return RegisterNode<TNodeType, ConfigurationType>(FModuleInfo{});
+	}
+
+
 	template<typename TNodeType>
-	bool UnregisterNode()
+	bool UnregisterNode(const FModuleInfo& InOwningModuleInfo)
 	{
-		FNodeClassRegistryKey RegistryKey;
+		using namespace NodeRegistrationPrivate;
 		if constexpr(TIsCreateNodeClassMetadataDeclared<TNodeType>::Value)
 		{
-			RegistryKey = FNodeClassRegistryKey(TNodeType::CreateNodeClassMetadata());
+			return UnregisterNodeInternal(TNodeType::CreateNodeClassMetadata());
 		}
 		else
 		{
@@ -225,31 +255,16 @@ namespace Metasound::Frontend
 			// Register a node using a prototype node.
 			FNodeInitData InitData;
 			TUniquePtr<Metasound::INode> Node = MakeUnique<TNodeType>(InitData);
-			RegistryKey = FNodeClassRegistryKey(Node->GetMetadata());
+			return UnregisterNodeInternal(Node->GetMetadata());
 		}
-
-		return INodeClassRegistry::Get()->UnregisterNode(RegistryKey);
 	}
 
-	namespace NodeRegistrationPrivate
+	template<typename TNodeType>
+	UE_DEPRECATED(5.7, "Use UnregisterNode(...) which provides FModuleInfo")
+	bool UnregisterNode()
 	{
-		UE_API bool RegisterNodeUsingMacroCheck(bool& bAlreadyRegistered);
-		UE_API bool RegisterNodeDeferred(void (*InRegisterFunc)());
-
-		template<typename NodeType, typename ConfigurationType=void>
-		void RegisterNodeUsingMacro()
-		{
-			// if we reenter this code (because METASOUND_REGISTER_NODE was called twice with the same node type),
-			// we catch it here.
-			static bool bAlreadyRegisteredThisDataType = false;
-			if (!RegisterNodeUsingMacroCheck(bAlreadyRegisteredThisDataType))
-			{
-				return;
-			}
-
-			RegisterNode<NodeType, ConfigurationType>();
-		}
-	} // namespace NodeRegistrationPrivate
+		return UnregisterNode<TNodeType>(FModuleInfo{});
+	}
 }
 
 namespace Metasound
@@ -287,11 +302,11 @@ namespace Metasound
 #define METASOUND_REGISTER_NODE_AND_CONFIGURATION(NodeClass, ConfigurationClass) \
 	 static_assert(std::is_base_of<::Metasound::INodeBase, NodeClass>::value, "To be registered as a  Metasound Node," #NodeClass "need to be a derived class from Metasound::INodeBase, Metasound::INode, or Metasound::FNode."); \
 	 static_assert(::Metasound::TIsNodeConstructorSupported<NodeClass>::Value, "In order to be registered as a Metasound Node, " #NodeClass " needs to implement the following public constructor: " #NodeClass "(Metasound::FNodeData InNodeData);"); \
-	 static bool bSuccessfullyRegistered##NodeClass = \
-		/* This static bool is useful for debugging, but also is the only way the compiler will let us call this function outside of an expression. */ \
-		::Metasound::Frontend::NodeRegistrationPrivate::RegisterNodeDeferred(::Metasound::Frontend::NodeRegistrationPrivate::RegisterNodeUsingMacro<NodeClass, ConfigurationClass>);
+	 METASOUND_IMPLEMENT_REGISTRATION_ACTION(NodeClass, (::Metasound::Frontend::RegisterNode<NodeClass, ConfigurationClass>), (::Metasound::Frontend::UnregisterNode<NodeClass>));
+
 
 #define METASOUND_REGISTER_NODE(NodeClass) METASOUND_REGISTER_NODE_AND_CONFIGURATION(NodeClass, void)
+
 /*
 Macros to help define various FText node fields.
 */
diff --git a/Engine/Plugins/Runtime/Metasound/Source/MetasoundStandardNodes/MetasoundStandardNodes.Build.cs b/Engine/Plugins/Runtime/Metasound/Source/MetasoundStandardNodes/MetasoundStandardNodes.Build.cs
--- a/Engine/Plugins/Runtime/Metasound/Source/MetasoundStandardNodes/MetasoundStandardNodes.Build.cs
+++ b/Engine/Plugins/Runtime/Metasound/Source/MetasoundStandardNodes/MetasoundStandardNodes.Build.cs
@@ -31,6 +31,13 @@ namespace UnrealBuildTool.Rules
 			);
 
 			NumIncludedBytesPerUnityCPPOverride = 120 * 1024;
+
+			PrivateDefinitions.AddRange(
+				new string[]
+				{
+					"METASOUND_PLUGIN=Metasound",
+					"METASOUND_MODULE=MetasoundStandardNodes"
+				});
 		}
 	}
 }
diff --git a/Engine/Plugins/Runtime/Metasound/Source/MetasoundStandardNodes/Private/MetasoundConversionNodes.cpp b/Engine/Plugins/Runtime/Metasound/Source/MetasoundStandardNodes/Private/MetasoundConversionNodes.cpp
--- a/Engine/Plugins/Runtime/Metasound/Source/MetasoundStandardNodes/Private/MetasoundConversionNodes.cpp
+++ b/Engine/Plugins/Runtime/Metasound/Source/MetasoundStandardNodes/Private/MetasoundConversionNodes.cpp
@@ -15,7 +15,8 @@
 
 #define LOCTEXT_NAMESPACE "MetasoundStandardNodes_Conversion"
 
-#define METASOUND_REGISTER_CONVERSION(NodeClass, FromDataType, ToDataType) static bool bSuccessfullyRegisteredConversion##NodeClass = FMetasoundFrontendRegistryContainer::Get()->EnqueueInitCommand([]() { Metasound::RegisterConversionOperator<FromDataType, ToDataType>(); });
+#define METASOUND_REGISTER_CONVERSION(NodeClass, FromDataType, ToDataType) \
+	METASOUND_IMPLEMENT_REGISTRATION_ACTION(UE_JOIN(Conversion,NodeClass), (&Metasound::RegisterConversionOperator<FromDataType, ToDataType>), nullptr)
 
 namespace Metasound
 {
@@ -299,7 +300,7 @@ namespace Metasound
 
 
 	template<typename TFromType, typename TToType>
-	void RegisterConversionOperator()
+	bool RegisterConversionOperator(const Frontend::FModuleInfo& InModuleInfo)
 	{
 		using FConverterNodeRegistryKey = ::Metasound::Frontend::FConverterNodeRegistryKey;
 		using FConverterNodeInfo = ::Metasound::Frontend::FConverterNodeInfo;
@@ -318,7 +319,7 @@ namespace Metasound
 			Metasound::Frontend::FNodeRegistryKey(Metadata)
 		};
 
-		FMetasoundFrontendRegistryContainer::Get()->RegisterConversionNode(RegistryKey, ConverterNodeInfo);
+		return FMetasoundFrontendRegistryContainer::Get()->RegisterConversionNode(RegistryKey, ConverterNodeInfo);
 	}
 
 	METASOUND_REGISTER_CONVERSION(FConversionFloatToTime, float, FTime)
diff --git a/Engine/Plugins/Runtime/Metasound/Source/MetasoundStandardNodes/Private/MetasoundStandardNodesModule.cpp b/Engine/Plugins/Runtime/Metasound/Source/MetasoundStandardNodes/Private/MetasoundStandardNodesModule.cpp
--- a/Engine/Plugins/Runtime/Metasound/Source/MetasoundStandardNodes/Private/MetasoundStandardNodesModule.cpp
+++ b/Engine/Plugins/Runtime/Metasound/Source/MetasoundStandardNodes/Private/MetasoundStandardNodesModule.cpp
@@ -1,25 +1,66 @@
 // Copyright Epic Games, Inc. All Rights Reserved.
 
 #include "CoreMinimal.h"
-#include "MetasoundGain.h"
-#include "MetasoundGraphCoreModule.h"
-#include "MetasoundPrimitives.h"
-#include "MetasoundTime.h"
-#include "MetasoundTrigger.h"
+#include "MetasoundFrontendModuleRegistrationMacros.h"
+#include "MetasoundFrontendNodeClassRegistry.h"
+#include "MetasoundFrontendNodeMigration.h"
 #include "Modules/ModuleInterface.h"
 #include "Modules/ModuleManager.h"
 
 
 namespace Metasound 
 {
-	class FMetasoundStandardNodesModule : public IMetasoundGraphCoreModule
+	class FMetasoundStandardNodesModule : public IModuleInterface
 	{
 		virtual void StartupModule() override
 		{
-			FMetasoundFrontendRegistryContainer::Get()->RegisterPendingNodes();
+			using namespace Frontend;
+			METASOUND_REGISTER_ITEMS_IN_MODULE
+
+#if 0
+			// Example of node migration
+			INodeClassRegistry& NodeRegistry = INodeClassRegistry::GetChecked();
+			NodeRegistry.RegisterNodeMigration(
+				FNodeMigrationInfo
+				{
+					"5.7",
+					FNodeClassName {"Dummy", "Dummy", "Dummy"},
+					1,
+					0,
+					UE_STRINGIZE(METASOUND_PLUGIN),
+					UE_STRINGIZE(METASOUND_MODULE),
+					"MetasoundExperimental",
+					"MetasoundRuntime"
+				});
+#endif
+		}
+
+		virtual void ShutdownModule() override
+		{
+			using namespace Frontend;
+			METASOUND_UNREGISTER_ITEMS_IN_MODULE
+			
+
+#if 0
+			// Example of node migration
+			INodeClassRegistry& NodeRegistry = INodeClassRegistry::GetChecked();
+			NodeRegistry.UnregisterNodeMigration(
+				FNodeMigrationInfo
+				{
+					"5.7",
+					FNodeClassName {"Dummy", "Dummy", "Dummy"},
+					1,
+					0,
+					UE_STRINGIZE(METASOUND_PLUGIN),
+					UE_STRINGIZE(METASOUND_MODULE),
+					"MetasoundExperimental",
+					"MetasoundRuntime"
+				});
+#endif
 		}
 	};
 }
 
+METASOUND_IMPLEMENT_MODULE_REGISTRATION_LIST
 IMPLEMENT_MODULE(Metasound::FMetasoundStandardNodesModule, MetasoundStandardNodes);
 
