MetaSound Module Registration Bughawk Fixes
#rnx

diff --git a/Engine/Plugins/Runtime/Metasound/Source/MetasoundFrontend/Public/MetasoundDataTypeRegistrationMacro.h b/Engine/Plugins/Runtime/Metasound/Source/MetasoundFrontend/Public/MetasoundDataTypeRegistrationMacro.h
--- a/Engine/Plugins/Runtime/Metasound/Source/MetasoundFrontend/Public/MetasoundDataTypeRegistrationMacro.h
+++ b/Engine/Plugins/Runtime/Metasound/Source/MetasoundFrontend/Public/MetasoundDataTypeRegistrationMacro.h
@@ -878,8 +878,8 @@ namespace Metasound::Frontend
 		// To register a data type, an input node must be able to instantiate it.
 		static constexpr bool bCanRegister = TInputNode<DataType, EVertexAccessType::Reference>::bCanRegister;
 
-		// If a data type has been successfully registered, this will be true.
-		static const bool bSuccessfullyRegistered;
+		// This is no longer used.
+		static const bool bSuccessfullyRegistered = false;
 	};
 }
 
diff --git a/Engine/Plugins/Runtime/Metasound/Source/MetasoundFrontend/Public/MetasoundFrontendModuleRegistrationMacros.h b/Engine/Plugins/Runtime/Metasound/Source/MetasoundFrontend/Public/MetasoundFrontendModuleRegistrationMacros.h
--- a/Engine/Plugins/Runtime/Metasound/Source/MetasoundFrontend/Public/MetasoundFrontendModuleRegistrationMacros.h
+++ b/Engine/Plugins/Runtime/Metasound/Source/MetasoundFrontend/Public/MetasoundFrontendModuleRegistrationMacros.h
@@ -128,7 +128,7 @@ namespace Metasound::Frontend
 			 * warning.
 			 *
 			 * - Add private definitions to the module's Build.cs file for METASOUND_PLUGIN 
-			 *   and METAOUND_MODULE
+			 *   and METASOUND_MODULE
 			 *   	PrivateDefinitions.Add("METASOUND_PLUGIN=MyPlugin")
 			 *   	PrivateDefinitions.Add("METASOUND_MODULE=MyModule")
 			 *   These definitions must be private to the module. 
@@ -180,8 +180,8 @@ namespace Metasound::Frontend::RegistrationPrivate \
 	/* Methods in this class are purposely not exported because they should only be accessed within the module. */ \
 	class METASOUND_MODULE_REGISTRATION_LIST_CLASS : public FRegistrationListBase \
 	{\
-		static_assert(sizeof(UE_STRINGIZE(METASOUND_PLUGIN)) > 0, "\"METASOUND_PLUGIN\" macro definition cannot be defined and empty"); \
-		static_assert(sizeof(UE_STRINGIZE(METASOUND_MODULE)) > 0, "\"METASOUND_MODULE\" macro definition cannot be defined and empty"); \
+		static_assert(sizeof(UE_STRINGIZE(METASOUND_PLUGIN)) > 1, "\"METASOUND_PLUGIN\" macro definition cannot be empty"); \
+		static_assert(sizeof(UE_STRINGIZE(METASOUND_MODULE)) > 1, "\"METASOUND_MODULE\" macro definition cannot be empty"); \
 		\
 	public: \
 		static const Metasound::Frontend::FModuleInfo ModuleInfo; \
