CIS fix. quiet CIS error while determining content which causes the issue
#rnx

diff --git a/Engine/Plugins/Runtime/Metasound/Source/MetasoundFrontend/Private/MetasoundFrontendDataTypeRegistry.cpp b/Engine/Plugins/Runtime/Metasound/Source/MetasoundFrontend/Private/MetasoundFrontendDataTypeRegistry.cpp
--- a/Engine/Plugins/Runtime/Metasound/Source/MetasoundFrontend/Private/MetasoundFrontendDataTypeRegistry.cpp
+++ b/Engine/Plugins/Runtime/Metasound/Source/MetasoundFrontend/Private/MetasoundFrontendDataTypeRegistry.cpp
@@ -1508,10 +1508,12 @@ PRAGMA_ENABLE_DEPRECATION_WARNINGS
 					{
 						return &Entry->Get();
 					}
+					/* TODO: temporarily quieting this error.
 					else
 					{
 						UE_LOG(LogMetaSound, Error, TEXT("Data type is not registered [Name:%s]"), *InDataTypeName.ToString());
 					}
+					*/
 				}
 
 				return nullptr;
