Do not register arrays of variables
Create vertex interfaces in CPP instead of template.

#rnx

diff --git a/Engine/Plugins/Runtime/Metasound/Source/MetasoundFrontend/Private/MetasoundArrayNodes.cpp b/Engine/Plugins/Runtime/Metasound/Source/MetasoundFrontend/Private/MetasoundArrayNodes.cpp
--- a/Engine/Plugins/Runtime/Metasound/Source/MetasoundFrontend/Private/MetasoundArrayNodes.cpp
+++ b/Engine/Plugins/Runtime/Metasound/Source/MetasoundFrontend/Private/MetasoundArrayNodes.cpp
@@ -32,6 +32,102 @@ namespace Metasound
 
 			return Metadata;
 		}
+
+		FVertexInterface CreateArrayNumInterface(const FName& InArrayDataTypeName)
+		{
+			using namespace ArrayNodeVertexNames;
+
+			return FVertexInterface {
+				FInputVertexInterface(
+					FInputDataVertex(METASOUND_GET_PARAM_NAME(InputArray), InArrayDataTypeName, METASOUND_GET_PARAM_METADATA(InputArray))
+				),
+				FOutputVertexInterface(
+					TOutputDataVertex<int32>(METASOUND_GET_PARAM_NAME_AND_METADATA(OutputNum))
+				)
+			};
+		}
+
+		FVertexInterface CreateArrayGetInterface(const FName& InArrayDataTypeName, const FName& InDataTypeName)
+		{
+			using namespace ArrayNodeVertexNames;
+			return FVertexInterface {
+				FInputVertexInterface(
+					TInputDataVertex<FTrigger>(METASOUND_GET_PARAM_NAME_AND_METADATA(InputTriggerGet)),
+					FInputDataVertex(METASOUND_GET_PARAM_NAME(InputArray), InArrayDataTypeName, METASOUND_GET_PARAM_METADATA(InputArray)),
+					TInputDataVertex<int32>(METASOUND_GET_PARAM_NAME_AND_METADATA(InputIndex))
+				),
+				FOutputVertexInterface(
+					FOutputDataVertex(METASOUND_GET_PARAM_NAME(OutputValue), InDataTypeName, METASOUND_GET_PARAM_METADATA(OutputValue), EVertexAccessType::Reference)
+				)
+			};
+		}
+
+		FVertexInterface CreateArraySetInterface(const FName& InArrayDataTypeName, const FName& InDataTypeName)
+		{
+			using namespace ArrayNodeVertexNames;
+			return FVertexInterface{ 
+				FInputVertexInterface(
+					TInputDataVertex<FTrigger>(METASOUND_GET_PARAM_NAME_AND_METADATA(InputTriggerSet)),
+					FInputDataVertex(InputInitialArrayName, InArrayDataTypeName, FDataVertexMetadata { InputInitialArrayTooltip, InputInitialArrayDisplayName }),
+					TInputDataVertex<int32>(METASOUND_GET_PARAM_NAME_AND_METADATA(InputIndex)),
+					FInputDataVertex(METASOUND_GET_PARAM_NAME(InputValue), InDataTypeName, METASOUND_GET_PARAM_METADATA(InputValue))
+				),
+				FOutputVertexInterface(
+					FOutputDataVertex(METASOUND_GET_PARAM_NAME(OutputArraySet), InArrayDataTypeName, METASOUND_GET_PARAM_METADATA(OutputArraySet), EVertexAccessType::Reference)
+
+				)
+			};
+		}
+
+		FVertexInterface CreateArrayConcatInterface(const FName& InArrayDataTypeName)
+		{
+			using namespace ArrayNodeVertexNames;
+
+			return FVertexInterface{
+				FInputVertexInterface(
+					TInputDataVertex<FTrigger>(METASOUND_GET_PARAM_NAME_AND_METADATA(InputTriggerGet)),
+					FInputDataVertex(METASOUND_GET_PARAM_NAME(InputLeftArray), InArrayDataTypeName,METASOUND_GET_PARAM_METADATA(InputLeftArray)),
+					FInputDataVertex(METASOUND_GET_PARAM_NAME(InputRightArray), InArrayDataTypeName, METASOUND_GET_PARAM_METADATA(InputRightArray))
+				),
+				FOutputVertexInterface(
+					FOutputDataVertex(METASOUND_GET_PARAM_NAME(OutputArrayConcat), InArrayDataTypeName, METASOUND_GET_PARAM_METADATA(OutputArrayConcat), EVertexAccessType::Reference)
+				)
+			};
+		}
+
+		FVertexInterface CreateArraySubsetInterface(const FName& InArrayDataTypeName)
+		{
+			using namespace ArrayNodeVertexNames;
+
+			return FVertexInterface {
+				FInputVertexInterface(
+					TInputDataVertex<FTrigger>(METASOUND_GET_PARAM_NAME_AND_METADATA(InputTriggerGet)),
+					FInputDataVertex(METASOUND_GET_PARAM_NAME(InputArray), InArrayDataTypeName, METASOUND_GET_PARAM_METADATA(InputArray)),
+
+					TInputDataVertex<int32>(METASOUND_GET_PARAM_NAME_AND_METADATA(InputStartIndex)),
+					TInputDataVertex<int32>(METASOUND_GET_PARAM_NAME_AND_METADATA(InputEndIndex))
+
+				),
+				FOutputVertexInterface(
+					FOutputDataVertex(METASOUND_GET_PARAM_NAME(OutputArraySubset), InArrayDataTypeName, METASOUND_GET_PARAM_METADATA(OutputArraySubset), EVertexAccessType::Reference)
+
+				)
+			};
+		}
+
+		FVertexInterface CreateArrayLastIndexInterface(const FName& InArrayDataTypeName)
+		{
+			using namespace ArrayNodeVertexNames;
+
+			return FVertexInterface {
+				FInputVertexInterface(
+					FInputDataVertex(METASOUND_GET_PARAM_NAME(InputArray), InArrayDataTypeName, METASOUND_GET_PARAM_METADATA(InputArray))
+				),
+				FOutputVertexInterface(
+					TOutputDataVertex<int32>(METASOUND_GET_PARAM_NAME_AND_METADATA(OutputLastIndex))
+				)
+			};
+		}
 	}
 }
 
diff --git a/Engine/Plugins/Runtime/Metasound/Source/MetasoundFrontend/Public/MetasoundArrayNodes.h b/Engine/Plugins/Runtime/Metasound/Source/MetasoundFrontend/Public/MetasoundArrayNodes.h
--- a/Engine/Plugins/Runtime/Metasound/Source/MetasoundFrontend/Public/MetasoundArrayNodes.h
+++ b/Engine/Plugins/Runtime/Metasound/Source/MetasoundFrontend/Public/MetasoundArrayNodes.h
@@ -28,6 +28,13 @@ namespace Metasound
 		// Convenience function for make FNodeClassMetadata of array nodes.
 		METASOUNDFRONTEND_API FNodeClassMetadata CreateArrayNodeClassMetadata(const FName& InDataTypeName, const FName& InOperatorName, const FText& InDisplayName, const FText& InDescription, const FVertexInterface& InDefaultInterface, int32 MajorVersion=1, int32 MinorVersion=0, bool bIsDeprecated=false);
 
+		METASOUNDFRONTEND_API FVertexInterface CreateArrayNumInterface(const FName& InArrayDataTypeName);
+		METASOUNDFRONTEND_API FVertexInterface CreateArrayGetInterface(const FName& InArrayDataTypeName, const FName& InDataTypeName);
+		METASOUNDFRONTEND_API FVertexInterface CreateArraySetInterface(const FName& InArrayDataTypeName, const FName& InDataTypeName);
+		METASOUNDFRONTEND_API FVertexInterface CreateArraySubsetInterface(const FName& InArrayDataTypeName);
+		METASOUNDFRONTEND_API FVertexInterface CreateArrayConcatInterface(const FName& InArrayDataTypeName);
+		METASOUNDFRONTEND_API FVertexInterface CreateArrayLastIndexInterface(const FName& InArrayDataTypeName);
+
 		// Retrieve the ElementType from an ArrayType
 		template<typename ArrayType>
 		struct TArrayElementType
@@ -72,6 +79,7 @@ namespace Metasound
 		METASOUND_PARAM(OutputLastIndex, "Last Index", "Last index of the array.")
 	};
 
+
 	/** TArrayNumOperator gets the number of elements in an Array. The operator
 	 * uses the FNodeFacade and defines the vertex, metadata and vertex interface
 	 * statically on the operator class. */
@@ -81,23 +89,6 @@ namespace Metasound
 	public:
 		using FArrayDataReadReference = TDataReadReference<ArrayType>;
 
-		// Declare the vertex interface
-		static const FVertexInterface& GetDefaultInterface()
-		{
-			using namespace ArrayNodeVertexNames;
-
-			static const FVertexInterface DefaultInterface(
-				FInputVertexInterface(
-					TInputDataVertex<ArrayType>(METASOUND_GET_PARAM_NAME_AND_METADATA(InputArray))
-				),
-				FOutputVertexInterface(
-					TOutputDataVertex<int32>(METASOUND_GET_PARAM_NAME_AND_METADATA(OutputNum))
-				)
-			);
-
-			return DefaultInterface;
-		}
-
 		static const FNodeClassMetadata& GetNodeInfo()
 		{
 			auto CreateNodeClassMetadata = []() -> FNodeClassMetadata
@@ -106,7 +97,7 @@ namespace Metasound
 				const FName OperatorName = TEXT("Num");
 				const FText NodeDisplayName = METASOUND_LOCTEXT_FORMAT("ArrayOpArrayNumDisplayNamePattern", "Num ({0})", GetMetasoundDataTypeDisplayText<ArrayType>());
 				const FText NodeDescription = METASOUND_LOCTEXT("ArrayOpArrayNumDescription", "Number of elements in the array");
-				const FVertexInterface NodeInterface = GetDefaultInterface();
+				const FVertexInterface NodeInterface = MetasoundArrayNodesPrivate::CreateArrayNumInterface(DataTypeName);
 
 				FNodeClassMetadata NodeClassMetadata = MetasoundArrayNodesPrivate::CreateArrayNodeClassMetadata(DataTypeName, OperatorName, NodeDisplayName, NodeDescription, NodeInterface);
 				NodeClassMetadata.Keywords.Append({
@@ -144,18 +135,6 @@ namespace Metasound
 
 		virtual ~TArrayNumOperator() = default;
 
-		virtual FDataReferenceCollection GetInputs() const override
-		{
-			checkNoEntry();
-			return {};
-		}
-
-		virtual FDataReferenceCollection GetOutputs() const override
-		{
-			checkNoEntry();
-			return {};
-		}
-
 		virtual void BindInputs(FInputVertexInterfaceData& InVertexData) override
 		{
 			using namespace ArrayNodeVertexNames;
@@ -199,22 +178,6 @@ namespace Metasound
 		using FArrayDataReadReference = TDataReadReference<ArrayType>;
 		using ElementType = typename MetasoundArrayNodesPrivate::TArrayElementType<ArrayType>::Type;
 
-		static const FVertexInterface& GetDefaultInterface()
-		{
-			using namespace ArrayNodeVertexNames;
-			static const FVertexInterface DefaultInterface(
-				FInputVertexInterface(
-					TInputDataVertex<FTrigger>(METASOUND_GET_PARAM_NAME_AND_METADATA(InputTriggerGet)),
-					TInputDataVertex<ArrayType>(METASOUND_GET_PARAM_NAME_AND_METADATA(InputArray)),
-					TInputDataVertex<int32>(METASOUND_GET_PARAM_NAME_AND_METADATA(InputIndex))
-				),
-				FOutputVertexInterface(
-					TOutputDataVertex<ElementType>(METASOUND_GET_PARAM_NAME_AND_METADATA(OutputValue))
-				)
-			);
-
-			return DefaultInterface;
-		}
 
 		static const FNodeClassMetadata& GetNodeInfo()
 		{
@@ -224,8 +187,7 @@ namespace Metasound
 				const FName OperatorName = TEXT("Get"); 
 				const FText NodeDisplayName = METASOUND_LOCTEXT_FORMAT("ArrayOpArrayGetDisplayNamePattern", "Get ({0})", GetMetasoundDataTypeDisplayText<ArrayType>());
 				const FText NodeDescription = METASOUND_LOCTEXT("ArrayOpArrayGetDescription", "Get element at index in array.");
-				const FVertexInterface NodeInterface = GetDefaultInterface();
-			
+				const FVertexInterface NodeInterface = MetasoundArrayNodesPrivate::CreateArrayGetInterface(DataTypeName, GetMetasoundDataTypeName<ElementType>());
 				return MetasoundArrayNodesPrivate::CreateArrayNodeClassMetadata(DataTypeName, OperatorName, NodeDisplayName, NodeDescription, NodeInterface);
 			};
 			
@@ -302,18 +264,6 @@ namespace Metasound
 
 		virtual ~TArrayGetOperator() = default;
 
-		virtual FDataReferenceCollection GetInputs() const override
-		{
-			checkNoEntry();
-			return {};
-		}
-
-		virtual FDataReferenceCollection GetOutputs() const override
-		{
-			checkNoEntry();
-			return {};
-		}
-
 		virtual void BindInputs(FInputVertexInterfaceData& InVertexData) override
 		{
 			using namespace ArrayNodeVertexNames;
@@ -391,24 +341,6 @@ namespace Metasound
 		using FArrayDataWriteReference = TDataWriteReference<ArrayType>;
 		using ElementType = typename MetasoundArrayNodesPrivate::TArrayElementType<ArrayType>::Type;
 
-		static const FVertexInterface& GetDefaultInterface()
-		{
-			using namespace ArrayNodeVertexNames;
-			static const FVertexInterface DefaultInterface(
-				FInputVertexInterface(
-					TInputDataVertex<FTrigger>(METASOUND_GET_PARAM_NAME_AND_METADATA(InputTriggerSet)),
-					TInputDataVertex<ArrayType>(InputInitialArrayName, FDataVertexMetadata { InputInitialArrayTooltip, InputInitialArrayDisplayName }),
-					TInputDataVertex<int32>(METASOUND_GET_PARAM_NAME_AND_METADATA(InputIndex)),
-					TInputDataVertex<ElementType>(METASOUND_GET_PARAM_NAME_AND_METADATA(InputValue))
-				),
-				FOutputVertexInterface(
-					TOutputDataVertex<ArrayType>(METASOUND_GET_PARAM_NAME_AND_METADATA(OutputArraySet))
-				)
-			);
-
-			return DefaultInterface;
-		}
-
 		static const FNodeClassMetadata& GetNodeInfo()
 		{
 			auto CreateNodeClassMetadata = []() -> FNodeClassMetadata
@@ -417,7 +349,7 @@ namespace Metasound
 				const FName OperatorName = TEXT("Set"); 
 				const FText NodeDisplayName = METASOUND_LOCTEXT_FORMAT("ArrayOpArraySetDisplayNamePattern", "Set ({0})", GetMetasoundDataTypeDisplayText<ArrayType>());
 				const FText NodeDescription = METASOUND_LOCTEXT("ArrayOpArraySetDescription", "Set element at index in array.");
-				const FVertexInterface NodeInterface = GetDefaultInterface();
+				const FVertexInterface NodeInterface = MetasoundArrayNodesPrivate::CreateArraySetInterface(DataTypeName, GetMetasoundDataTypeName<ElementType>());
 			
 				return MetasoundArrayNodesPrivate::CreateArrayNodeClassMetadata(DataTypeName, OperatorName, NodeDisplayName, NodeDescription, NodeInterface);
 			};
@@ -492,18 +424,6 @@ namespace Metasound
 
 		virtual ~TArraySetOperator() = default;
 
-		virtual FDataReferenceCollection GetInputs() const override
-		{
-			checkNoEntry();
-			return {};
-		}
-
-		virtual FDataReferenceCollection GetOutputs() const override
-		{
-			checkNoEntry();
-			return {};
-		}
-
 		virtual void BindInputs(FInputVertexInterfaceData& InVertexData) override
 		{
 			using namespace ArrayNodeVertexNames;
@@ -573,24 +493,6 @@ namespace Metasound
 		using FArrayDataWriteReference = TDataWriteReference<ArrayType>;
 		using ElementType = typename MetasoundArrayNodesPrivate::TArrayElementType<ArrayType>::Type;
 
-		static const FVertexInterface& GetDefaultInterface()
-		{
-			using namespace ArrayNodeVertexNames;
-
-			static const FVertexInterface DefaultInterface(
-				FInputVertexInterface(
-					TInputDataVertex<FTrigger>(METASOUND_GET_PARAM_NAME_AND_METADATA(InputTriggerGet)),
-					TInputDataVertex<ArrayType>(METASOUND_GET_PARAM_NAME_AND_METADATA(InputLeftArray)),
-					TInputDataVertex<ArrayType>(METASOUND_GET_PARAM_NAME_AND_METADATA(InputRightArray))
-				),
-				FOutputVertexInterface(
-					TOutputDataVertex<ArrayType>(METASOUND_GET_PARAM_NAME_AND_METADATA(OutputArrayConcat))
-				)
-			);
-
-			return DefaultInterface;
-		}
-
 		static const FNodeClassMetadata& GetNodeInfo()
 		{
 			auto CreateNodeClassMetadata = []() -> FNodeClassMetadata
@@ -599,7 +501,7 @@ namespace Metasound
 				const FName OperatorName = TEXT("Concat"); 
 				const FText NodeDisplayName = METASOUND_LOCTEXT_FORMAT("ArrayOpArrayConcatDisplayNamePattern", "Concatenate ({0})", GetMetasoundDataTypeDisplayText<ArrayType>());
 				const FText NodeDescription = METASOUND_LOCTEXT("ArrayOpArrayConcatDescription", "Concatenates two arrays on trigger.");
-				const FVertexInterface NodeInterface = GetDefaultInterface();
+				const FVertexInterface NodeInterface = MetasoundArrayNodesPrivate::CreateArrayConcatInterface(DataTypeName);
 			
 				FNodeClassMetadata NodeClassMetadata = MetasoundArrayNodesPrivate::CreateArrayNodeClassMetadata(DataTypeName, OperatorName, NodeDisplayName, NodeDescription, NodeInterface);
 				NodeClassMetadata.Keywords.Append({
@@ -639,18 +541,6 @@ namespace Metasound
 
 		virtual ~TArrayConcatOperator() = default;
 
-		virtual FDataReferenceCollection GetInputs() const override
-		{
-			checkNoEntry();
-			return {};
-		}
-
-		virtual FDataReferenceCollection GetOutputs() const override
-		{
-			checkNoEntry();
-			return {};
-		}
-
 		virtual void BindInputs(FInputVertexInterfaceData& InVertexData) override
 		{
 			using namespace ArrayNodeVertexNames;
@@ -701,26 +591,6 @@ namespace Metasound
 		using FArrayDataWriteReference = TDataWriteReference<ArrayType>;
 		using ElementType = typename MetasoundArrayNodesPrivate::TArrayElementType<ArrayType>::Type;
 
-		static const FVertexInterface& GetDefaultInterface()
-		{
-			using namespace ArrayNodeVertexNames;
-
-			static const FVertexInterface DefaultInterface(
-				FInputVertexInterface(
-					TInputDataVertex<FTrigger>(METASOUND_GET_PARAM_NAME_AND_METADATA(InputTriggerGet)),
-					TInputDataVertex<ArrayType>(METASOUND_GET_PARAM_NAME_AND_METADATA(InputArray)),
-					TInputDataVertex<int32>(METASOUND_GET_PARAM_NAME_AND_METADATA(InputStartIndex)),
-					TInputDataVertex<int32>(METASOUND_GET_PARAM_NAME_AND_METADATA(InputEndIndex))
-
-				),
-				FOutputVertexInterface(
-					TOutputDataVertex<ArrayType>(METASOUND_GET_PARAM_NAME_AND_METADATA(OutputArraySubset))
-				)
-			);
-
-			return DefaultInterface;
-		}
-
 		static const FNodeClassMetadata& GetNodeInfo()
 		{
 			auto CreateNodeClassMetadata = []() -> FNodeClassMetadata
@@ -729,7 +599,7 @@ namespace Metasound
 				const FName OperatorName = TEXT("Subset"); 
 				const FText NodeDisplayName = METASOUND_LOCTEXT_FORMAT("ArrayOpArraySubsetDisplayNamePattern", "Subset ({0})", GetMetasoundDataTypeDisplayText<ArrayType>());
 				const FText NodeDescription = METASOUND_LOCTEXT("ArrayOpArraySubsetDescription", "Subset array on trigger.");
-				const FVertexInterface NodeInterface = GetDefaultInterface();
+				const FVertexInterface NodeInterface = MetasoundArrayNodesPrivate::CreateArraySubsetInterface(DataTypeName);
 			
 				FNodeClassMetadata NodeClassMetadata = MetasoundArrayNodesPrivate::CreateArrayNodeClassMetadata(DataTypeName, OperatorName, NodeDisplayName, NodeDescription, NodeInterface);
 				NodeClassMetadata.Keywords.Append({
@@ -774,18 +644,6 @@ namespace Metasound
 
 		virtual ~TArraySubsetOperator() = default;
 
-		virtual FDataReferenceCollection GetInputs() const override
-		{
-			checkNoEntry();
-			return {};
-		}
-
-		virtual FDataReferenceCollection GetOutputs() const override
-		{
-			checkNoEntry();
-			return {};
-		}
-
 		virtual void BindInputs(FInputVertexInterfaceData& InVertexData) override
 		{
 			using namespace ArrayNodeVertexNames;
@@ -845,23 +703,6 @@ namespace Metasound
 	public:
 		using FArrayDataReadReference = TDataReadReference<ArrayType>;
 
-		// Declare the vertex interface
-		static const FVertexInterface& GetDefaultInterface()
-		{
-			using namespace ArrayNodeVertexNames;
-
-			static const FVertexInterface DefaultInterface(
-				FInputVertexInterface(
-					TInputDataVertex<ArrayType>(METASOUND_GET_PARAM_NAME_AND_METADATA(InputArray))
-				),
-				FOutputVertexInterface(
-					TOutputDataVertex<int32>(METASOUND_GET_PARAM_NAME_AND_METADATA(OutputLastIndex))
-				)
-			);
-
-			return DefaultInterface;
-		}
-
 		static const FNodeClassMetadata& GetNodeInfo()
 		{
 			auto CreateNodeClassMetadata = []() -> FNodeClassMetadata
@@ -870,7 +711,7 @@ namespace Metasound
 					const FName OperatorName = TEXT("GetLastIndex");
 					const FText NodeDisplayName = METASOUND_LOCTEXT_FORMAT("ArrayOpArrayLastIndexDisplayNamePattern", "Get Last Index ({0})", GetMetasoundDataTypeDisplayText<ArrayType>());
 					const FText NodeDescription = METASOUND_LOCTEXT("ArrayOpArrayLastIndexDescription", "Last index of the array");
-					const FVertexInterface NodeInterface = GetDefaultInterface();
+					const FVertexInterface NodeInterface = MetasoundArrayNodesPrivate::CreateArrayLastIndexInterface(DataTypeName);
 
 					return MetasoundArrayNodesPrivate::CreateArrayNodeClassMetadata(DataTypeName, OperatorName, NodeDisplayName, NodeDescription, NodeInterface);
 				};
diff --git a/Engine/Plugins/Runtime/Metasound/Source/MetasoundFrontend/Public/MetasoundDataTypeRegistrationMacro.h b/Engine/Plugins/Runtime/Metasound/Source/MetasoundFrontend/Public/MetasoundDataTypeRegistrationMacro.h
--- a/Engine/Plugins/Runtime/Metasound/Source/MetasoundFrontend/Public/MetasoundDataTypeRegistrationMacro.h
+++ b/Engine/Plugins/Runtime/Metasound/Source/MetasoundFrontend/Public/MetasoundDataTypeRegistrationMacro.h
@@ -33,8 +33,19 @@
 
 #define UE_API METASOUNDFRONTEND_API
 
+namespace Metasound
+{
+	// Disable registering of arrays of variables. 
+	template<typename UnderlyingType>
+	struct TEnableAutoArrayTypeRegistration<TVariable<UnderlyingType>>
+	{
+		static constexpr bool Value = false;
+	};
+}
+
 namespace Metasound::Frontend
 {
+
 	namespace MetasoundDataTypeRegistrationPrivate
 	{
 		template<typename DataType>
