Fix for missing LoadModuleChecked calls in editors using metasounds
Fix for cook issues related to metasound module loading changes.
#rnx

diff --git a/Engine/Plugins/Runtime/Metasound/Source/MetasoundFrontend/Private/MetasoundFrontendDataTypeRegistry.cpp b/Engine/Plugins/Runtime/Metasound/Source/MetasoundFrontend/Private/MetasoundFrontendDataTypeRegistry.cpp
--- a/Engine/Plugins/Runtime/Metasound/Source/MetasoundFrontend/Private/MetasoundFrontendDataTypeRegistry.cpp
+++ b/Engine/Plugins/Runtime/Metasound/Source/MetasoundFrontend/Private/MetasoundFrontendDataTypeRegistry.cpp
@@ -1508,12 +1508,10 @@ PRAGMA_ENABLE_DEPRECATION_WARNINGS
 					{
 						return &Entry->Get();
 					}
-					/* TODO: temporarily quieting this error.
 					else
 					{
 						UE_LOG(LogMetaSound, Error, TEXT("Data type is not registered [Name:%s]"), *InDataTypeName.ToString());
 					}
-					*/
 				}
 
 				return nullptr;
diff --git a/Engine/Plugins/Runtime/Metasound/Source/MetasoundFrontend/Private/MetasoundFrontendModuleRegistrationMacros.cpp b/Engine/Plugins/Runtime/Metasound/Source/MetasoundFrontend/Private/MetasoundFrontendModuleRegistrationMacros.cpp
--- a/Engine/Plugins/Runtime/Metasound/Source/MetasoundFrontend/Private/MetasoundFrontendModuleRegistrationMacros.cpp
+++ b/Engine/Plugins/Runtime/Metasound/Source/MetasoundFrontend/Private/MetasoundFrontendModuleRegistrationMacros.cpp
@@ -5,6 +5,14 @@
 
 namespace Metasound::Frontend::RegistrationPrivate
 {
+	FModuleInfo FGlobalRegistrationList::ModuleInfo = 
+	{
+#if WITH_EDITORONLY_DATA
+		"UnknownPlugin",
+		"UnknownModule"
+#endif // if WITH_EDITORONLY_DATA
+	};
+
 	FRegistrationAction::FRegistrationAction(FRegistrationAction*& InOutHead, FRegistrationFunction InRegister, FRegistrationFunction InUnregister)
 	: Register(InRegister)
 	, Unregister(InUnregister)
diff --git a/Engine/Plugins/Runtime/Metasound/Source/MetasoundFrontend/Public/MetasoundFrontendModuleRegistrationMacros.h b/Engine/Plugins/Runtime/Metasound/Source/MetasoundFrontend/Public/MetasoundFrontendModuleRegistrationMacros.h
--- a/Engine/Plugins/Runtime/Metasound/Source/MetasoundFrontend/Public/MetasoundFrontendModuleRegistrationMacros.h
+++ b/Engine/Plugins/Runtime/Metasound/Source/MetasoundFrontend/Public/MetasoundFrontendModuleRegistrationMacros.h
@@ -5,6 +5,16 @@
 #include "UObject/NameTypes.h"
 
 #define UE_API METASOUNDFRONTEND_API
+#define METASOUND_ENABLE_PER_MODULE_REGISTRATION 0 // Disabling per module registration while fixing cook issues. 
+
+// These macros are required in order to create module local registration actions.
+#ifndef METASOUND_ENABLE_PER_MODULE_REGISTRATION
+#if defined(METASOUND_MODULE) && defined(METASOUND_PLUGIN)
+#define METASOUND_ENABLE_PER_MODULE_REGISTRATION 1
+#else
+#define METASOUND_ENABLE_PER_MODULE_REGISTRATION 0
+#endif // if defined(METASOUND_MODULE) && defined(METASOUND_PLUGIN)
+#endif // ifndef METASOUND_ENABLE_PER_MODULE_REGISTRATION
 
 /** MetaSound Frontend Registration Actions
  *
@@ -119,7 +129,7 @@ namespace Metasound::Frontend
 		class FGlobalRegistrationList 
 		{
 		public: 
-			static FModuleInfo ModuleInfo;
+			UE_API static FModuleInfo ModuleInfo;
 
 			/* If you've hit this deprecation warning it is because you are using
 			 * a backwards compatible MetaSound node registration code path for your 
@@ -149,8 +159,12 @@ namespace Metasound::Frontend
 			 *   }
 			 *
 			 */
+#if METASOUND_ENABLE_PER_MODULE_REGISTRATION 
 			UE_DEPRECATED(5.7, "Macro based registration has been updated. See MetasoundFrontendModuleRegistrationMacros.h for a guide on defining module level registration.")
 			UE_API static FRegistrationAction*& GetHeadAction();
+#else
+			UE_API static FRegistrationAction*& GetHeadAction();
+#endif
 
 			// Used internally to do incremental registration on a global list. 
 			static FRegistrationAction*& GetRegisterTailAction();
@@ -161,8 +175,9 @@ namespace Metasound::Frontend
 	} // namespace RegistrationPrivate
 } // namespace Metasound::Frontend
 
-// These macros are required in order to create module local registration actions.
-#if defined(METASOUND_MODULE) && defined(METASOUND_PLUGIN)
+
+
+#if METASOUND_ENABLE_PER_MODULE_REGISTRATION
 
 // By default, always unregister on module shutdown. The only time we do not unregister
 // is when things cannot be associated with a module. This currently only happens
@@ -221,7 +236,7 @@ namespace Metasound::Frontend::RegistrationPrivate \
 	} \
 }
 
-#else // if defined(METASOUND_MODULE) && defined(METASOUND_PLUGIN)
+#else // if METASOUND_ENABLE_PER_MODULE_REGISTRATION
 
 // If METASOUND_MODULE and METASOUND_PLUGIN are not both defined, we rely on the
 // global registration list class as a fallback for ensuring the registration 
