Re-enable metasound module registration

#rnx

diff --git a/Engine/Plugins/Runtime/Metasound/Source/MetasoundFrontend/Public/MetasoundFrontendModuleRegistrationMacros.h b/Engine/Plugins/Runtime/Metasound/Source/MetasoundFrontend/Public/MetasoundFrontendModuleRegistrationMacros.h
--- a/Engine/Plugins/Runtime/Metasound/Source/MetasoundFrontend/Public/MetasoundFrontendModuleRegistrationMacros.h
+++ b/Engine/Plugins/Runtime/Metasound/Source/MetasoundFrontend/Public/MetasoundFrontendModuleRegistrationMacros.h
@@ -5,7 +5,6 @@
 #include "UObject/NameTypes.h"
 
 #define UE_API METASOUNDFRONTEND_API
-#define METASOUND_ENABLE_PER_MODULE_REGISTRATION 0 // Disabling per module registration while fixing cook issues. 
 
 // These macros are required in order to create module local registration actions.
 #ifndef METASOUND_ENABLE_PER_MODULE_REGISTRATION
