Fix MetaSound delegates getting removed in preset transform

diff --git a/Engine/Plugins/Runtime/Metasound/Source/MetasoundEditor/Private/MetasoundEditorBuilderListener.cpp b/Engine/Plugins/Runtime/Metasound/Source/MetasoundEditor/Private/MetasoundEditorBuilderListener.cpp
--- a/Engine/Plugins/Runtime/Metasound/Source/MetasoundEditor/Private/MetasoundEditorBuilderListener.cpp
+++ b/Engine/Plugins/Runtime/Metasound/Source/MetasoundEditor/Private/MetasoundEditorBuilderListener.cpp
@@ -9,8 +9,6 @@ void UMetaSoundEditorBuilderListener::Init(const TWeakObjectPtr<UMetaSoundBuilde
 	if (InBuilder.IsValid())
 	{
 		Builder = InBuilder;
-		FMetaSoundFrontendDocumentBuilder& DocumentBuilder = Builder->GetBuilder();
-		Metasound::Frontend::FDocumentModifyDelegates& BuilderDelegates = DocumentBuilder.GetDocumentDelegates();
 		BuilderListener = MakeShared<FEditorBuilderListener>(this);
 		Builder->AddTransactionListener(BuilderListener->AsShared());
 	}
diff --git a/Engine/Plugins/Runtime/Metasound/Source/MetasoundFrontend/Private/MetasoundFrontendDocumentCache.cpp b/Engine/Plugins/Runtime/Metasound/Source/MetasoundFrontend/Private/MetasoundFrontendDocumentCache.cpp
--- a/Engine/Plugins/Runtime/Metasound/Source/MetasoundFrontend/Private/MetasoundFrontendDocumentCache.cpp
+++ b/Engine/Plugins/Runtime/Metasound/Source/MetasoundFrontend/Private/MetasoundFrontendDocumentCache.cpp
@@ -339,7 +339,7 @@ namespace Metasound::Frontend
 	FDocumentGraphInterfaceCache::FDocumentGraphInterfaceCache(TSharedRef<const FDocumentCache> ParentCache)
 		: Parent(ParentCache)
 	{
-		const FMetasoundFrontendGraphClass& GraphClass = Parent->GetDocument().RootGraph;
+		const FMetasoundFrontendGraphClass& GraphClass = ParentCache->GetDocument().RootGraph;
 		const TArray<FMetasoundFrontendClassInput>& Inputs = GraphClass.GetDefaultInterface().Inputs;
 		for (int32 Index = 0; Index < Inputs.Num(); ++Index)
 		{
@@ -367,8 +367,11 @@ namespace Metasound::Frontend
 	{
 		if (const int32* Index = InputNameToIndex.Find(InputName))
 		{
-			const FMetasoundFrontendDocument& Doc = Parent->GetDocument();
-			return &Doc.RootGraph.GetDefaultInterface().Inputs[*Index];
+			if (TSharedPtr<const FDocumentCache> ParentCache = Parent.Pin())
+			{
+				const FMetasoundFrontendDocument& Doc = ParentCache->GetDocument();
+				return &Doc.RootGraph.GetDefaultInterface().Inputs[*Index];
+			}
 		}
 
 		return nullptr;
@@ -383,8 +386,11 @@ namespace Metasound::Frontend
 	{
 		if (const int32* Index = OutputNameToIndex.Find(OutputName))
 		{
-			const FMetasoundFrontendDocument& Doc = Parent->GetDocument();
-			return &Doc.RootGraph.GetDefaultInterface().Outputs[*Index];
+			if (TSharedPtr<const FDocumentCache> ParentCache = Parent.Pin())
+			{
+				const FMetasoundFrontendDocument& Doc = ParentCache->GetDocument();
+				return &Doc.RootGraph.GetDefaultInterface().Outputs[*Index];
+			}
 		}
 
 		return nullptr;
@@ -431,22 +437,28 @@ namespace Metasound::Frontend
 
 	void FDocumentGraphInterfaceCache::OnInputAdded(int32 NewIndex)
 	{
-		const FMetasoundFrontendDocument& Document = Parent->GetDocument();
-		const FMetasoundFrontendGraphClass& GraphClass = Document.RootGraph;
-		const FMetasoundFrontendClassInput& Input = GraphClass.GetDefaultInterface().Inputs[NewIndex];
-		InputNameToIndex.Add(Input.Name, NewIndex);
+		if (TSharedPtr<const FDocumentCache> ParentCache = Parent.Pin())
+		{
+			const FMetasoundFrontendDocument& Document = ParentCache->GetDocument();
+			const FMetasoundFrontendGraphClass& GraphClass = Document.RootGraph;
+			const FMetasoundFrontendClassInput& Input = GraphClass.GetDefaultInterface().Inputs[NewIndex];
+			InputNameToIndex.Add(Input.Name, NewIndex);
 
-		++TransactionCount;
+			++TransactionCount;
+		}
 	}
 
 	void FDocumentGraphInterfaceCache::OnRemovingInput(int32 IndexBeingRemoved)
 	{
-		const FMetasoundFrontendDocument& Document = Parent->GetDocument();
-		const FMetasoundFrontendGraphClass& GraphClass = Document.RootGraph;
-		const FMetasoundFrontendClassInput& Input = GraphClass.GetDefaultInterface().Inputs[IndexBeingRemoved];
-		InputNameToIndex.Remove(Input.Name);
+		if (TSharedPtr<const FDocumentCache> ParentCache = Parent.Pin())
+		{
+			const FMetasoundFrontendDocument& Document = ParentCache->GetDocument();
+			const FMetasoundFrontendGraphClass& GraphClass = Document.RootGraph;
+			const FMetasoundFrontendClassInput& Input = GraphClass.GetDefaultInterface().Inputs[IndexBeingRemoved];
+			InputNameToIndex.Remove(Input.Name);
 
-		++TransactionCount;
+			++TransactionCount;
+		}
 	}
 
 	void FDocumentGraphInterfaceCache::OnInputNameChanged(FName OldName, FName NewName)
@@ -498,22 +510,28 @@ namespace Metasound::Frontend
 
 	void FDocumentGraphInterfaceCache::OnOutputAdded(int32 NewIndex)
 	{
-		const FMetasoundFrontendDocument& Document = Parent->GetDocument();
-		const FMetasoundFrontendGraphClass& GraphClass = Document.RootGraph;
-		const FMetasoundFrontendClassOutput& Output = GraphClass.GetDefaultInterface().Outputs[NewIndex];
-		OutputNameToIndex.Add(Output.Name, NewIndex);
+		if (TSharedPtr<const FDocumentCache> ParentCache = Parent.Pin())
+		{
+			const FMetasoundFrontendDocument& Document = ParentCache->GetDocument();
+			const FMetasoundFrontendGraphClass& GraphClass = Document.RootGraph;
+			const FMetasoundFrontendClassOutput& Output = GraphClass.GetDefaultInterface().Outputs[NewIndex];
+			OutputNameToIndex.Add(Output.Name, NewIndex);
 
-		++TransactionCount;
+			++TransactionCount;
+		}
 	}
 
 	void FDocumentGraphInterfaceCache::OnRemovingOutput(int32 IndexBeingRemoved)
 	{
-		const FMetasoundFrontendDocument& Document = Parent->GetDocument();
-		const FMetasoundFrontendGraphClass& GraphClass = Document.RootGraph;
-		const FMetasoundFrontendClassOutput& Output = GraphClass.GetDefaultInterface().Outputs[IndexBeingRemoved];
-		OutputNameToIndex.Remove(Output.Name);
+		if (TSharedPtr<const FDocumentCache> ParentCache = Parent.Pin())
+		{
+			const FMetasoundFrontendDocument& Document = ParentCache->GetDocument();
+			const FMetasoundFrontendGraphClass& GraphClass = Document.RootGraph;
+			const FMetasoundFrontendClassOutput& Output = GraphClass.GetDefaultInterface().Outputs[IndexBeingRemoved];
+			OutputNameToIndex.Remove(Output.Name);
 
-		++TransactionCount;
+			++TransactionCount;
+		}
 	}
 
 	void FDocumentGraphInterfaceCache::OnOutputNameChanged(FName OldName, FName NewName)
@@ -562,12 +580,12 @@ namespace Metasound::Frontend
 		: Parent(ParentCache)
 		, PageID(InPageID)
 	{
-		const FMetasoundFrontendGraph& Graph = Parent->GetDocument().RootGraph.FindConstGraphChecked(InPageID);
+		const FMetasoundFrontendGraph& Graph = ParentCache->GetDocument().RootGraph.FindConstGraphChecked(InPageID);
 		const TArray<FMetasoundFrontendNode>& Nodes = Graph.Nodes;
 		for (int32 Index = 0; Index < Nodes.Num(); ++Index)
 		{
 			const FMetasoundFrontendNode& Node = Nodes[Index];
-			const FMetasoundFrontendClass* Class = Parent->FindDependency(Node.ClassID);
+			const FMetasoundFrontendClass* Class = ParentCache->FindDependency(Node.ClassID);
 			check(Class);
 
 			const FGuid& NodeID = Node.GetID();
@@ -601,9 +619,12 @@ namespace Metasound::Frontend
 		TArray<const FMetasoundFrontendNode*> Nodes;
 		if (const TArray<int32>* NodeIndices = ClassIDToNodeIndices.Find(InClassID))
 		{
-			const FMetasoundFrontendDocument& Document = Parent->GetDocument();
-			const FMetasoundFrontendGraph& Graph = Document.RootGraph.FindConstGraphChecked(PageID);
-			Algo::Transform(*NodeIndices, Nodes, [&Graph](const int32& Index) { return &Graph.Nodes[Index]; });
+			if (TSharedPtr<const FDocumentCache> ParentCache = Parent.Pin())
+			{
+				const FMetasoundFrontendDocument& Document = ParentCache->GetDocument();
+				const FMetasoundFrontendGraph& Graph = Document.RootGraph.FindConstGraphChecked(PageID);
+				Algo::Transform(*NodeIndices, Nodes, [&Graph](const int32& Index) { return &Graph.Nodes[Index]; });
+			}
 		}
 		return Nodes;
 	}
@@ -612,9 +633,12 @@ namespace Metasound::Frontend
 	{
 		if (const int32* NodeIndex = IDToIndex.Find(InNodeID))
 		{
-			const FMetasoundFrontendDocument& Document = Parent->GetDocument();
-			const FMetasoundFrontendGraph& Graph = Document.RootGraph.FindConstGraphChecked(PageID);
-			return &Graph.Nodes[*NodeIndex];
+			if (TSharedPtr<const FDocumentCache> ParentCache = Parent.Pin())
+			{
+				const FMetasoundFrontendDocument& Document = ParentCache->GetDocument();
+				const FMetasoundFrontendGraph& Graph = Document.RootGraph.FindConstGraphChecked(PageID);
+				return &Graph.Nodes[*NodeIndex];
+			}
 		}
 
 		return nullptr;
@@ -723,48 +747,51 @@ namespace Metasound::Frontend
 
 		if (const FMetasoundFrontendNode* Node = FindNode(InNodeID))
 		{
-			const FMetasoundFrontendClass* Class = Parent->FindDependency(Node->ClassID);
-			check(Class);
-			if (Class->Metadata.GetClassName() == FRerouteNodeTemplate::ClassName)
+			if (TSharedPtr<const FDocumentCache> ParentCache = Parent.Pin())
 			{
-				if (bOutIsRerouted)
-				{
-					*bOutIsRerouted = true;
-				}
-
-				const IDocumentGraphEdgeCache& EdgeCache = Parent->GetEdgeCache(PageID);
-				const FMetasoundFrontendVertex& RerouteOutput = Node->Interface.Outputs.Last();
-				TArrayView<const int32> EdgeIndices = EdgeCache.FindEdgeIndicesFromNodeOutput(InNodeID, RerouteOutput.VertexID);
-				const FMetasoundFrontendDocument& Doc = Parent->GetDocument();
-				for (int32 Index : EdgeIndices)
+				const FMetasoundFrontendClass* Class = ParentCache->FindDependency(Node->ClassID);
+				check(Class);
+				if (Class->Metadata.GetClassName() == FRerouteNodeTemplate::ClassName)
 				{
-					const FMetasoundFrontendEdge& EdgeConnectedToOutput = Doc.RootGraph.FindConstGraphChecked(PageID).Edges[Index];
-					TArray<const FMetasoundFrontendVertex*> ReroutedVertices;
-					if (ConnectedNodes)
+					if (bOutIsRerouted)
 					{
-						TArray<const FMetasoundFrontendNode*> ReroutedConnectedNodes;
-						ReroutedVertices = FindReroutedInputVertices(EdgeConnectedToOutput.ToNodeID, EdgeConnectedToOutput.ToVertexID, &ReroutedConnectedNodes, bOutIsRerouted);
-						ConnectedNodes->Append(ReroutedConnectedNodes);
+						*bOutIsRerouted = true;
 					}
-					else
+
+					const IDocumentGraphEdgeCache& EdgeCache = ParentCache->GetEdgeCache(PageID);
+					const FMetasoundFrontendVertex& RerouteOutput = Node->Interface.Outputs.Last();
+					TArrayView<const int32> EdgeIndices = EdgeCache.FindEdgeIndicesFromNodeOutput(InNodeID, RerouteOutput.VertexID);
+					const FMetasoundFrontendDocument& Doc = ParentCache->GetDocument();
+					for (int32 Index : EdgeIndices)
 					{
-						ReroutedVertices = FindReroutedInputVertices(EdgeConnectedToOutput.ToNodeID, EdgeConnectedToOutput.ToVertexID, nullptr, bOutIsRerouted);
+						const FMetasoundFrontendEdge& EdgeConnectedToOutput = Doc.RootGraph.FindConstGraphChecked(PageID).Edges[Index];
+						TArray<const FMetasoundFrontendVertex*> ReroutedVertices;
+						if (ConnectedNodes)
+						{
+							TArray<const FMetasoundFrontendNode*> ReroutedConnectedNodes;
+							ReroutedVertices = FindReroutedInputVertices(EdgeConnectedToOutput.ToNodeID, EdgeConnectedToOutput.ToVertexID, &ReroutedConnectedNodes, bOutIsRerouted);
+							ConnectedNodes->Append(ReroutedConnectedNodes);
+						}
+						else
+						{
+							ReroutedVertices = FindReroutedInputVertices(EdgeConnectedToOutput.ToNodeID, EdgeConnectedToOutput.ToVertexID, nullptr, bOutIsRerouted);
+						}
+						Vertices.Append(MoveTemp(ReroutedVertices));
 					}
-					Vertices.Append(MoveTemp(ReroutedVertices));
 				}
-			}
-			else
-			{
-				auto VertexMatchesPredicate = [&InVertexID](const FMetasoundFrontendVertex& Vertex)
-				{
-					return Vertex.VertexID == InVertexID;
-				};
-				const FMetasoundFrontendVertex* InputVertex = Node->Interface.Inputs.FindByPredicate(VertexMatchesPredicate);
-				check(InputVertex);
-				Vertices.Add(InputVertex);
-				if (ConnectedNodes)
+				else
 				{
-					ConnectedNodes->Add(Node);
+					auto VertexMatchesPredicate = [&InVertexID](const FMetasoundFrontendVertex& Vertex)
+						{
+							return Vertex.VertexID == InVertexID;
+						};
+					const FMetasoundFrontendVertex* InputVertex = Node->Interface.Inputs.FindByPredicate(VertexMatchesPredicate);
+					check(InputVertex);
+					Vertices.Add(InputVertex);
+					if (ConnectedNodes)
+					{
+						ConnectedNodes->Add(Node);
+					}
 				}
 			}
 		}
@@ -778,49 +805,52 @@ namespace Metasound::Frontend
 
 		if (const FMetasoundFrontendNode* Node = FindNode(InNodeID))
 		{
-			const FMetasoundFrontendClass* Class = Parent->FindDependency(Node->ClassID);
-			check(Class);
-			if (Class->Metadata.GetClassName() == FRerouteNodeTemplate::ClassName)
+			if (TSharedPtr<const FDocumentCache> ParentCache = Parent.Pin())
 			{
-				if (bOutIsRerouted)
-				{
-					*bOutIsRerouted = true;
-				}
-
-				const IDocumentGraphEdgeCache& EdgeCache = Parent->GetEdgeCache(PageID);
-				const FMetasoundFrontendVertex& RerouteOutput = Node->Interface.Outputs.Last();
-				TArrayView<const int32> EdgeIndices = EdgeCache.FindEdgeIndicesFromNodeOutput(InNodeID, RerouteOutput.VertexID);
-				const FMetasoundFrontendDocument& Doc = Parent->GetDocument();
-				for (int32 Index : EdgeIndices)
+				const FMetasoundFrontendClass* Class = ParentCache->FindDependency(Node->ClassID);
+				check(Class);
+				if (Class->Metadata.GetClassName() == FRerouteNodeTemplate::ClassName)
 				{
-					const FMetasoundFrontendEdge& EdgeConnectedToOutput = Doc.RootGraph.FindConstGraphChecked(PageID).Edges[Index];
-
-					TArray<const FMetasoundFrontendVertex*> ReroutedVertices;
-					if (ConnectedNodes)
+					if (bOutIsRerouted)
 					{
-						TArray<const FMetasoundFrontendNode*> ReroutedConnectedNodes;
-						ReroutedVertices = FindReroutedInputVertices(EdgeConnectedToOutput.ToNodeID, EdgeConnectedToOutput.ToVertexID, &ReroutedConnectedNodes, bOutIsRerouted);
-						ConnectedNodes->Append(MoveTemp(ReroutedConnectedNodes));
+						*bOutIsRerouted = true;
 					}
-					else
+
+					const IDocumentGraphEdgeCache& EdgeCache = ParentCache->GetEdgeCache(PageID);
+					const FMetasoundFrontendVertex& RerouteOutput = Node->Interface.Outputs.Last();
+					TArrayView<const int32> EdgeIndices = EdgeCache.FindEdgeIndicesFromNodeOutput(InNodeID, RerouteOutput.VertexID);
+					const FMetasoundFrontendDocument& Doc = ParentCache->GetDocument();
+					for (int32 Index : EdgeIndices)
 					{
-						ReroutedVertices = FindReroutedInputVertices(EdgeConnectedToOutput.ToNodeID, EdgeConnectedToOutput.ToVertexID, nullptr, bOutIsRerouted);
+						const FMetasoundFrontendEdge& EdgeConnectedToOutput = Doc.RootGraph.FindConstGraphChecked(PageID).Edges[Index];
+
+						TArray<const FMetasoundFrontendVertex*> ReroutedVertices;
+						if (ConnectedNodes)
+						{
+							TArray<const FMetasoundFrontendNode*> ReroutedConnectedNodes;
+							ReroutedVertices = FindReroutedInputVertices(EdgeConnectedToOutput.ToNodeID, EdgeConnectedToOutput.ToVertexID, &ReroutedConnectedNodes, bOutIsRerouted);
+							ConnectedNodes->Append(MoveTemp(ReroutedConnectedNodes));
+						}
+						else
+						{
+							ReroutedVertices = FindReroutedInputVertices(EdgeConnectedToOutput.ToNodeID, EdgeConnectedToOutput.ToVertexID, nullptr, bOutIsRerouted);
+						}
+						Vertices.Append(MoveTemp(ReroutedVertices));
 					}
-					Vertices.Append(MoveTemp(ReroutedVertices));
 				}
-			}
-			else
-			{
-				auto VertexMatchesPredicate = [&InVertexName](const FMetasoundFrontendVertex& Vertex)
-				{
-					return Vertex.Name == InVertexName;
-				};
-				const FMetasoundFrontendVertex* InputVertex = Node->Interface.Inputs.FindByPredicate(VertexMatchesPredicate);
-				check(InputVertex);
-				Vertices.Add(InputVertex);
-				if (ConnectedNodes)
+				else
 				{
-					ConnectedNodes->Add(Node);
+					auto VertexMatchesPredicate = [&InVertexName](const FMetasoundFrontendVertex& Vertex)
+						{
+							return Vertex.Name == InVertexName;
+						};
+					const FMetasoundFrontendVertex* InputVertex = Node->Interface.Inputs.FindByPredicate(VertexMatchesPredicate);
+					check(InputVertex);
+					Vertices.Add(InputVertex);
+					if (ConnectedNodes)
+					{
+						ConnectedNodes->Add(Node);
+					}
 				}
 			}
 		}
@@ -832,36 +862,39 @@ namespace Metasound::Frontend
 	{
 		if (const FMetasoundFrontendNode* Node = FindNode(InNodeID))
 		{
-			const FMetasoundFrontendClass* Class = Parent->FindDependency(Node->ClassID);
-			check(Class);
-			if (Class->Metadata.GetClassName() == FRerouteNodeTemplate::ClassName)
+			if (TSharedPtr<const FDocumentCache> ParentCache = Parent.Pin())
 			{
-				if (bOutIsRerouted)
+				const FMetasoundFrontendClass* Class = ParentCache->FindDependency(Node->ClassID);
+				check(Class);
+				if (Class->Metadata.GetClassName() == FRerouteNodeTemplate::ClassName)
 				{
-					*bOutIsRerouted = true;
+					if (bOutIsRerouted)
+					{
+						*bOutIsRerouted = true;
+					}
+
+					const IDocumentGraphEdgeCache& EdgeCache = ParentCache->GetEdgeCache(PageID);
+					const FMetasoundFrontendVertex& RerouteInput = Node->Interface.Inputs.Last();
+					const FMetasoundFrontendDocument& Doc = ParentCache->GetDocument();
+					if (const int32* ConnectedEdgeIndex = EdgeCache.FindEdgeIndexToNodeInput(InNodeID, RerouteInput.VertexID))
+					{
+						const FMetasoundFrontendEdge& EdgeConnectedToInput = Doc.RootGraph.FindConstGraphChecked(PageID).Edges[*ConnectedEdgeIndex];
+						return FindReroutedOutputVertex(EdgeConnectedToInput.FromNodeID, EdgeConnectedToInput.FromVertexID, ConnectedNode, bOutIsRerouted);
+					}
+
+					return nullptr;
 				}
 
-				const IDocumentGraphEdgeCache& EdgeCache = Parent->GetEdgeCache(PageID);
-				const FMetasoundFrontendVertex& RerouteInput = Node->Interface.Inputs.Last();
-				const FMetasoundFrontendDocument& Doc = Parent->GetDocument();
-				if (const int32* ConnectedEdgeIndex = EdgeCache.FindEdgeIndexToNodeInput(InNodeID, RerouteInput.VertexID))
+				if (ConnectedNode)
 				{
-					const FMetasoundFrontendEdge& EdgeConnectedToInput = Doc.RootGraph.FindConstGraphChecked(PageID).Edges[*ConnectedEdgeIndex];
-					return FindReroutedOutputVertex(EdgeConnectedToInput.FromNodeID, EdgeConnectedToInput.FromVertexID, ConnectedNode, bOutIsRerouted);
+					*ConnectedNode = Node;
 				}
-
-				return nullptr;
-			}
-
-			if (ConnectedNode)
-			{
-				*ConnectedNode = Node;
+				auto VertexMatchesPredicate = [&InVertexID](const FMetasoundFrontendVertex& Vertex)
+					{
+						return Vertex.VertexID == InVertexID;
+					};
+				return Node->Interface.Outputs.FindByPredicate(VertexMatchesPredicate);
 			}
-			auto VertexMatchesPredicate = [&InVertexID](const FMetasoundFrontendVertex& Vertex)
-			{
-				return Vertex.VertexID == InVertexID;
-			};
-			return Node->Interface.Outputs.FindByPredicate(VertexMatchesPredicate);
 		}
 
 		return nullptr;
@@ -871,36 +904,39 @@ namespace Metasound::Frontend
 	{
 		if (const FMetasoundFrontendNode* Node = FindNode(InNodeID))
 		{
-			const FMetasoundFrontendClass* Class = Parent->FindDependency(Node->ClassID);
-			check(Class);
-			if (Class->Metadata.GetClassName() == FRerouteNodeTemplate::ClassName)
+			if (TSharedPtr<const FDocumentCache> ParentCache = Parent.Pin())
 			{
-				if (bOutIsRerouted)
+				const FMetasoundFrontendClass* Class = ParentCache->FindDependency(Node->ClassID);
+				check(Class);
+				if (Class->Metadata.GetClassName() == FRerouteNodeTemplate::ClassName)
 				{
-					*bOutIsRerouted = true;
+					if (bOutIsRerouted)
+					{
+						*bOutIsRerouted = true;
+					}
+
+					const IDocumentGraphEdgeCache& EdgeCache = ParentCache->GetEdgeCache(PageID);
+					const FMetasoundFrontendVertex& RerouteInput = Node->Interface.Inputs.Last();
+					const FMetasoundFrontendDocument& Doc = ParentCache->GetDocument();
+					if (const int32* ConnectedEdgeIndex = EdgeCache.FindEdgeIndexToNodeInput(InNodeID, RerouteInput.VertexID))
+					{
+						const FMetasoundFrontendEdge& EdgeConnectedToInput = Doc.RootGraph.FindConstGraphChecked(PageID).Edges[*ConnectedEdgeIndex];
+						return FindReroutedOutputVertex(EdgeConnectedToInput.FromNodeID, EdgeConnectedToInput.FromVertexID, ConnectedNode, bOutIsRerouted);
+					}
+
+					return nullptr;
 				}
 
-				const IDocumentGraphEdgeCache& EdgeCache = Parent->GetEdgeCache(PageID);
-				const FMetasoundFrontendVertex& RerouteInput = Node->Interface.Inputs.Last();
-				const FMetasoundFrontendDocument& Doc = Parent->GetDocument();
-				if (const int32* ConnectedEdgeIndex = EdgeCache.FindEdgeIndexToNodeInput(InNodeID, RerouteInput.VertexID))
+				auto VertexMatchesPredicate = [&InVertexName](const FMetasoundFrontendVertex& Vertex)
+				{
+					return Vertex.Name == InVertexName;
+				};
+				if (ConnectedNode)
 				{
-					const FMetasoundFrontendEdge& EdgeConnectedToInput = Doc.RootGraph.FindConstGraphChecked(PageID).Edges[*ConnectedEdgeIndex];
-					return FindReroutedOutputVertex(EdgeConnectedToInput.FromNodeID, EdgeConnectedToInput.FromVertexID, ConnectedNode, bOutIsRerouted);
+					*ConnectedNode = Node;
 				}
-
-				return nullptr;
-			}
-
-			auto VertexMatchesPredicate = [&InVertexName](const FMetasoundFrontendVertex& Vertex)
-			{
-				return Vertex.Name == InVertexName;
-			};
-			if (ConnectedNode)
-			{
-				*ConnectedNode = Node;
+				return Node->Interface.Outputs.FindByPredicate(VertexMatchesPredicate);
 			}
-			return Node->Interface.Outputs.FindByPredicate(VertexMatchesPredicate);
 		}
 
 		return nullptr;
@@ -916,13 +952,16 @@ namespace Metasound::Frontend
 
 	void FDocumentGraphNodeCache::OnNodeAdded(int32 InNewIndex)
 	{
-		const FMetasoundFrontendDocument& Document = Parent->GetDocument();
-		const FMetasoundFrontendGraph& Graph = Document.RootGraph.FindConstGraphChecked(PageID);
-		const FMetasoundFrontendNode& Node = Graph.Nodes[InNewIndex];
-		IDToIndex.Add(Node.GetID(), InNewIndex);
-		ClassIDToNodeIndices.FindOrAdd(Node.ClassID).Add(InNewIndex);
+		if (TSharedPtr<const FDocumentCache> ParentCache = Parent.Pin())
+		{
+			const FMetasoundFrontendDocument& Document = ParentCache->GetDocument();
+			const FMetasoundFrontendGraph& Graph = Document.RootGraph.FindConstGraphChecked(PageID);
+			const FMetasoundFrontendNode& Node = Graph.Nodes[InNewIndex];
+			IDToIndex.Add(Node.GetID(), InNewIndex);
+			ClassIDToNodeIndices.FindOrAdd(Node.ClassID).Add(InNewIndex);
 
-		++TransactionCount;
+			++TransactionCount;
+		}
 	}
 
 	void FDocumentGraphNodeCache::OnNodeInputLiteralSet(int32 /* NodeIndex */, int32 /* VertexIndex */, int32 /* LiteralIndex */)
@@ -933,16 +972,18 @@ namespace Metasound::Frontend
 	void FDocumentGraphNodeCache::OnRemoveSwappingNode(int32 SwapIndex, int32 LastIndex)
 	{
 		using namespace Metasound::Frontend;
+		if (TSharedPtr<const FDocumentCache> ParentCache = Parent.Pin())
+		{
+			const FMetasoundFrontendDocument& Document = ParentCache->GetDocument();
+			const TArray<FMetasoundFrontendNode>& Nodes = Document.RootGraph.FindConstGraphChecked(PageID).Nodes;
+			const FMetasoundFrontendNode& SwapNode = Nodes[SwapIndex];
+			const FMetasoundFrontendNode& LastNode = Nodes[LastIndex];
 
-		const FMetasoundFrontendDocument& Document = Parent->GetDocument();
-		const TArray<FMetasoundFrontendNode>& Nodes = Document.RootGraph.FindConstGraphChecked(PageID).Nodes;
-		const FMetasoundFrontendNode& SwapNode = Nodes[SwapIndex];
-		const FMetasoundFrontendNode& LastNode = Nodes[LastIndex];
-
-		DocumentCachePrivate::RemoveSwapMapArrayIndexChecked(SwapNode.ClassID, SwapIndex, LastNode.ClassID, LastIndex, ClassIDToNodeIndices);
-		DocumentCachePrivate::RemoveSwapMapIndexChecked(SwapNode.GetID(), LastNode.GetID(), IDToIndex);
+			DocumentCachePrivate::RemoveSwapMapArrayIndexChecked(SwapNode.ClassID, SwapIndex, LastNode.ClassID, LastIndex, ClassIDToNodeIndices);
+			DocumentCachePrivate::RemoveSwapMapIndexChecked(SwapNode.GetID(), LastNode.GetID(), IDToIndex);
 
-		++TransactionCount;
+			++TransactionCount;
+		}
 	}
 
 	void FDocumentGraphNodeCache::OnRemovingNodeInputLiteral(int32 /* NodeIndex */, int32 /* VertexIndex */, int32 /* LiteralIndex */)
@@ -959,7 +1000,7 @@ namespace Metasound::Frontend
 		: Parent(ParentCache)
 		, PageID(InPageID)
 	{
-		const TArray<FMetasoundFrontendEdge>& Edges = Parent->GetDocument().RootGraph.FindConstGraphChecked(InPageID).Edges;
+		const TArray<FMetasoundFrontendEdge>& Edges = ParentCache->GetDocument().RootGraph.FindConstGraphChecked(InPageID).Edges;
 		for (int32 Index = 0; Index < Edges.Num(); ++Index)
 		{
 			const FMetasoundFrontendEdge& Edge = Edges[Index];
@@ -973,9 +1014,12 @@ namespace Metasound::Frontend
 		const FMetasoundFrontendVertexHandle InputPair = { InEdge.ToNodeID, InEdge.ToVertexID };
 		if (const int32* Index = InputToEdgeIndex.Find(InputPair))
 		{
-			const FMetasoundFrontendDocument& Document = Parent->GetDocument();
-			const FMetasoundFrontendEdge& Edge = Document.RootGraph.FindConstGraphChecked(PageID).Edges[*Index];
-			return Edge.FromNodeID == InEdge.FromNodeID && Edge.FromVertexID == InEdge.FromVertexID;
+			if (TSharedPtr<const FDocumentCache> ParentCache = Parent.Pin())
+			{
+				const FMetasoundFrontendDocument& Document = ParentCache->GetDocument();
+				const FMetasoundFrontendEdge& Edge = Document.RootGraph.FindConstGraphChecked(PageID).Edges[*Index];
+				return Edge.FromNodeID == InEdge.FromNodeID && Edge.FromVertexID == InEdge.FromVertexID;
+			}
 		}
 
 		return false;
@@ -994,19 +1038,21 @@ namespace Metasound::Frontend
 	TArray<const FMetasoundFrontendEdge*> FDocumentGraphEdgeCache::FindEdges(const FGuid& InNodeID, const FGuid& InVertexID) const
 	{
 		TArray<const FMetasoundFrontendEdge*> Edges;
-
-		const FMetasoundFrontendDocument& Document = Parent->GetDocument();
-		const FMetasoundFrontendGraph& Graph = Document.RootGraph.FindConstGraphChecked(PageID);
-
-		const FMetasoundFrontendVertexHandle Handle { InNodeID, InVertexID };
-		if (const int32* Index = InputToEdgeIndex.Find(Handle))
+		if (TSharedPtr<const FDocumentCache> ParentCache = Parent.Pin())
 		{
-			Edges.Add(&Graph.Edges[*Index]);
-		}
+			const FMetasoundFrontendDocument& Document = ParentCache->GetDocument();
+			const FMetasoundFrontendGraph& Graph = Document.RootGraph.FindConstGraphChecked(PageID);
 
-		if (const TArray<int32>* Indices = OutputToEdgeIndices.Find(Handle))
-		{
-			Algo::Transform(*Indices, Edges, [&Graph](const int32& Index) { return &Graph.Edges[Index]; });
+			const FMetasoundFrontendVertexHandle Handle{ InNodeID, InVertexID };
+			if (const int32* Index = InputToEdgeIndex.Find(Handle))
+			{
+				Edges.Add(&Graph.Edges[*Index]);
+			}
+
+			if (const TArray<int32>* Indices = OutputToEdgeIndices.Find(Handle))
+			{
+				Algo::Transform(*Indices, Edges, [&Graph](const int32& Index) { return &Graph.Edges[Index]; });
+			}
 		}
 
 		return Edges;
@@ -1050,26 +1096,32 @@ namespace Metasound::Frontend
 
 	void FDocumentGraphEdgeCache::OnEdgeAdded(int32 InNewIndex)
 	{
-		const FMetasoundFrontendDocument& Document = Parent->GetDocument();
-		const FMetasoundFrontendGraph& Graph = Document.RootGraph.FindConstGraphChecked(PageID);
-		const FMetasoundFrontendEdge& NewEdge = Graph.Edges[InNewIndex];
-		InputToEdgeIndex.Add(NewEdge.GetToVertexHandle(), InNewIndex);
-		OutputToEdgeIndices.FindOrAdd(NewEdge.GetFromVertexHandle()).Add(InNewIndex);
+		if (TSharedPtr<const FDocumentCache> ParentCache = Parent.Pin())
+		{
+			const FMetasoundFrontendDocument& Document = ParentCache->GetDocument();
+			const FMetasoundFrontendGraph& Graph = Document.RootGraph.FindConstGraphChecked(PageID);
+			const FMetasoundFrontendEdge& NewEdge = Graph.Edges[InNewIndex];
+			InputToEdgeIndex.Add(NewEdge.GetToVertexHandle(), InNewIndex);
+			OutputToEdgeIndices.FindOrAdd(NewEdge.GetFromVertexHandle()).Add(InNewIndex);
 
-		++TransactionCount;
+			++TransactionCount;
+		}
 	}
 
 	void FDocumentGraphEdgeCache::OnRemoveSwappingEdge(int32 SwapIndex, int32 LastIndex)
 	{
-		const FMetasoundFrontendDocument& Document = Parent->GetDocument();
-		const TArray<FMetasoundFrontendEdge>& Edges = Document.RootGraph.FindConstGraphChecked(PageID).Edges;
+		if (TSharedPtr<const FDocumentCache> ParentCache = Parent.Pin())
+		{
+			const FMetasoundFrontendDocument& Document = ParentCache->GetDocument();
+			const TArray<FMetasoundFrontendEdge>& Edges = Document.RootGraph.FindConstGraphChecked(PageID).Edges;
 
-		const FMetasoundFrontendEdge& SwapEdge = Edges[SwapIndex];
-		const FMetasoundFrontendEdge& LastEdge = Edges[LastIndex];
-		DocumentCachePrivate::RemoveSwapMapIndexChecked(SwapEdge.GetToVertexHandle(), LastEdge.GetToVertexHandle(), InputToEdgeIndex);
-		DocumentCachePrivate::RemoveSwapMapArrayIndexChecked(SwapEdge.GetFromVertexHandle(), SwapIndex, LastEdge.GetFromVertexHandle(), LastIndex, OutputToEdgeIndices);
+			const FMetasoundFrontendEdge& SwapEdge = Edges[SwapIndex];
+			const FMetasoundFrontendEdge& LastEdge = Edges[LastIndex];
+			DocumentCachePrivate::RemoveSwapMapIndexChecked(SwapEdge.GetToVertexHandle(), LastEdge.GetToVertexHandle(), InputToEdgeIndex);
+			DocumentCachePrivate::RemoveSwapMapArrayIndexChecked(SwapEdge.GetFromVertexHandle(), SwapIndex, LastEdge.GetFromVertexHandle(), LastIndex, OutputToEdgeIndices);
 
-		++TransactionCount;
+			++TransactionCount;
+		}
 	}
 
 	FDocumentCache::FDocumentDependencyCache::FDocumentDependencyCache(const FMetasoundFrontendDocument& InDocument)
diff --git a/Engine/Plugins/Runtime/Metasound/Source/MetasoundFrontend/Private/MetasoundFrontendDocumentCache.h b/Engine/Plugins/Runtime/Metasound/Source/MetasoundFrontend/Private/MetasoundFrontendDocumentCache.h
--- a/Engine/Plugins/Runtime/Metasound/Source/MetasoundFrontend/Private/MetasoundFrontendDocumentCache.h
+++ b/Engine/Plugins/Runtime/Metasound/Source/MetasoundFrontend/Private/MetasoundFrontendDocumentCache.h
@@ -52,7 +52,7 @@ namespace Metasound::Frontend
 		// Cache of Input NodeId/VertexId pairs to associated edge indices
 		TMap<FMetasoundFrontendVertexHandle, int32> InputToEdgeIndex;
 
-		TSharedRef<const FDocumentCache> Parent;
+		TWeakPtr<const FDocumentCache> Parent;
 		FGuid PageID;
 		int32 TransactionCount = 0;
 	};
@@ -107,7 +107,7 @@ namespace Metasound::Frontend
 		// Cache of ClassID to referencing node indices
 		TSortedMap<FGuid, TArray<int32>> ClassIDToNodeIndices;
 
-		TSharedRef<const FDocumentCache> Parent;
+		TWeakPtr<const FDocumentCache> Parent;
 		FGuid PageID;
 		int32 TransactionCount = 0;
 	};
@@ -165,7 +165,7 @@ namespace Metasound::Frontend
 		// Cache of Output name to array index of output
 		TMap<FName, int32> OutputNameToIndex;
 
-		TSharedRef<const FDocumentCache> Parent;
+		TWeakPtr<const FDocumentCache> Parent;
 
 		int32 TransactionCount = 0;
 	};
diff --git a/Engine/Plugins/Runtime/Metasound/Source/MetasoundFrontend/Private/MetasoundFrontendDocumentModifyDelegates.cpp b/Engine/Plugins/Runtime/Metasound/Source/MetasoundFrontend/Private/MetasoundFrontendDocumentModifyDelegates.cpp
--- a/Engine/Plugins/Runtime/Metasound/Source/MetasoundFrontend/Private/MetasoundFrontendDocumentModifyDelegates.cpp
+++ b/Engine/Plugins/Runtime/Metasound/Source/MetasoundFrontend/Private/MetasoundFrontendDocumentModifyDelegates.cpp
@@ -14,18 +14,30 @@ PRAGMA_ENABLE_DEPRECATION_WARNINGS
 	}
 
 	FDocumentModifyDelegates::FDocumentModifyDelegates(const FDocumentModifyDelegates& InModifyDelegates)
-PRAGMA_DISABLE_DEPRECATION_WARNINGS
-		: NodeDelegates(InModifyDelegates.NodeDelegates)
+		: OnDocumentMetadataChanged(InModifyDelegates.OnDocumentMetadataChanged)
+		, OnDependencyAdded(InModifyDelegates.OnDependencyAdded)
+		, OnRemoveSwappingDependency(InModifyDelegates.OnRemoveSwappingDependency)
+		, OnRenamingDependencyClass(InModifyDelegates.OnRenamingDependencyClass)
+		, PageDelegates(InModifyDelegates.PageDelegates)
+		, InterfaceDelegates(InModifyDelegates.InterfaceDelegates)
+		PRAGMA_DISABLE_DEPRECATION_WARNINGS
+		, NodeDelegates(InModifyDelegates.NodeDelegates)
 		, EdgeDelegates(InModifyDelegates.EdgeDelegates)
-PRAGMA_ENABLE_DEPRECATION_WARNINGS
+		PRAGMA_ENABLE_DEPRECATION_WARNINGS
 		, PageNodeDelegates(InModifyDelegates.PageNodeDelegates)
 		, PageEdgeDelegates(InModifyDelegates.PageEdgeDelegates)
 	{
 	}
 
 	FDocumentModifyDelegates::FDocumentModifyDelegates(FDocumentModifyDelegates&& InModifyDelegates)
+		: OnDocumentMetadataChanged(MoveTemp(InModifyDelegates.OnDocumentMetadataChanged))
+		, OnDependencyAdded(MoveTemp(InModifyDelegates.OnDependencyAdded))
+		, OnRemoveSwappingDependency(MoveTemp(InModifyDelegates.OnRemoveSwappingDependency))
+		, OnRenamingDependencyClass(MoveTemp(InModifyDelegates.OnRenamingDependencyClass))
+		, PageDelegates(MoveTemp(InModifyDelegates.PageDelegates))
+		, InterfaceDelegates(MoveTemp(InModifyDelegates.InterfaceDelegates))
 PRAGMA_DISABLE_DEPRECATION_WARNINGS
-		: NodeDelegates(MoveTemp(InModifyDelegates.NodeDelegates))
+		, NodeDelegates(MoveTemp(InModifyDelegates.NodeDelegates))
 		, EdgeDelegates(MoveTemp(InModifyDelegates.EdgeDelegates))
 PRAGMA_ENABLE_DEPRECATION_WARNINGS
 		, PageNodeDelegates(MoveTemp(InModifyDelegates.PageNodeDelegates))
@@ -35,6 +47,12 @@ PRAGMA_ENABLE_DEPRECATION_WARNINGS
 
 	FDocumentModifyDelegates& FDocumentModifyDelegates::operator=(const FDocumentModifyDelegates& InModifyDelegates)
 	{
+		OnDocumentMetadataChanged = InModifyDelegates.OnDocumentMetadataChanged;
+		OnDependencyAdded = InModifyDelegates.OnDependencyAdded;
+		OnRemoveSwappingDependency = InModifyDelegates.OnRemoveSwappingDependency;
+		OnRenamingDependencyClass = InModifyDelegates.OnRenamingDependencyClass;
+		PageDelegates = InModifyDelegates.PageDelegates;
+		InterfaceDelegates = InModifyDelegates.InterfaceDelegates;
 PRAGMA_DISABLE_DEPRECATION_WARNINGS
 		NodeDelegates = InModifyDelegates.NodeDelegates;
 		EdgeDelegates = InModifyDelegates.EdgeDelegates;
@@ -46,6 +64,12 @@ PRAGMA_ENABLE_DEPRECATION_WARNINGS
 
 	FDocumentModifyDelegates& FDocumentModifyDelegates::operator=(FDocumentModifyDelegates&& InModifyDelegates)
 	{
+		OnDocumentMetadataChanged = MoveTemp(InModifyDelegates.OnDocumentMetadataChanged);
+		OnDependencyAdded = MoveTemp(InModifyDelegates.OnDependencyAdded);
+		OnRemoveSwappingDependency = MoveTemp(InModifyDelegates.OnRemoveSwappingDependency);
+		OnRenamingDependencyClass = MoveTemp(InModifyDelegates.OnRenamingDependencyClass);
+		PageDelegates = MoveTemp(InModifyDelegates.PageDelegates);
+		InterfaceDelegates = MoveTemp(InModifyDelegates.InterfaceDelegates);
 PRAGMA_DISABLE_DEPRECATION_WARNINGS
 		NodeDelegates = MoveTemp(InModifyDelegates.NodeDelegates);
 		EdgeDelegates = MoveTemp(InModifyDelegates.EdgeDelegates);
diff --git a/Engine/Plugins/Runtime/Metasound/Source/MetasoundFrontend/Private/MetasoundFrontendTransform.cpp b/Engine/Plugins/Runtime/Metasound/Source/MetasoundFrontend/Private/MetasoundFrontendTransform.cpp
--- a/Engine/Plugins/Runtime/Metasound/Source/MetasoundFrontend/Private/MetasoundFrontendTransform.cpp
+++ b/Engine/Plugins/Runtime/Metasound/Source/MetasoundFrontend/Private/MetasoundFrontendTransform.cpp
@@ -699,7 +699,7 @@ namespace Metasound
 			}
 
 			// Clear the root graph so it can be rebuilt.
-			TSharedRef<FDocumentModifyDelegates> ModifyDelegates = MakeShared<FDocumentModifyDelegates>(InOutBuilderToTransform.GetConstDocumentChecked());
+			TSharedRef<FDocumentModifyDelegates> ModifyDelegates = MakeShared<FDocumentModifyDelegates>(InOutBuilderToTransform.GetDocumentDelegates());
 			InOutBuilderToTransform.ClearDocument(ModifyDelegates);
 
 			InOutBuilderToTransform.SetPresetFlags();
