Added MetasoundPresetWidgetInterface function to provide preset widgets with the preview audio component when auditioning

#rnx

diff --git a/Engine/Plugins/Runtime/Metasound/Source/MetasoundEditor/Private/MetasoundEditor.cpp b/Engine/Plugins/Runtime/Metasound/Source/MetasoundEditor/Private/MetasoundEditor.cpp
--- a/Engine/Plugins/Runtime/Metasound/Source/MetasoundEditor/Private/MetasoundEditor.cpp
+++ b/Engine/Plugins/Runtime/Metasound/Source/MetasoundEditor/Private/MetasoundEditor.cpp
@@ -1860,6 +1860,11 @@ namespace Metasound
 				// This MetaSound is in the inclusion list 
 				for (TScriptInterface<IMetaSoundDocumentInterface> IncludedMetaSound : IncludedMetaSounds)
 				{
+					if (!IncludedMetaSound)
+					{
+						continue;
+					}
+
 					// This MetaSound or the referenced MetaSound is included 
 					const FMetasoundFrontendGraphClass& IncludedGraphClass = IncludedMetaSound->GetConstDocument().RootGraph;
 					const FNodeRegistryKey IncludedRegistryKey(IncludedGraphClass);
@@ -2614,6 +2619,12 @@ namespace Metasound
 							CastChecked<UMetaSoundSourceBuilder>(Builder.Get())->Audition(nullptr, PreviewComp, GeneratorHandleDelegate, AssetEditorPrivate::bEnableLiveAudition);
 						}
 
+						if (UserPresetWidget.IsValid() && IsPlaying())
+						{
+							constexpr bool bIsAuditioning = true;
+							IMetaSoundPresetWidgetInterface::Execute_OnAuditionStateChanged(UserPresetWidget.Get(), PreviewComp, bIsAuditioning);
+						}
+
 						SetPreviewID(PreviewComp->GetUniqueID());
 
 						if (UAudioBus* AudioBus = AnalyzerRack->GetAudioBus())
@@ -2639,6 +2650,13 @@ namespace Metasound
 						{
 							UpdatePageInfo(bIsPlaying);
 							GraphConnectionManager = RebuildConnectionManager();
+
+							if (UserPresetWidget.IsValid())
+							{
+								constexpr bool bIsAuditioning = false;
+								IMetaSoundPresetWidgetInterface::Execute_OnAuditionStateChanged(UserPresetWidget.Get(), nullptr, bIsAuditioning);
+							}
+
 							return EActiveTimerReturnType::Stop;
 						}
 					})
diff --git a/Engine/Plugins/Runtime/Metasound/Source/MetasoundEditor/Public/MetasoundPresetWidgetInterface.h b/Engine/Plugins/Runtime/Metasound/Source/MetasoundEditor/Public/MetasoundPresetWidgetInterface.h
--- a/Engine/Plugins/Runtime/Metasound/Source/MetasoundEditor/Public/MetasoundPresetWidgetInterface.h
+++ b/Engine/Plugins/Runtime/Metasound/Source/MetasoundEditor/Public/MetasoundPresetWidgetInterface.h
@@ -1,4 +1,5 @@
 // Copyright Epic Games, Inc. All Rights Reserved.
+
 #pragma once
 
 #include "Audio/AudioPanelWidgetInterface.h"
@@ -25,10 +26,14 @@ public:
 	// The MetaSounds whose presets are supported by this widget. If Support All Presets is true, this widget is supported by all presets except those in the Excluded array.
 	// MetaSounds in the include/exclude arrays can be MetaSound presets or non presets. 
 	// If a MetaSound is not a preset, then presets of that MetaSound will be supported/excluded by this widget.
-	UFUNCTION(BlueprintImplementableEvent)
+	UFUNCTION(BlueprintImplementableEvent, Category = "MetaSound Preset Widget")
 	void GetSupportedMetaSounds(bool& bSupportAllPresets, TArray<TScriptInterface<IMetaSoundDocumentInterface>>& ExcludedMetaSounds, TArray<TScriptInterface<IMetaSoundDocumentInterface>>& IncludedMetaSounds) const;
 
 	// Called when the preset widget is constructed, giving the builder of the associated MetaSound preset
-	UFUNCTION(BlueprintImplementableEvent, meta = (DisplayName = "On MetaSoundPreset Widget Constructed"))
+	UFUNCTION(BlueprintImplementableEvent, Category = "MetaSound Preset Widget", meta = (DisplayName = "On MetaSoundPreset Widget Constructed"))
 	void OnConstructed(UMetaSoundBuilderBase* Builder);
+
+	// Called when the MetaSound starts and stops auditioning. Provides a reference to the audio component when auditioning starts, and returns nullptr when auditioning stops.
+	UFUNCTION(BlueprintCallable, BlueprintNativeEvent, Category = "MetaSound Preset Widget", meta = (DisplayName = "On MetaSound Audition State Changed"))
+	void OnAuditionStateChanged(UAudioComponent* AudioComponent, bool bIsAuditioning);
 };
