Added check to prevent possible null dereference crash

- Previous CL 45421692 added a check for IncludedMetaSounds, this CL adds the same check for ExcludedMetaSounds

#rnx

diff --git a/Engine/Plugins/Runtime/Metasound/Source/MetasoundEditor/Private/MetasoundEditor.cpp b/Engine/Plugins/Runtime/Metasound/Source/MetasoundEditor/Private/MetasoundEditor.cpp
--- a/Engine/Plugins/Runtime/Metasound/Source/MetasoundEditor/Private/MetasoundEditor.cpp
+++ b/Engine/Plugins/Runtime/Metasound/Source/MetasoundEditor/Private/MetasoundEditor.cpp
@@ -1844,6 +1844,11 @@ namespace Metasound
 					{
 						for (TScriptInterface<IMetaSoundDocumentInterface> ExcludedMetaSound : ExcludedMetaSounds)
 						{
+							if (!ExcludedMetaSound)
+							{
+								continue;
+							}
+
 							// This MetaSound or the referenced MetaSound is excluded 
 							const FMetasoundFrontendGraphClass& ExcludedGraphClass = ExcludedMetaSound->GetConstDocument().RootGraph;
 							const FNodeRegistryKey ExcludedRegistryKey(ExcludedGraphClass);
