[Metasound] Trigger Filter display name is now Trigger Chance to better convey its meaning to first-time users
+ Additional keyword when searching: "chance"
+ Tails output advanced, because most of the time a user only need their triggered filtered by their probability and not the 1-p one
+ For similar reasons, Reset and Seed inputs are also advanced
+ Updated pins & node tooltips to be more user friendly and less probabilistey

diff --git a/Engine/Plugins/Runtime/Metasound/Source/MetasoundStandardNodes/Private/MetasoundTriggerCoinNode.cpp b/Engine/Plugins/Runtime/Metasound/Source/MetasoundStandardNodes/Private/MetasoundTriggerCoinNode.cpp
--- a/Engine/Plugins/Runtime/Metasound/Source/MetasoundStandardNodes/Private/MetasoundTriggerCoinNode.cpp
+++ b/Engine/Plugins/Runtime/Metasound/Source/MetasoundStandardNodes/Private/MetasoundTriggerCoinNode.cpp
@@ -25,9 +25,9 @@ namespace Metasound
 		METASOUND_PARAM(InputTrigger, "Trigger", "Input trigger.");
 		METASOUND_PARAM(InputReset, "Reset", "Trigger to reset the random sequence with the supplied seed. Useful to get randomized repetition.");
 		METASOUND_PARAM(InputSeed, "Seed", "The seed value to use for the random node. Set to -1 to use a random seed.");
-		METASOUND_PARAM(InputProbability, "Probability", "Probability that the output trigger will get called. 0.0 is always tails, 1.0 is always heads, and 0.5 is a 50% chance for either.");
-		METASOUND_PARAM(OutputTrueTrigger, "Heads", "The first possible output trigger.");
-		METASOUND_PARAM(OutputFalseTrigger, "Tails", "The second possible output trigger.");
+		METASOUND_PARAM(InputProbability, "Probability", "Probability that the output trigger will get called. 1.0 is always heads, 0.0 is always tails, and 0.5 is a 50% chance for either.");
+		METASOUND_PARAM(OutputTrueTrigger, "Heads", "The first possible output trigger, chosen with Probability chance.");
+		METASOUND_PARAM(OutputFalseTrigger, "Tails", "The second possible output trigger, chosen with (1-Probability) chance).");
 	}
 
 	class FTriggerCoinOperator : public TExecutableOperator<FTriggerCoinOperator>
@@ -183,13 +183,13 @@ namespace Metasound
 		static const FVertexInterface Interface(
 			FInputVertexInterface(
 				TInputDataVertex<FTrigger>(METASOUND_GET_PARAM_NAME_AND_METADATA(InputTrigger)),
-				TInputDataVertex<FTrigger>(METASOUND_GET_PARAM_NAME_AND_METADATA(InputReset)),
-				TInputDataVertex<int32>(METASOUND_GET_PARAM_NAME_AND_METADATA(InputSeed), -1),
+				TInputDataVertex<FTrigger>(METASOUND_GET_PARAM_NAME_AND_METADATA_ADVANCED(InputReset)),
+				TInputDataVertex<int32>(METASOUND_GET_PARAM_NAME_AND_METADATA_ADVANCED(InputSeed), -1),
 				TInputDataVertex<float>(METASOUND_GET_PARAM_NAME_AND_METADATA(InputProbability), 1.0f)
 			),
 			FOutputVertexInterface(
 				TOutputDataVertex<FTrigger>(METASOUND_GET_PARAM_NAME_AND_METADATA(OutputTrueTrigger)),
-				TOutputDataVertex<FTrigger>(METASOUND_GET_PARAM_NAME_AND_METADATA(OutputFalseTrigger))
+				TOutputDataVertex<FTrigger>(METASOUND_GET_PARAM_NAME_AND_METADATA_ADVANCED(OutputFalseTrigger))
 			)
 		);
 
@@ -204,13 +204,19 @@ namespace Metasound
 			Info.ClassName = { StandardNodes::Namespace, TEXT("Trigger Filter"), TEXT("") };
 			Info.MajorVersion = 1;
 			Info.MinorVersion = 0;
-			Info.DisplayName = METASOUND_LOCTEXT("Metasound_TriggerCoinNodeDisplayName", "Trigger Filter");
-			Info.Description = METASOUND_LOCTEXT("Metasound_TriggerCoinNodeDescription", "When triggered, randomly sends one of two output triggers based on a given probability.");
+			Info.DisplayName = METASOUND_LOCTEXT("Metasound_TriggerCoinNodeDisplayName", "Trigger Chance");
+			Info.Description = METASOUND_LOCTEXT("Metasound_TriggerCoinNodeDescription", "When triggered, Heads has the given probability chance to trigger, otherwise it's Tail.");
 			Info.Author = PluginAuthor;
 			Info.PromptIfMissing = PluginNodeMissingPrompt;
 			Info.DefaultInterface = GetVertexInterface();
 			Info.CategoryHierarchy.Emplace(NodeCategories::Trigger);
-			Info.Keywords = { METASOUND_LOCTEXT("TriggerFilterRandomKeyword", "Random"), METASOUND_LOCTEXT("TriggerFilterProbabilityKeyword", "Probability"), METASOUND_LOCTEXT("TriggerFilterCoinKeyword", "Coin"), METASOUND_LOCTEXT("TriggerFilterBernoulliKeyword", "Bernoulli Gate")};
+			Info.Keywords = {
+				METASOUND_LOCTEXT("TriggerFilterRandomKeyword", "Random"),
+				METASOUND_LOCTEXT("TriggerFilterProbabilityKeyword", "Probability"),
+				METASOUND_LOCTEXT("TriggerFilterChanceKeyword", "Chance"),
+				METASOUND_LOCTEXT("TriggerFilterCoinKeyword", "Coin"),
+				METASOUND_LOCTEXT("TriggerFilterBernoulliKeyword", "Bernoulli Gate")
+			};
 
 			return Info;
 		};
