Fix metasound bugs by using transactionally safe critical sections.

diff --git a/Engine/Plugins/Runtime/Metasound/Source/MetasoundFrontend/Private/MetasoundFrontendDataTypeRegistry.cpp b/Engine/Plugins/Runtime/Metasound/Source/MetasoundFrontend/Private/MetasoundFrontendDataTypeRegistry.cpp
--- a/Engine/Plugins/Runtime/Metasound/Source/MetasoundFrontend/Private/MetasoundFrontendDataTypeRegistry.cpp
+++ b/Engine/Plugins/Runtime/Metasound/Source/MetasoundFrontend/Private/MetasoundFrontendDataTypeRegistry.cpp
@@ -4,6 +4,7 @@
 #include "MetasoundLiteral.h"
 #include "MetasoundTrace.h"
 #include "Misc/ScopeLock.h"
+#include "Misc/TransactionallySafeCriticalSection.h"
 #include "Misc/App.h"
 
 namespace Metasound
@@ -694,8 +695,8 @@ PRAGMA_ENABLE_DEPRECATION_WARNINGS
 				// UObject type names to DataTypeNames
 				TMap<const UClass*, FName> RegisteredObjectClasses;
 
-				mutable FCriticalSection RegistryMapMutex;
-				mutable FCriticalSection RegistryObjectMapMutex;
+				mutable FTransactionallySafeCriticalSection RegistryMapMutex;
+				mutable FTransactionallySafeCriticalSection RegistryObjectMapMutex;
 			};
 
 			bool FDataTypeRegistry::RegisterDataType(TUniquePtr<IDataTypeRegistryEntry>&& InEntry)
@@ -709,8 +710,8 @@ PRAGMA_ENABLE_DEPRECATION_WARNINGS
 					const FName Name = Entry->GetDataTypeInfo().DataTypeName;
 
 					{
-						FScopeLock RegistryLock(&RegistryMapMutex);
-						FScopeLock ObjectRegistryLock(&RegistryObjectMapMutex);
+						UE::TScopeLock RegistryLock(RegistryMapMutex);
+						UE::TScopeLock ObjectRegistryLock(RegistryObjectMapMutex);
 
 						if (!ensureAlwaysMsgf(!RegisteredDataTypes.Contains(Name),
 							TEXT("Name collision when trying to register Metasound Data Type [Name:%s]. DataType must have "
@@ -763,7 +764,7 @@ PRAGMA_ENABLE_DEPRECATION_WARNINGS
 
 			void FDataTypeRegistry::GetRegisteredDataTypeNames(TArray<FName>& OutNames) const
 			{
-				FScopeLock Lock(&RegistryMapMutex);
+				UE::TScopeLock Lock(RegistryMapMutex);
 				RegisteredDataTypes.GetKeys(OutNames);
 			}
 
@@ -776,7 +777,7 @@ PRAGMA_ENABLE_DEPRECATION_WARNINGS
 			{
 				if (InObject)
 				{
-					FScopeLock Lock(&RegistryObjectMapMutex);
+					UE::TScopeLock Lock(RegistryObjectMapMutex);
 					for (UClass* Class = InObject->GetClass(); Class != UObject::StaticClass(); Class = Class->GetSuperClass())
 					{
 						if (const FName* DataTypeName = RegisteredObjectClasses.Find(Class))
@@ -810,7 +811,7 @@ PRAGMA_ENABLE_DEPRECATION_WARNINGS
 
 			void FDataTypeRegistry::IterateDataTypeInfo(TFunctionRef<void(const FDataTypeRegistryInfo&)> InFunction) const
 			{
-				FScopeLock Lock(&RegistryMapMutex);
+				UE::TScopeLock Lock(RegistryMapMutex);
 				for (const TPair<FName, TSharedRef<IDataTypeRegistryEntry, ESPMode::ThreadSafe>>& Entry : RegisteredDataTypes)
 				{
 					InFunction(Entry.Value->GetDataTypeInfo());
@@ -819,7 +820,7 @@ PRAGMA_ENABLE_DEPRECATION_WARNINGS
 
 			bool FDataTypeRegistry::IsRegistered(const FName& InDataType) const
 			{
-				FScopeLock Lock(&RegistryMapMutex);
+				UE::TScopeLock Lock(RegistryMapMutex);
 				return RegisteredDataTypes.Contains(InDataType);
 			}
 
@@ -965,7 +966,7 @@ PRAGMA_ENABLE_DEPRECATION_WARNINGS
 					return false;
 				}
 
-				FScopeLock Lock(&RegistryObjectMapMutex);
+				UE::TScopeLock Lock(RegistryObjectMapMutex);
 				UClass* ObjectClass = InObject->GetClass();
 				while (ObjectClass != UObject::StaticClass())
 				{
@@ -1018,8 +1019,8 @@ PRAGMA_ENABLE_DEPRECATION_WARNINGS
 					{
 						if (InUObject)
 						{
-							FScopeLock ObjectRegistryLock(&RegistryObjectMapMutex);
-							FScopeLock RegistryLock(&RegistryMapMutex);
+							UE::TScopeLock ObjectRegistryLock(RegistryObjectMapMutex);
+							UE::TScopeLock RegistryLock(RegistryMapMutex);
 
 							const UClass* ObjectClass = InUObject->GetClass();
 							while (ObjectClass != UObject::StaticClass())
@@ -1500,7 +1501,7 @@ PRAGMA_ENABLE_DEPRECATION_WARNINGS
 
 			const IDataTypeRegistryEntry* FDataTypeRegistry::FindDataTypeEntry(const FName& InDataTypeName) const
 			{
-				FScopeLock Lock(&RegistryMapMutex);
+				UE::TScopeLock Lock(RegistryMapMutex);
 				{
 					const TSharedRef<IDataTypeRegistryEntry, ESPMode::ThreadSafe>* Entry = RegisteredDataTypes.Find(InDataTypeName);
 
diff --git a/Engine/Plugins/Runtime/Metasound/Source/MetasoundFrontend/Private/MetasoundFrontendNodeClassRegistry.cpp b/Engine/Plugins/Runtime/Metasound/Source/MetasoundFrontend/Private/MetasoundFrontendNodeClassRegistry.cpp
--- a/Engine/Plugins/Runtime/Metasound/Source/MetasoundFrontend/Private/MetasoundFrontendNodeClassRegistry.cpp
+++ b/Engine/Plugins/Runtime/Metasound/Source/MetasoundFrontend/Private/MetasoundFrontendNodeClassRegistry.cpp
@@ -431,7 +431,7 @@ PRAGMA_ENABLE_DEPRECATION_WARNINGS
 		METASOUND_LLM_SCOPE;
 		METASOUND_TRACE_CPUPROFILER_EVENT_SCOPE(metasound::FNodeClassRegistry::RegisterPendingNodes);
 		{
-			FScopeLock ScopeLock(&LazyInitCommandCritSection);
+			UE::TScopeLock ScopeLock(LazyInitCommandCritSection);
 
 			for (TUniqueFunction<void()>& Command : LazyInitCommands)
 			{
@@ -458,7 +458,7 @@ PRAGMA_ENABLE_DEPRECATION_WARNINGS
 	bool FNodeClassRegistry::EnqueueInitCommand(TUniqueFunction<void()>&& InFunc)
 	{
 
-		FScopeLock ScopeLock(&LazyInitCommandCritSection);
+		UE::TScopeLock ScopeLock(LazyInitCommandCritSection);
 		if (LazyInitCommands.Num() >= MaxNumNodesAndDatatypesToInitialize)
 		{
 			UE_LOG(LogMetaSound, Warning, TEXT("Registering more that %d nodes and datatypes for metasounds! Consider increasing MetasoundFrontendRegistryContainer::MaxNumNodesAndDatatypesToInitialize."), MaxNumNodesAndDatatypesToInitialize);
@@ -470,7 +470,7 @@ PRAGMA_ENABLE_DEPRECATION_WARNINGS
 
 	void FNodeClassRegistry::SetObjectReferencer(TUniquePtr<IObjectReferencer> InReferencer)
 	{
-		FScopeLock LockActiveReg(&ActiveRegistrationTasksCriticalSection);
+		UE::TScopeLock LockActiveReg(ActiveRegistrationTasksCriticalSection);
 		checkf(ActiveRegistrationTasks.IsEmpty(), TEXT("Object Referencer cannot be set while registry is actively being manipulated"));
 		ObjectReferencer = MoveTemp(InReferencer);
 	}
@@ -649,13 +649,13 @@ PRAGMA_ENABLE_DEPRECATION_WARNINGS
 
 	void FNodeClassRegistry::AddRegistrationTask(const FGraphRegistryKey& InKey, FActiveRegistrationTaskInfo&& TaskInfo)
 	{
-		FScopeLock LockActiveReg(&ActiveRegistrationTasksCriticalSection);
+		UE::TScopeLock LockActiveReg(ActiveRegistrationTasksCriticalSection);
 		ActiveRegistrationTasks.FindOrAdd(InKey.NodeKey).Add(MoveTemp(TaskInfo));
 	}
 
 	void FNodeClassRegistry::RemoveRegistrationTask(const FGraphRegistryKey& InKey, FNodeClassRegistryTransaction::ETransactionType TransactionType)
 	{
-		FScopeLock LockActiveReg(&ActiveRegistrationTasksCriticalSection);
+		UE::TScopeLock LockActiveReg(ActiveRegistrationTasksCriticalSection);
 
 		int32 NumRemoved = 0;
 		if (InKey.AssetPath.IsNull())  // Null provided path instructs to remove all tasks related to the underlying Node Registry Key
@@ -697,7 +697,7 @@ PRAGMA_ENABLE_DEPRECATION_WARNINGS
 
 	void FNodeClassRegistry::AddDocumentReference(TScriptInterface<IMetaSoundDocumentInterface> DocumentInterface)
 	{
-		FScopeLock LockActiveReg(&ObjectReferencerCriticalSection);
+		UE::TScopeLock LockActiveReg(ObjectReferencerCriticalSection);
 		if (UObject* Object = DocumentInterface.GetObject())
 		{
 			if (ObjectReferencer)
@@ -709,7 +709,7 @@ PRAGMA_ENABLE_DEPRECATION_WARNINGS
 
 	void FNodeClassRegistry::RemoveDocumentReference(TScriptInterface<IMetaSoundDocumentInterface> DocumentInterface)
 	{
-		FScopeLock LockActiveReg(&ObjectReferencerCriticalSection);
+		UE::TScopeLock LockActiveReg(ObjectReferencerCriticalSection);
 		if (UObject* Object = DocumentInterface.GetObject())
 		{
 			if (ObjectReferencer)
@@ -723,7 +723,7 @@ PRAGMA_ENABLE_DEPRECATION_WARNINGS
 	{
 		using namespace RegistryPrivate;
 
-		FScopeLock Lock(&RegistryMapsCriticalSection);
+		UE::TScopeLock Lock(RegistryMapsCriticalSection);
 
 #if !NO_LOGGING
 		if (RegisteredGraphs.Contains(InKey))
@@ -741,7 +741,7 @@ PRAGMA_ENABLE_DEPRECATION_WARNINGS
 
 		METASOUND_TRACE_CPUPROFILER_EVENT_SCOPE_TEXT(*InKey.ToString(TEXT("FNodeClassRegistry::UnregisterGraphInternal")));
 
-		FScopeLock Lock(&RegistryMapsCriticalSection);
+		UE::TScopeLock Lock(RegistryMapsCriticalSection);
 		{
 			if (!RegisteredGraphs.Contains(InKey))
 			{
@@ -845,7 +845,7 @@ PRAGMA_ENABLE_DEPRECATION_WARNINGS
 
 		TSharedPtr<const FGraph> Graph;
 		{
-			FScopeLock Lock(&RegistryMapsCriticalSection);
+			UE::TScopeLock Lock(RegistryMapsCriticalSection);
 			if (const TSharedPtr<const FGraph>* RegisteredGraph = RegisteredGraphs.Find(InKey))
 			{
 				Graph = *RegisteredGraph;
@@ -881,7 +881,7 @@ PRAGMA_ENABLE_DEPRECATION_WARNINGS
 
 		{
 			TSharedRef<INodeClassRegistryEntry, ESPMode::ThreadSafe> Entry(InEntry.Release());
-			FScopeLock Lock(&RegistryMapsCriticalSection);
+			UE::TScopeLock Lock(RegistryMapsCriticalSection);
 			// check to see if an identical node was already registered, and log if necessary
 			// Store registry elements in map so nodes can be queried using registry key.
 			RegisteredNodes.Add(Key, Entry);
@@ -1014,7 +1014,7 @@ PRAGMA_ENABLE_DEPRECATION_WARNINGS
 			Key = FNodeClassRegistryKey(Entry->GetClassInfo());
 
 			{
-				FScopeLock Lock(&RegistryMapsCriticalSection);
+				UE::TScopeLock Lock(RegistryMapsCriticalSection);
 				// check to see if an identical node was already registered, and log
 				ensureAlwaysMsgf(
 					!RegisteredNodeTemplates.Contains(Key),
@@ -1042,7 +1042,7 @@ PRAGMA_ENABLE_DEPRECATION_WARNINGS
 		{
 			METASOUND_TRACE_CPUPROFILER_EVENT_SCOPE_TEXT(*FString::Printf(TEXT("FNodeClassRegistry::UnregisterNodeInternal key %s"), *InKey.ToString()))
 
-			FScopeLock Lock(&RegistryMapsCriticalSection);
+			UE::TScopeLock Lock(RegistryMapsCriticalSection);
 			if (const TSharedRef<INodeClassRegistryEntry, ESPMode::ThreadSafe>* EntryPtr = RegisteredNodes.Find(InKey))
 			{
 				const TSharedRef<INodeClassRegistryEntry>& Entry = *EntryPtr;
@@ -1092,7 +1092,7 @@ PRAGMA_ENABLE_DEPRECATION_WARNINGS
 				TransactionBuffer->AddTransaction(FNodeClassRegistryTransaction(FNodeClassRegistryTransaction::ETransactionType::NodeUnregistration, Entry->GetClassInfo(), Timestamp));
 
 				{
-					FScopeLock Lock(&RegistryMapsCriticalSection);
+					UE::TScopeLock Lock(RegistryMapsCriticalSection);
 					RegisteredNodeTemplates.Remove(InKey);
 				}
 				return true;
@@ -1127,7 +1127,7 @@ PRAGMA_ENABLE_DEPRECATION_WARNINGS
 	{
 		auto IsNodeRegisteredInternal = [this, &InKey]() -> bool
 		{
-			FScopeLock Lock(&RegistryMapsCriticalSection);
+			UE::TScopeLock Lock(RegistryMapsCriticalSection);
 			return RegisteredNodes.Contains(InKey) || RegisteredNodeTemplates.Contains(InKey);
 		};
 
@@ -1147,7 +1147,7 @@ PRAGMA_ENABLE_DEPRECATION_WARNINGS
 		WaitForAsyncGraphRegistration(InKey);
 
 		{
-			FScopeLock Lock(&RegistryMapsCriticalSection);
+			UE::TScopeLock Lock(RegistryMapsCriticalSection);
 			return RegisteredGraphs.Contains(InKey);
 		}
 	}
@@ -1325,7 +1325,7 @@ PRAGMA_ENABLE_DEPRECATION_WARNINGS
 	{
 		auto TryAccessNodeEntry = [this, &InKey, &InFunc]() -> bool
 		{
-			FScopeLock Lock(&RegistryMapsCriticalSection);
+			UE::TScopeLock Lock(RegistryMapsCriticalSection);
 			if (const TSharedRef<INodeClassRegistryEntry, ESPMode::ThreadSafe>* Entry = RegisteredNodes.Find(InKey))
 			{
 				InFunc(*(*Entry));
@@ -1359,7 +1359,7 @@ PRAGMA_ENABLE_DEPRECATION_WARNINGS
 
 	const INodeTemplateRegistryEntry* FNodeClassRegistry::FindNodeTemplateEntry(const FNodeClassRegistryKey& InKey) const
 	{
-		FScopeLock Lock(&RegistryMapsCriticalSection);
+		UE::TScopeLock Lock(RegistryMapsCriticalSection);
 		if (const TSharedRef<INodeTemplateRegistryEntry, ESPMode::ThreadSafe>* Entry = RegisteredNodeTemplates.Find(InKey))
 		{
 			return &Entry->Get();
@@ -1386,7 +1386,7 @@ PRAGMA_ENABLE_DEPRECATION_WARNINGS
 
 		TArray<FTask> TasksToWaitFor;
 		{
-			FScopeLock Lock(&ActiveRegistrationTasksCriticalSection);
+			UE::TScopeLock Lock(ActiveRegistrationTasksCriticalSection);
 			if (const TArray<FActiveRegistrationTaskInfo>* FoundTasks = ActiveRegistrationTasks.Find(InRegistryKey))
 			{
 				// Filter by asset path or ignore if not provided
diff --git a/Engine/Plugins/Runtime/Metasound/Source/MetasoundFrontend/Private/MetasoundFrontendNodeClassRegistryPrivate.h b/Engine/Plugins/Runtime/Metasound/Source/MetasoundFrontend/Private/MetasoundFrontendNodeClassRegistryPrivate.h
--- a/Engine/Plugins/Runtime/Metasound/Source/MetasoundFrontend/Private/MetasoundFrontendNodeClassRegistryPrivate.h
+++ b/Engine/Plugins/Runtime/Metasound/Source/MetasoundFrontend/Private/MetasoundFrontendNodeClassRegistryPrivate.h
@@ -4,7 +4,7 @@
 
 #include "Containers/Array.h"
 #include "Containers/Map.h"
-#include "HAL/CriticalSection.h"
+#include "Misc/TransactionallySafeCriticalSection.h"
 #include "Tasks/Pipe.h"
 
 #include "MetasoundFrontendDataTypeRegistry.h"
@@ -202,9 +202,9 @@ namespace Metasound::Frontend
 
 		TSharedRef<FNodeClassRegistryTransactionBuffer> TransactionBuffer;
 
-		mutable FCriticalSection RegistryMapsCriticalSection;
-		mutable FCriticalSection ActiveRegistrationTasksCriticalSection;
-		mutable FCriticalSection ObjectReferencerCriticalSection;
+		mutable FTransactionallySafeCriticalSection RegistryMapsCriticalSection;
+		mutable FTransactionallySafeCriticalSection ActiveRegistrationTasksCriticalSection;
+		mutable FTransactionallySafeCriticalSection ObjectReferencerCriticalSection;
 
 		UE::Tasks::FPipe AsyncRegistrationPipe;
 		TMap<FNodeClassRegistryKey, TArray<FActiveRegistrationTaskInfo>> ActiveRegistrationTasks;
