[Backout] - CL45672698
#rnx
Original CL Desc
-----------------------------------------------------------------
Fix metasound bugs by using transactionally safe critical sections.

diff --git a/Engine/Plugins/Runtime/Metasound/Source/MetasoundFrontend/Private/MetasoundFrontendDataTypeRegistry.cpp b/Engine/Plugins/Runtime/Metasound/Source/MetasoundFrontend/Private/MetasoundFrontendDataTypeRegistry.cpp
--- a/Engine/Plugins/Runtime/Metasound/Source/MetasoundFrontend/Private/MetasoundFrontendDataTypeRegistry.cpp
+++ b/Engine/Plugins/Runtime/Metasound/Source/MetasoundFrontend/Private/MetasoundFrontendDataTypeRegistry.cpp
@@ -4,7 +4,6 @@
 #include "MetasoundLiteral.h"
 #include "MetasoundTrace.h"
 #include "Misc/ScopeLock.h"
-#include "Misc/TransactionallySafeCriticalSection.h"
 #include "Misc/App.h"
 
 namespace Metasound
@@ -695,8 +694,8 @@ PRAGMA_ENABLE_DEPRECATION_WARNINGS
 				// UObject type names to DataTypeNames
 				TMap<const UClass*, FName> RegisteredObjectClasses;
 
-				mutable FTransactionallySafeCriticalSection RegistryMapMutex;
-				mutable FTransactionallySafeCriticalSection RegistryObjectMapMutex;
+				mutable FCriticalSection RegistryMapMutex;
+				mutable FCriticalSection RegistryObjectMapMutex;
 			};
 
 			bool FDataTypeRegistry::RegisterDataType(TUniquePtr<IDataTypeRegistryEntry>&& InEntry)
@@ -710,8 +709,8 @@ PRAGMA_ENABLE_DEPRECATION_WARNINGS
 					const FName Name = Entry->GetDataTypeInfo().DataTypeName;
 
 					{
-						UE::TScopeLock RegistryLock(RegistryMapMutex);
-						UE::TScopeLock ObjectRegistryLock(RegistryObjectMapMutex);
+						FScopeLock RegistryLock(&RegistryMapMutex);
+						FScopeLock ObjectRegistryLock(&RegistryObjectMapMutex);
 
 						if (!ensureAlwaysMsgf(!RegisteredDataTypes.Contains(Name),
 							TEXT("Name collision when trying to register Metasound Data Type [Name:%s]. DataType must have "
@@ -764,7 +763,7 @@ PRAGMA_ENABLE_DEPRECATION_WARNINGS
 
 			void FDataTypeRegistry::GetRegisteredDataTypeNames(TArray<FName>& OutNames) const
 			{
-				UE::TScopeLock Lock(RegistryMapMutex);
+				FScopeLock Lock(&RegistryMapMutex);
 				RegisteredDataTypes.GetKeys(OutNames);
 			}
 
@@ -777,7 +776,7 @@ PRAGMA_ENABLE_DEPRECATION_WARNINGS
 			{
 				if (InObject)
 				{
-					UE::TScopeLock Lock(RegistryObjectMapMutex);
+					FScopeLock Lock(&RegistryObjectMapMutex);
 					for (UClass* Class = InObject->GetClass(); Class != UObject::StaticClass(); Class = Class->GetSuperClass())
 					{
 						if (const FName* DataTypeName = RegisteredObjectClasses.Find(Class))
@@ -811,7 +810,7 @@ PRAGMA_ENABLE_DEPRECATION_WARNINGS
 
 			void FDataTypeRegistry::IterateDataTypeInfo(TFunctionRef<void(const FDataTypeRegistryInfo&)> InFunction) const
 			{
-				UE::TScopeLock Lock(RegistryMapMutex);
+				FScopeLock Lock(&RegistryMapMutex);
 				for (const TPair<FName, TSharedRef<IDataTypeRegistryEntry, ESPMode::ThreadSafe>>& Entry : RegisteredDataTypes)
 				{
 					InFunction(Entry.Value->GetDataTypeInfo());
@@ -820,7 +819,7 @@ PRAGMA_ENABLE_DEPRECATION_WARNINGS
 
 			bool FDataTypeRegistry::IsRegistered(const FName& InDataType) const
 			{
-				UE::TScopeLock Lock(RegistryMapMutex);
+				FScopeLock Lock(&RegistryMapMutex);
 				return RegisteredDataTypes.Contains(InDataType);
 			}
 
@@ -966,7 +965,7 @@ PRAGMA_ENABLE_DEPRECATION_WARNINGS
 					return false;
 				}
 
-				UE::TScopeLock Lock(RegistryObjectMapMutex);
+				FScopeLock Lock(&RegistryObjectMapMutex);
 				UClass* ObjectClass = InObject->GetClass();
 				while (ObjectClass != UObject::StaticClass())
 				{
@@ -1019,8 +1018,8 @@ PRAGMA_ENABLE_DEPRECATION_WARNINGS
 					{
 						if (InUObject)
 						{
-							UE::TScopeLock ObjectRegistryLock(RegistryObjectMapMutex);
-							UE::TScopeLock RegistryLock(RegistryMapMutex);
+							FScopeLock ObjectRegistryLock(&RegistryObjectMapMutex);
+							FScopeLock RegistryLock(&RegistryMapMutex);
 
 							const UClass* ObjectClass = InUObject->GetClass();
 							while (ObjectClass != UObject::StaticClass())
@@ -1501,7 +1500,7 @@ PRAGMA_ENABLE_DEPRECATION_WARNINGS
 
 			const IDataTypeRegistryEntry* FDataTypeRegistry::FindDataTypeEntry(const FName& InDataTypeName) const
 			{
-				UE::TScopeLock Lock(RegistryMapMutex);
+				FScopeLock Lock(&RegistryMapMutex);
 				{
 					const TSharedRef<IDataTypeRegistryEntry, ESPMode::ThreadSafe>* Entry = RegisteredDataTypes.Find(InDataTypeName);
 
diff --git a/Engine/Plugins/Runtime/Metasound/Source/MetasoundFrontend/Private/MetasoundFrontendNodeClassRegistry.cpp b/Engine/Plugins/Runtime/Metasound/Source/MetasoundFrontend/Private/MetasoundFrontendNodeClassRegistry.cpp
--- a/Engine/Plugins/Runtime/Metasound/Source/MetasoundFrontend/Private/MetasoundFrontendNodeClassRegistry.cpp
+++ b/Engine/Plugins/Runtime/Metasound/Source/MetasoundFrontend/Private/MetasoundFrontendNodeClassRegistry.cpp
@@ -431,7 +431,7 @@ PRAGMA_ENABLE_DEPRECATION_WARNINGS
 		METASOUND_LLM_SCOPE;
 		METASOUND_TRACE_CPUPROFILER_EVENT_SCOPE(metasound::FNodeClassRegistry::RegisterPendingNodes);
 		{
-			UE::TScopeLock ScopeLock(LazyInitCommandCritSection);
+			FScopeLock ScopeLock(&LazyInitCommandCritSection);
 
 			for (TUniqueFunction<void()>& Command : LazyInitCommands)
 			{
@@ -458,7 +458,7 @@ PRAGMA_ENABLE_DEPRECATION_WARNINGS
 	bool FNodeClassRegistry::EnqueueInitCommand(TUniqueFunction<void()>&& InFunc)
 	{
 
-		UE::TScopeLock ScopeLock(LazyInitCommandCritSection);
+		FScopeLock ScopeLock(&LazyInitCommandCritSection);
 		if (LazyInitCommands.Num() >= MaxNumNodesAndDatatypesToInitialize)
 		{
 			UE_LOG(LogMetaSound, Warning, TEXT("Registering more that %d nodes and datatypes for metasounds! Consider increasing MetasoundFrontendRegistryContainer::MaxNumNodesAndDatatypesToInitialize."), MaxNumNodesAndDatatypesToInitialize);
@@ -470,7 +470,7 @@ PRAGMA_ENABLE_DEPRECATION_WARNINGS
 
 	void FNodeClassRegistry::SetObjectReferencer(TUniquePtr<IObjectReferencer> InReferencer)
 	{
-		UE::TScopeLock LockActiveReg(ActiveRegistrationTasksCriticalSection);
+		FScopeLock LockActiveReg(&ActiveRegistrationTasksCriticalSection);
 		checkf(ActiveRegistrationTasks.IsEmpty(), TEXT("Object Referencer cannot be set while registry is actively being manipulated"));
 		ObjectReferencer = MoveTemp(InReferencer);
 	}
@@ -649,13 +649,13 @@ PRAGMA_ENABLE_DEPRECATION_WARNINGS
 
 	void FNodeClassRegistry::AddRegistrationTask(const FGraphRegistryKey& InKey, FActiveRegistrationTaskInfo&& TaskInfo)
 	{
-		UE::TScopeLock LockActiveReg(ActiveRegistrationTasksCriticalSection);
+		FScopeLock LockActiveReg(&ActiveRegistrationTasksCriticalSection);
 		ActiveRegistrationTasks.FindOrAdd(InKey.NodeKey).Add(MoveTemp(TaskInfo));
 	}
 
 	void FNodeClassRegistry::RemoveRegistrationTask(const FGraphRegistryKey& InKey, FNodeClassRegistryTransaction::ETransactionType TransactionType)
 	{
-		UE::TScopeLock LockActiveReg(ActiveRegistrationTasksCriticalSection);
+		FScopeLock LockActiveReg(&ActiveRegistrationTasksCriticalSection);
 
 		int32 NumRemoved = 0;
 		if (InKey.AssetPath.IsNull())  // Null provided path instructs to remove all tasks related to the underlying Node Registry Key
@@ -697,7 +697,7 @@ PRAGMA_ENABLE_DEPRECATION_WARNINGS
 
 	void FNodeClassRegistry::AddDocumentReference(TScriptInterface<IMetaSoundDocumentInterface> DocumentInterface)
 	{
-		UE::TScopeLock LockActiveReg(ObjectReferencerCriticalSection);
+		FScopeLock LockActiveReg(&ObjectReferencerCriticalSection);
 		if (UObject* Object = DocumentInterface.GetObject())
 		{
 			if (ObjectReferencer)
@@ -709,7 +709,7 @@ PRAGMA_ENABLE_DEPRECATION_WARNINGS
 
 	void FNodeClassRegistry::RemoveDocumentReference(TScriptInterface<IMetaSoundDocumentInterface> DocumentInterface)
 	{
-		UE::TScopeLock LockActiveReg(ObjectReferencerCriticalSection);
+		FScopeLock LockActiveReg(&ObjectReferencerCriticalSection);
 		if (UObject* Object = DocumentInterface.GetObject())
 		{
 			if (ObjectReferencer)
@@ -723,7 +723,7 @@ PRAGMA_ENABLE_DEPRECATION_WARNINGS
 	{
 		using namespace RegistryPrivate;
 
-		UE::TScopeLock Lock(RegistryMapsCriticalSection);
+		FScopeLock Lock(&RegistryMapsCriticalSection);
 
 #if !NO_LOGGING
 		if (RegisteredGraphs.Contains(InKey))
@@ -741,7 +741,7 @@ PRAGMA_ENABLE_DEPRECATION_WARNINGS
 
 		METASOUND_TRACE_CPUPROFILER_EVENT_SCOPE_TEXT(*InKey.ToString(TEXT("FNodeClassRegistry::UnregisterGraphInternal")));
 
-		UE::TScopeLock Lock(RegistryMapsCriticalSection);
+		FScopeLock Lock(&RegistryMapsCriticalSection);
 		{
 			if (!RegisteredGraphs.Contains(InKey))
 			{
@@ -845,7 +845,7 @@ PRAGMA_ENABLE_DEPRECATION_WARNINGS
 
 		TSharedPtr<const FGraph> Graph;
 		{
-			UE::TScopeLock Lock(RegistryMapsCriticalSection);
+			FScopeLock Lock(&RegistryMapsCriticalSection);
 			if (const TSharedPtr<const FGraph>* RegisteredGraph = RegisteredGraphs.Find(InKey))
 			{
 				Graph = *RegisteredGraph;
@@ -881,7 +881,7 @@ PRAGMA_ENABLE_DEPRECATION_WARNINGS
 
 		{
 			TSharedRef<INodeClassRegistryEntry, ESPMode::ThreadSafe> Entry(InEntry.Release());
-			UE::TScopeLock Lock(RegistryMapsCriticalSection);
+			FScopeLock Lock(&RegistryMapsCriticalSection);
 			// check to see if an identical node was already registered, and log if necessary
 			// Store registry elements in map so nodes can be queried using registry key.
 			RegisteredNodes.Add(Key, Entry);
@@ -1014,7 +1014,7 @@ PRAGMA_ENABLE_DEPRECATION_WARNINGS
 			Key = FNodeClassRegistryKey(Entry->GetClassInfo());
 
 			{
-				UE::TScopeLock Lock(RegistryMapsCriticalSection);
+				FScopeLock Lock(&RegistryMapsCriticalSection);
 				// check to see if an identical node was already registered, and log
 				ensureAlwaysMsgf(
 					!RegisteredNodeTemplates.Contains(Key),
@@ -1042,7 +1042,7 @@ PRAGMA_ENABLE_DEPRECATION_WARNINGS
 		{
 			METASOUND_TRACE_CPUPROFILER_EVENT_SCOPE_TEXT(*FString::Printf(TEXT("FNodeClassRegistry::UnregisterNodeInternal key %s"), *InKey.ToString()))
 
-			UE::TScopeLock Lock(RegistryMapsCriticalSection);
+			FScopeLock Lock(&RegistryMapsCriticalSection);
 			if (const TSharedRef<INodeClassRegistryEntry, ESPMode::ThreadSafe>* EntryPtr = RegisteredNodes.Find(InKey))
 			{
 				const TSharedRef<INodeClassRegistryEntry>& Entry = *EntryPtr;
@@ -1092,7 +1092,7 @@ PRAGMA_ENABLE_DEPRECATION_WARNINGS
 				TransactionBuffer->AddTransaction(FNodeClassRegistryTransaction(FNodeClassRegistryTransaction::ETransactionType::NodeUnregistration, Entry->GetClassInfo(), Timestamp));
 
 				{
-					UE::TScopeLock Lock(RegistryMapsCriticalSection);
+					FScopeLock Lock(&RegistryMapsCriticalSection);
 					RegisteredNodeTemplates.Remove(InKey);
 				}
 				return true;
@@ -1127,7 +1127,7 @@ PRAGMA_ENABLE_DEPRECATION_WARNINGS
 	{
 		auto IsNodeRegisteredInternal = [this, &InKey]() -> bool
 		{
-			UE::TScopeLock Lock(RegistryMapsCriticalSection);
+			FScopeLock Lock(&RegistryMapsCriticalSection);
 			return RegisteredNodes.Contains(InKey) || RegisteredNodeTemplates.Contains(InKey);
 		};
 
@@ -1147,7 +1147,7 @@ PRAGMA_ENABLE_DEPRECATION_WARNINGS
 		WaitForAsyncGraphRegistration(InKey);
 
 		{
-			UE::TScopeLock Lock(RegistryMapsCriticalSection);
+			FScopeLock Lock(&RegistryMapsCriticalSection);
 			return RegisteredGraphs.Contains(InKey);
 		}
 	}
@@ -1325,7 +1325,7 @@ PRAGMA_ENABLE_DEPRECATION_WARNINGS
 	{
 		auto TryAccessNodeEntry = [this, &InKey, &InFunc]() -> bool
 		{
-			UE::TScopeLock Lock(RegistryMapsCriticalSection);
+			FScopeLock Lock(&RegistryMapsCriticalSection);
 			if (const TSharedRef<INodeClassRegistryEntry, ESPMode::ThreadSafe>* Entry = RegisteredNodes.Find(InKey))
 			{
 				InFunc(*(*Entry));
@@ -1359,7 +1359,7 @@ PRAGMA_ENABLE_DEPRECATION_WARNINGS
 
 	const INodeTemplateRegistryEntry* FNodeClassRegistry::FindNodeTemplateEntry(const FNodeClassRegistryKey& InKey) const
 	{
-		UE::TScopeLock Lock(RegistryMapsCriticalSection);
+		FScopeLock Lock(&RegistryMapsCriticalSection);
 		if (const TSharedRef<INodeTemplateRegistryEntry, ESPMode::ThreadSafe>* Entry = RegisteredNodeTemplates.Find(InKey))
 		{
 			return &Entry->Get();
@@ -1386,7 +1386,7 @@ PRAGMA_ENABLE_DEPRECATION_WARNINGS
 
 		TArray<FTask> TasksToWaitFor;
 		{
-			UE::TScopeLock Lock(ActiveRegistrationTasksCriticalSection);
+			FScopeLock Lock(&ActiveRegistrationTasksCriticalSection);
 			if (const TArray<FActiveRegistrationTaskInfo>* FoundTasks = ActiveRegistrationTasks.Find(InRegistryKey))
 			{
 				// Filter by asset path or ignore if not provided
diff --git a/Engine/Plugins/Runtime/Metasound/Source/MetasoundFrontend/Private/MetasoundFrontendNodeClassRegistryPrivate.h b/Engine/Plugins/Runtime/Metasound/Source/MetasoundFrontend/Private/MetasoundFrontendNodeClassRegistryPrivate.h
--- a/Engine/Plugins/Runtime/Metasound/Source/MetasoundFrontend/Private/MetasoundFrontendNodeClassRegistryPrivate.h
+++ b/Engine/Plugins/Runtime/Metasound/Source/MetasoundFrontend/Private/MetasoundFrontendNodeClassRegistryPrivate.h
@@ -4,7 +4,7 @@
 
 #include "Containers/Array.h"
 #include "Containers/Map.h"
-#include "Misc/TransactionallySafeCriticalSection.h"
+#include "HAL/CriticalSection.h"
 #include "Tasks/Pipe.h"
 
 #include "MetasoundFrontendDataTypeRegistry.h"
@@ -202,9 +202,9 @@ namespace Metasound::Frontend
 
 		TSharedRef<FNodeClassRegistryTransactionBuffer> TransactionBuffer;
 
-		mutable FTransactionallySafeCriticalSection RegistryMapsCriticalSection;
-		mutable FTransactionallySafeCriticalSection ActiveRegistrationTasksCriticalSection;
-		mutable FTransactionallySafeCriticalSection ObjectReferencerCriticalSection;
+		mutable FCriticalSection RegistryMapsCriticalSection;
+		mutable FCriticalSection ActiveRegistrationTasksCriticalSection;
+		mutable FCriticalSection ObjectReferencerCriticalSection;
 
 		UE::Tasks::FPipe AsyncRegistrationPipe;
 		TMap<FNodeClassRegistryKey, TArray<FActiveRegistrationTaskInfo>> ActiveRegistrationTasks;
