Prevent consecutive repetition when using MetaSound array random get node with No Repeats with an array size of 3

#rnx

diff --git a/Engine/Plugins/Runtime/Metasound/Source/MetasoundFrontend/Private/MetasoundArrayRandomNode.cpp b/Engine/Plugins/Runtime/Metasound/Source/MetasoundFrontend/Private/MetasoundArrayRandomNode.cpp
--- a/Engine/Plugins/Runtime/Metasound/Source/MetasoundFrontend/Private/MetasoundArrayRandomNode.cpp
+++ b/Engine/Plugins/Runtime/Metasound/Source/MetasoundFrontend/Private/MetasoundArrayRandomNode.cpp
@@ -74,10 +74,10 @@ namespace Metasound
 
 	void FArrayRandomGet::SetNoRepeatOrder(int32 InNoRepeatOrder)
 	{
-		// Only set the no repeats if the array length is larger than 3, otherwise, ignore no-repeats as it won't work anyway (it will end up repeating a clear pattern)
+		// Only set the no repeats if the array length is larger than 2, otherwise, ignore no-repeats as it won't work anyway (it will end up repeating a clear pattern)
 		// Make sure the no repeat order is between 0 and half of the max index
 		// This is to prevent clear patterns in the output that causes confusion by user. By clamping to half the max index we achieve maximum no-repeats behavior, which is the intent of a user trying to maximize no repeats.
-		if (MaxIndex > 3)
+		if (MaxIndex > 2)
 		{
 			// Check for -1 on no-repeat order. This means to utilize the maximum non-repetition, matched to the array size. This prevents user from having to adjust their no-repeats to match the array size themselves.
 			if (InNoRepeatOrder == -1)
