Add derived type check to MetaSound node configuration setter

#rnx

diff --git a/Engine/Plugins/Runtime/Metasound/Source/MetasoundFrontend/Private/MetasoundFrontendDocumentBuilder.cpp b/Engine/Plugins/Runtime/Metasound/Source/MetasoundFrontend/Private/MetasoundFrontendDocumentBuilder.cpp
--- a/Engine/Plugins/Runtime/Metasound/Source/MetasoundFrontend/Private/MetasoundFrontendDocumentBuilder.cpp
+++ b/Engine/Plugins/Runtime/Metasound/Source/MetasoundFrontend/Private/MetasoundFrontendDocumentBuilder.cpp
@@ -6592,8 +6592,19 @@ bool FMetaSoundFrontendDocumentBuilder::SetNodeConfiguration(const FGuid& InNode
 	if (const int32* NodeIndex = NodeCache.FindNodeIndex(InNodeID))
 	{
 		FMetasoundFrontendNode& Node = Graph.Nodes[*NodeIndex];
-		Node.Configuration = MoveTemp(InNodeConfiguration);
-		return UpdateNodeInterfaceFromConfiguration(InNodeID, InPageID);
+		// Node config derived type must match
+		if (const FMetasoundFrontendClass* Class = FindDependency(Node.ClassID))
+		{
+			TInstancedStruct<FMetaSoundFrontendNodeConfiguration> RegisteredConfiguration = CreateFrontendNodeConfiguration(Class->Metadata);
+			if (RegisteredConfiguration.GetScriptStruct() != InNodeConfiguration.GetScriptStruct())
+			{
+				UE_LOG(LogMetaSound, Warning, TEXT("Cannot set node configuration to type %s, expected registered type %s"), *InNodeConfiguration.GetScriptStruct()->GetFName().ToString(), *RegisteredConfiguration.GetScriptStruct()->GetFName().ToString());
+				return false;
+			}
+
+			Node.Configuration = MoveTemp(InNodeConfiguration);
+			return UpdateNodeInterfaceFromConfiguration(InNodeID, InPageID);
+		}
 	}
 	return false;
 }
