Fix for property matrix edits causing crashes when metasound editor attempting to synchronize references

#rnx

diff --git a/Engine/Plugins/Runtime/Metasound/Source/MetasoundEditor/Private/MetasoundEditorGraphBuilder.cpp b/Engine/Plugins/Runtime/Metasound/Source/MetasoundEditor/Private/MetasoundEditorGraphBuilder.cpp
--- a/Engine/Plugins/Runtime/Metasound/Source/MetasoundEditor/Private/MetasoundEditorGraphBuilder.cpp
+++ b/Engine/Plugins/Runtime/Metasound/Source/MetasoundEditor/Private/MetasoundEditorGraphBuilder.cpp
@@ -1963,7 +1963,13 @@ PRAGMA_ENABLE_DEPRECATION_WARNINGS
 							{
 								if (EditedMetaSound->IsReferencedAsset(*MetaSoundAsset))
 								{
-									EditedReferencingMetaSounds.Add(EditedMetaSound);
+									// Ensure the graph has been initialized. Its possible the "edited asset" is a MetaSound
+									// but its transient EdGraph has yet to be initialized due to the active editor not necessarily
+									// being the MetaSound editor (ex. property matrix editor).
+									if (const UEdGraph* EdGraph = EditedMetaSound->GetGraph())
+									{
+										EditedReferencingMetaSounds.Add(EditedMetaSound);
+									}
 								}
 							}
 						}
