Add MetaSound Node Configuration to Template Nodes
#rnx

diff --git a/Engine/Plugins/Runtime/Metasound/Source/MetasoundFrontend/Private/MetasoundFrontendDocumentBuilder.cpp b/Engine/Plugins/Runtime/Metasound/Source/MetasoundFrontend/Private/MetasoundFrontendDocumentBuilder.cpp
--- a/Engine/Plugins/Runtime/Metasound/Source/MetasoundFrontend/Private/MetasoundFrontendDocumentBuilder.cpp
+++ b/Engine/Plugins/Runtime/Metasound/Source/MetasoundFrontend/Private/MetasoundFrontendDocumentBuilder.cpp
@@ -1927,6 +1927,32 @@ const FMetasoundFrontendNode* FMetaSoundFrontendDocumentBuilder::AddNodeByTempla
 	return NewNode;
 }
 
+TInstancedStruct<FMetaSoundFrontendNodeConfiguration> FMetaSoundFrontendDocumentBuilder::CreateFrontendNodeConfiguration(const FMetasoundFrontendClassMetadata& InClassMetadata) const
+{
+	return CreateFrontendNodeConfiguration(InClassMetadata.GetType(), Metasound::Frontend::FNodeRegistryKey(InClassMetadata));
+}
+
+TInstancedStruct<FMetaSoundFrontendNodeConfiguration> FMetaSoundFrontendDocumentBuilder::CreateFrontendNodeConfiguration(EMetasoundFrontendClassType InClassType, const Metasound::Frontend::FNodeRegistryKey& InClassKey) const
+{
+	using namespace Metasound::Frontend;
+
+	TInstancedStruct<FMetaSoundFrontendNodeConfiguration> NodeConfiguration;
+
+	if (InClassType == EMetasoundFrontendClassType::Template)
+	{
+		if (const INodeTemplate* Template = INodeTemplateRegistry::Get().FindTemplate(InClassKey))
+		{
+			NodeConfiguration = Template->CreateFrontendTemplateNodeConfiguration();
+		}
+	}
+	else
+	{
+		NodeConfiguration = INodeClassRegistry::GetChecked().CreateFrontendNodeConfiguration(InClassKey);
+	}
+
+	return NodeConfiguration;
+}
+
 FMetasoundFrontendNode* FMetaSoundFrontendDocumentBuilder::AddNodeInternal(const FMetasoundFrontendClassMetadata& InClassMetadata, FFinalizeNodeFunctionRef FinalizeNode, const FGuid& InPageID, FGuid InNodeID, int32* NewNodeIndex)
 {
 	METASOUND_TRACE_CPUPROFILER_EVENT_SCOPE(FMetaSoundFrontendDocumentBuilder::AddNodeInternal);
@@ -1945,14 +1971,23 @@ FMetasoundFrontendNode* FMetaSoundFrontendDocumentBuilder::AddNodeInternal(const
 
 	if (const FMetasoundFrontendClass* Dependency = DocumentCache->FindDependency(InClassKey))
 	{
-		TInstancedStruct<FMetaSoundFrontendNodeConfiguration> NodeConfiguration = INodeClassRegistry::Get()->CreateFrontendNodeConfiguration(InClassKey);
+		// Get the array which holds the graph's nodes
 		FMetasoundFrontendDocument& Document = GetDocumentChecked();
 		FMetasoundFrontendGraph& Graph = Document.RootGraph.FindGraphChecked(InPageID);
 		TArray<FMetasoundFrontendNode>& Nodes = Graph.Nodes;
+
+		// Create a node configuration for the new node
+		TInstancedStruct<FMetaSoundFrontendNodeConfiguration> NodeConfiguration = CreateFrontendNodeConfiguration(Dependency->Metadata.GetType(), InClassKey);
+
+		// Add the new node into the graphs array
 		FMetasoundFrontendNode& Node = Nodes.Emplace_GetRef(*Dependency, MoveTemp(NodeConfiguration));
+
+		// Make sure the node utilizes the correct and unique NodeID
 		Node.UpdateID(InNodeID);
+
 		FinalizeNode(Node, InClassKey);
 
+		// Update document builder cache
 		const int32 NewIndex = Nodes.Num() - 1;
 		const IDocumentGraphNodeCache& NodeCache = DocumentCache->GetNodeCache(InPageID);
 		DocumentDelegates->FindNodeDelegatesChecked(InPageID).OnNodeAdded.Broadcast(NewIndex);
@@ -3198,13 +3233,10 @@ void FMetaSoundFrontendDocumentBuilder::FindNodeConfigurationUpdates(const FGuid
 {
 	using namespace Metasound::Frontend;
 
-	const INodeClassRegistry* Registry = INodeClassRegistry::Get();
-	checkf(nullptr != Registry, TEXT("The metasound node registry should always be available if the metasound plugin is loaded"));
-
 	OutNodeConfigurationUpdates.ExistingConfig = FindNodeConfiguration(InNodeID, &InPageID);
 	OutNodeConfigurationUpdates.ExistingClassInterfaceOverride = FindClassInterfaceOverride(InNodeID, &InPageID);
 
-	OutNodeConfigurationUpdates.RegisteredConfig = Registry->CreateFrontendNodeConfiguration(FNodeRegistryKey(InRegisteredClass.Metadata));
+	OutNodeConfigurationUpdates.RegisteredConfig = CreateFrontendNodeConfiguration(InRegisteredClass.Metadata);
 
 	// Determine which node config will be used to create the class interface override;
 	const bool bWillReplaceNodeConfig = ShouldReplaceExistingNodeConfig(OutNodeConfigurationUpdates.RegisteredConfig, OutNodeConfigurationUpdates.ExistingConfig);
diff --git a/Engine/Plugins/Runtime/Metasound/Source/MetasoundFrontend/Private/MetasoundFrontendGraphController.cpp b/Engine/Plugins/Runtime/Metasound/Source/MetasoundFrontend/Private/MetasoundFrontendGraphController.cpp
--- a/Engine/Plugins/Runtime/Metasound/Source/MetasoundFrontend/Private/MetasoundFrontendGraphController.cpp
+++ b/Engine/Plugins/Runtime/Metasound/Source/MetasoundFrontend/Private/MetasoundFrontendGraphController.cpp
@@ -1766,7 +1766,14 @@ PRAGMA_ENABLE_DEPRECATION_WARNINGS
 				if (const FMetasoundFrontendClass* NodeClass = InExistingDependency.Get())
 				{
 					TInstancedStruct<FMetaSoundFrontendNodeConfiguration> Configuration;
-					if (NodeClass->Metadata.GetType() != EMetasoundFrontendClassType::Graph) // Subgraphs do not have extensions
+					if (NodeClass->Metadata.GetType() == EMetasoundFrontendClassType::Template)
+					{
+						if (const INodeTemplate* Template = INodeTemplateRegistry::Get().FindTemplate(FNodeClassRegistryKey(NodeClass->Metadata)))
+						{
+							Configuration = Template->CreateFrontendTemplateNodeConfiguration();
+						}
+					}
+					else if (NodeClass->Metadata.GetType() != EMetasoundFrontendClassType::Graph) // Subgraphs do not have extensions
 					{
 						Configuration = INodeClassRegistry::Get()->CreateFrontendNodeConfiguration(FNodeClassRegistryKey(NodeClass->Metadata));
 					}
diff --git a/Engine/Plugins/Runtime/Metasound/Source/MetasoundFrontend/Private/MetasoundFrontendNodeController.cpp b/Engine/Plugins/Runtime/Metasound/Source/MetasoundFrontend/Private/MetasoundFrontendNodeController.cpp
--- a/Engine/Plugins/Runtime/Metasound/Source/MetasoundFrontend/Private/MetasoundFrontendNodeController.cpp
+++ b/Engine/Plugins/Runtime/Metasound/Source/MetasoundFrontend/Private/MetasoundFrontendNodeController.cpp
@@ -1070,7 +1070,17 @@ namespace Metasound
 			OutNodeConfigurationUpdates.ExistingConfig = GetNodeConfiguration();
 			OutNodeConfigurationUpdates.ExistingClassInterfaceOverride = GetClassInterfaceOverride();
 
-			OutNodeConfigurationUpdates.RegisteredConfig = Registry->CreateFrontendNodeConfiguration(FNodeRegistryKey(InRegisteredClass.Metadata));
+			if (InRegisteredClass.Metadata.GetType() == EMetasoundFrontendClassType::Template)
+			{
+				if (const INodeTemplate* Template = INodeTemplateRegistry::Get().FindTemplate(FNodeRegistryKey(InRegisteredClass.Metadata)))
+				{
+					OutNodeConfigurationUpdates.RegisteredConfig = Template->CreateFrontendTemplateNodeConfiguration();
+				}
+			}
+			else
+			{
+				OutNodeConfigurationUpdates.RegisteredConfig = Registry->CreateFrontendNodeConfiguration(FNodeRegistryKey(InRegisteredClass.Metadata));
+			}
 
 			// Determine which node config will be used to create the class interface override;
 			const bool bWillReplaceNodeConfig = ShouldReplaceExistingNodeConfig(OutNodeConfigurationUpdates.RegisteredConfig, OutNodeConfigurationUpdates.ExistingConfig);
diff --git a/Engine/Plugins/Runtime/Metasound/Source/MetasoundFrontend/Private/MetasoundFrontendNodeTemplateRegistry.cpp b/Engine/Plugins/Runtime/Metasound/Source/MetasoundFrontend/Private/MetasoundFrontendNodeTemplateRegistry.cpp
--- a/Engine/Plugins/Runtime/Metasound/Source/MetasoundFrontend/Private/MetasoundFrontendNodeTemplateRegistry.cpp
+++ b/Engine/Plugins/Runtime/Metasound/Source/MetasoundFrontend/Private/MetasoundFrontendNodeTemplateRegistry.cpp
@@ -14,7 +14,6 @@
 
 #include UE_INLINE_GENERATED_CPP_BY_NAME(MetasoundFrontendNodeTemplateRegistry)
 
-
 namespace Metasound::Frontend
 {
 #if WITH_EDITOR
@@ -129,6 +128,11 @@ namespace Metasound::Frontend
 		return nullptr;
 	}
 
+	TInstancedStruct<FMetaSoundFrontendTemplateNodeConfiguration> FNodeTemplateBase::CreateFrontendTemplateNodeConfiguration() const
+	{
+		return TInstancedStruct<FMetaSoundFrontendTemplateNodeConfiguration>();
+	}
+
 #if WITH_EDITOR
 	FText FNodeTemplateBase::GetNodeDisplayName(const IMetaSoundDocumentInterface& Interface, const FGuid& InPageID, const FGuid& InNodeID) const
 	{
diff --git a/Engine/Plugins/Runtime/Metasound/Source/MetasoundFrontend/Private/MetasoundFrontendTemplateNodeConfiguration.cpp b/Engine/Plugins/Runtime/Metasound/Source/MetasoundFrontend/Private/MetasoundFrontendTemplateNodeConfiguration.cpp
new file mode 100644
--- /dev/null
+++ b/Engine/Plugins/Runtime/Metasound/Source/MetasoundFrontend/Private/MetasoundFrontendTemplateNodeConfiguration.cpp
@@ -0,0 +1,7 @@
+// Copyright Epic Games, Inc. All Rights Reserved.
+
+#include "MetasoundFrontendTemplateNodeConfiguration.h"
+
+#include UE_INLINE_GENERATED_CPP_BY_NAME(MetasoundFrontendTemplateNodeConfiguration)
+
+	  
diff --git a/Engine/Plugins/Runtime/Metasound/Source/MetasoundFrontend/Private/NodeTemplates/MetasoundFrontendNodeTemplateReroute.cpp b/Engine/Plugins/Runtime/Metasound/Source/MetasoundFrontend/Private/NodeTemplates/MetasoundFrontendNodeTemplateReroute.cpp
--- a/Engine/Plugins/Runtime/Metasound/Source/MetasoundFrontend/Private/NodeTemplates/MetasoundFrontendNodeTemplateReroute.cpp
+++ b/Engine/Plugins/Runtime/Metasound/Source/MetasoundFrontend/Private/NodeTemplates/MetasoundFrontendNodeTemplateReroute.cpp
@@ -264,7 +264,7 @@ namespace Metasound::Frontend
 
 	const TArray<FMetasoundFrontendClassInputDefault>* FRerouteNodeTemplate::FindNodeClassInputDefaults(const FMetaSoundFrontendDocumentBuilder& InBuilder, const FGuid& InPageID, const FGuid& InNodeID, FName VertexName) const
 	{
-		// Recursive search up DAG for first connected non-reroute node's input class input
+		// Recursive search down DAG for first connected non-reroute node's input class input
 		if (const FMetasoundFrontendNode* Node = InBuilder.FindNode(InNodeID, &InPageID))
 		{
 			// Should only ever be one
@@ -293,7 +293,7 @@ namespace Metasound::Frontend
 
 	EMetasoundFrontendVertexAccessType FRerouteNodeTemplate::GetNodeInputAccessType(const FMetaSoundFrontendDocumentBuilder& InBuilder, const FGuid& InPageID, const FGuid& InNodeID, const FGuid& InVertexID) const
 	{
-		// Recursive search up DAG for first connected non-reroute node's input access type
+		// Recursive search down DAG for first connected non-reroute node's input access type
 		if (const FMetasoundFrontendNode* Node = InBuilder.FindNode(InNodeID, &InPageID))
 		{
 			// Should only ever be one
diff --git a/Engine/Plugins/Runtime/Metasound/Source/MetasoundFrontend/Public/MetasoundFrontendDocumentBuilder.h b/Engine/Plugins/Runtime/Metasound/Source/MetasoundFrontend/Public/MetasoundFrontendDocumentBuilder.h
--- a/Engine/Plugins/Runtime/Metasound/Source/MetasoundFrontend/Public/MetasoundFrontendDocumentBuilder.h
+++ b/Engine/Plugins/Runtime/Metasound/Source/MetasoundFrontend/Public/MetasoundFrontendDocumentBuilder.h
@@ -13,6 +13,7 @@
 #include "MetasoundLiteral.h"
 #include "MetasoundOperatorSettings.h"
 #include "MetasoundVertex.h"
+#include "StructUtils/InstancedStruct.h"
 #include "Templates/Function.h"
 #include "UObject/ScriptInterface.h"
 
@@ -731,6 +732,8 @@ public:
 private:
 	using FFinalizeNodeFunctionRef = TFunctionRef<void(FMetasoundFrontendNode&, const Metasound::Frontend::FNodeRegistryKey&)>;
 
+	TInstancedStruct<FMetaSoundFrontendNodeConfiguration> CreateFrontendNodeConfiguration(const FMetasoundFrontendClassMetadata& InClassMetadata) const;
+	TInstancedStruct<FMetaSoundFrontendNodeConfiguration> CreateFrontendNodeConfiguration(EMetasoundFrontendClassType InClassType, const Metasound::Frontend::FNodeRegistryKey& InClassKey) const;
 	FMetasoundFrontendNode* AddNodeInternal(const FMetasoundFrontendClassMetadata& InClassMetadata, Metasound::Frontend::FFinalizeNodeFunctionRef FinalizeNode, const FGuid& InPageID, FGuid InNodeID = FGuid::NewGuid(), int32* NewNodeIndex = nullptr);
 	FMetasoundFrontendNode* AddNodeInternal(const Metasound::Frontend::FNodeRegistryKey& InClassKey, FFinalizeNodeFunctionRef FinalizeNode, const FGuid& InPageID, FGuid InNodeID = FGuid::NewGuid(), int32* NewNodeIndex = nullptr);
 
diff --git a/Engine/Plugins/Runtime/Metasound/Source/MetasoundFrontend/Public/MetasoundFrontendNodeTemplateRegistry.h b/Engine/Plugins/Runtime/Metasound/Source/MetasoundFrontend/Public/MetasoundFrontendNodeTemplateRegistry.h
--- a/Engine/Plugins/Runtime/Metasound/Source/MetasoundFrontend/Public/MetasoundFrontendNodeTemplateRegistry.h
+++ b/Engine/Plugins/Runtime/Metasound/Source/MetasoundFrontend/Public/MetasoundFrontendNodeTemplateRegistry.h
@@ -4,11 +4,10 @@
 
 #include "MetasoundFrontendDocument.h"
 #include "MetasoundFrontendRegistryKey.h"
+#include "MetasoundFrontendTemplateNodeConfiguration.h"
 #include "MetasoundVertex.h"
-#include "Templates/SharedPointer.h"
 #include "Templates/UniquePtr.h"
 #include "UObject/NoExportTypes.h"
-#include "UObject/ScriptInterface.h"
 
 #include "MetasoundFrontendNodeTemplateRegistry.generated.h"
 
@@ -135,6 +134,9 @@ namespace Metasound::Frontend
 		// Returns note template class name.
 		virtual const FMetasoundFrontendClassName& GetClassName() const = 0;
 
+		// Create a node configuration for the node template.
+		virtual TInstancedStruct<FMetaSoundFrontendTemplateNodeConfiguration> CreateFrontendTemplateNodeConfiguration() const = 0;
+
 #if WITH_EDITOR
 		virtual FText GetNodeDisplayName(const IMetaSoundDocumentInterface& Interface, const FGuid& InPageID, const FGuid& InNodeID) const = 0;
 
@@ -224,6 +226,7 @@ namespace Metasound::Frontend
 	{
 	public:
 		UE_API virtual const TArray<FMetasoundFrontendClassInputDefault>* FindNodeClassInputDefaults(const FMetaSoundFrontendDocumentBuilder& InBuilder, const FGuid& InPageID, const FGuid& InNodeID, FName VertexName) const override;
+		UE_API virtual TInstancedStruct<FMetaSoundFrontendTemplateNodeConfiguration> CreateFrontendTemplateNodeConfiguration() const override;
 
 #if WITH_EDITOR
 		UE_API virtual FText GetNodeDisplayName(const IMetaSoundDocumentInterface& Interface, const FGuid& InPageID, const FGuid& InNodeID) const override;
diff --git a/Engine/Plugins/Runtime/Metasound/Source/MetasoundFrontend/Public/MetasoundFrontendTemplateNodeConfiguration.h b/Engine/Plugins/Runtime/Metasound/Source/MetasoundFrontend/Public/MetasoundFrontendTemplateNodeConfiguration.h
new file mode 100644
--- /dev/null
+++ b/Engine/Plugins/Runtime/Metasound/Source/MetasoundFrontend/Public/MetasoundFrontendTemplateNodeConfiguration.h
@@ -0,0 +1,14 @@
+// Copyright Epic Games, Inc. All Rights Reserved.
+
+#pragma once
+
+#include "MetasoundFrontendDocument.h"
+
+#include "MetasoundFrontendTemplateNodeConfiguration.generated.h"
+
+USTRUCT()
+struct FMetaSoundFrontendTemplateNodeConfiguration : public FMetaSoundFrontendNodeConfiguration
+{
+	GENERATED_BODY()
+};
+
