[Resubmit] Add derived type check to MetaSound node configuration setter

#rnx

diff --git a/Engine/Plugins/Runtime/Metasound/Source/MetasoundFrontend/Private/MetasoundAssetBase.cpp b/Engine/Plugins/Runtime/Metasound/Source/MetasoundFrontend/Private/MetasoundAssetBase.cpp
--- a/Engine/Plugins/Runtime/Metasound/Source/MetasoundFrontend/Private/MetasoundAssetBase.cpp
+++ b/Engine/Plugins/Runtime/Metasound/Source/MetasoundFrontend/Private/MetasoundAssetBase.cpp
@@ -124,6 +124,11 @@ namespace Metasound
 					return TInstancedStruct<FMetaSoundFrontendNodeConfiguration>();
 				}
 
+				virtual bool IsCompatibleNodeConfiguration(TConstStructView<FMetaSoundFrontendNodeConfiguration> InNodeConfiguration) const override
+				{
+					return false;
+				}
+
 			private:
 				TSet<FMetasoundFrontendVersion> Interfaces;
 				FMetasoundFrontendClass FrontendClass;
diff --git a/Engine/Plugins/Runtime/Metasound/Source/MetasoundFrontend/Private/MetasoundFrontendDataTypeRegistry.cpp b/Engine/Plugins/Runtime/Metasound/Source/MetasoundFrontend/Private/MetasoundFrontendDataTypeRegistry.cpp
--- a/Engine/Plugins/Runtime/Metasound/Source/MetasoundFrontend/Private/MetasoundFrontendDataTypeRegistry.cpp
+++ b/Engine/Plugins/Runtime/Metasound/Source/MetasoundFrontend/Private/MetasoundFrontendDataTypeRegistry.cpp
@@ -144,6 +144,11 @@ namespace Metasound
 					// By default, data type related nodes do not offer any extensions. 
 					return TInstancedStruct<FMetaSoundFrontendNodeConfiguration>();
 				}
+				
+				virtual bool IsCompatibleNodeConfiguration(TConstStructView<FMetaSoundFrontendNodeConfiguration> InNodeConfiguration) const override
+				{
+					return false;
+				}
 
 			protected:
 				void UpdateNodeClassInfo(const FMetasoundFrontendClass& InClass)
diff --git a/Engine/Plugins/Runtime/Metasound/Source/MetasoundFrontend/Private/MetasoundFrontendDocumentBuilder.cpp b/Engine/Plugins/Runtime/Metasound/Source/MetasoundFrontend/Private/MetasoundFrontendDocumentBuilder.cpp
--- a/Engine/Plugins/Runtime/Metasound/Source/MetasoundFrontend/Private/MetasoundFrontendDocumentBuilder.cpp
+++ b/Engine/Plugins/Runtime/Metasound/Source/MetasoundFrontend/Private/MetasoundFrontendDocumentBuilder.cpp
@@ -6624,8 +6624,22 @@ bool FMetaSoundFrontendDocumentBuilder::SetNodeConfiguration(const FGuid& InNode
 	if (const int32* NodeIndex = NodeCache.FindNodeIndex(InNodeID))
 	{
 		FMetasoundFrontendNode& Node = Graph.Nodes[*NodeIndex];
-		Node.Configuration = MoveTemp(InNodeConfiguration);
-		return UpdateNodeInterfaceFromConfiguration(InNodeID, InPageID);
+		// Node config derived type must match
+		if (const FMetasoundFrontendClass* Class = FindDependency(Node.ClassID))
+		{
+			const FNodeClassRegistryKey NodeClassKey(Class->Metadata);
+			if (!INodeClassRegistry::Get()->IsCompatibleNodeConfiguration(NodeClassKey, InNodeConfiguration))
+			{
+#if !NO_LOGGING
+				TInstancedStruct<FMetaSoundFrontendNodeConfiguration> RegisteredConfiguration = CreateFrontendNodeConfiguration(Class->Metadata);
+				UE_LOG(LogMetaSound, Warning, TEXT("Cannot set node configuration to type %s, expected registered type %s"), *GetNameSafe(InNodeConfiguration.GetScriptStruct()), *GetNameSafe(RegisteredConfiguration.GetScriptStruct()));
+#endif // !NO_LOGGING
+				return false;
+			}
+
+			Node.Configuration = MoveTemp(InNodeConfiguration);
+			return UpdateNodeInterfaceFromConfiguration(InNodeID, InPageID);
+		}
 	}
 	return false;
 }
diff --git a/Engine/Plugins/Runtime/Metasound/Source/MetasoundFrontend/Private/MetasoundFrontendNodeClassRegistry.cpp b/Engine/Plugins/Runtime/Metasound/Source/MetasoundFrontend/Private/MetasoundFrontendNodeClassRegistry.cpp
--- a/Engine/Plugins/Runtime/Metasound/Source/MetasoundFrontend/Private/MetasoundFrontendNodeClassRegistry.cpp
+++ b/Engine/Plugins/Runtime/Metasound/Source/MetasoundFrontend/Private/MetasoundFrontendNodeClassRegistry.cpp
@@ -193,6 +193,11 @@ namespace Metasound::Frontend
 				// overrides in FMetasoundFrontendDocument based nodes is unsupported. 
 				return TInstancedStruct<FMetaSoundFrontendNodeConfiguration>();
 			}
+
+			virtual bool IsCompatibleNodeConfiguration(TConstStructView<FMetaSoundFrontendNodeConfiguration> InNodeConfiguration) const override
+			{
+				return false;
+			}
 		private:
 			FMetasoundFrontendClass FrontendClass;
 			TSet<FMetasoundFrontendVersion> Interfaces;
@@ -1188,6 +1193,18 @@ PRAGMA_ENABLE_DEPRECATION_WARNINGS
 		return false;
 	}
 
+	bool FNodeClassRegistry::IsCompatibleNodeConfiguration(const FNodeClassRegistryKey& InKey, TConstStructView<FMetaSoundFrontendNodeConfiguration> InNodeConfiguration) const
+	{
+		bool bIsCompatible = false;
+		auto GetNodeConfigurationCompatibility = [&bIsCompatible, &InNodeConfiguration](const INodeClassRegistryEntry& Entry)
+		{
+			bIsCompatible = Entry.IsCompatibleNodeConfiguration(InNodeConfiguration);
+		};
+
+		AccessNodeEntryThreadSafe(InKey, GetNodeConfigurationCompatibility);
+		return bIsCompatible;
+	}
+
 	TInstancedStruct<FMetaSoundFrontendNodeConfiguration> FNodeClassRegistry::CreateFrontendNodeConfiguration(const FNodeClassRegistryKey& InKey) const
 	{
 		TInstancedStruct<FMetaSoundFrontendNodeConfiguration> NodeConfiguration;
diff --git a/Engine/Plugins/Runtime/Metasound/Source/MetasoundFrontend/Private/MetasoundFrontendNodeClassRegistryPrivate.h b/Engine/Plugins/Runtime/Metasound/Source/MetasoundFrontend/Private/MetasoundFrontendNodeClassRegistryPrivate.h
--- a/Engine/Plugins/Runtime/Metasound/Source/MetasoundFrontend/Private/MetasoundFrontendNodeClassRegistryPrivate.h
+++ b/Engine/Plugins/Runtime/Metasound/Source/MetasoundFrontend/Private/MetasoundFrontendNodeClassRegistryPrivate.h
@@ -102,6 +102,8 @@ namespace Metasound::Frontend
 		virtual bool FindDefaultVertexInterface(const FNodeClassRegistryKey& InKey, FVertexInterface& OutVertexInterface) const override;
 		
 		virtual bool FindFrontendClassFromRegistered(const FNodeClassRegistryKey& InKey, FMetasoundFrontendClass& OutClass) const override;
+
+		virtual bool IsCompatibleNodeConfiguration(const FNodeClassRegistryKey& InKey, TConstStructView<FMetaSoundFrontendNodeConfiguration> InNodeConfiguration) const override;
 		virtual TInstancedStruct<FMetaSoundFrontendNodeConfiguration> CreateFrontendNodeConfiguration(const FNodeClassRegistryKey& InKey) const override;
 
 		virtual bool FindImplementedInterfacesFromRegistered(const Metasound::Frontend::FNodeClassRegistryKey& InKey, TSet<FMetasoundFrontendVersion>& OutInterfaceVersions) const override;
diff --git a/Engine/Plugins/Runtime/Metasound/Source/MetasoundFrontend/Private/MetasoundNodeRegistrationMacro.cpp b/Engine/Plugins/Runtime/Metasound/Source/MetasoundFrontend/Private/MetasoundNodeRegistrationMacro.cpp
--- a/Engine/Plugins/Runtime/Metasound/Source/MetasoundFrontend/Private/MetasoundNodeRegistrationMacro.cpp
+++ b/Engine/Plugins/Runtime/Metasound/Source/MetasoundFrontend/Private/MetasoundNodeRegistrationMacro.cpp
@@ -51,6 +51,11 @@ namespace Metasound::Frontend
 			return TInstancedStruct<FMetaSoundFrontendNodeConfiguration>();
 		}
 
+		bool FNodeRegistryEntryBase::IsCompatibleNodeConfiguration(TConstStructView<FMetaSoundFrontendNodeConfiguration> InNodeConfiguration) const
+		{
+			return false;
+		}
+
 		TSharedRef<const FNodeClassMetadata> FNodeRegistryEntryBase::GetNodeClassMetadata() const
 		{
 			return ClassMetadata;
diff --git a/Engine/Plugins/Runtime/Metasound/Source/MetasoundFrontend/Public/MetasoundFrontendNodeClassRegistry.h b/Engine/Plugins/Runtime/Metasound/Source/MetasoundFrontend/Public/MetasoundFrontendNodeClassRegistry.h
--- a/Engine/Plugins/Runtime/Metasound/Source/MetasoundFrontend/Public/MetasoundFrontendNodeClassRegistry.h
+++ b/Engine/Plugins/Runtime/Metasound/Source/MetasoundFrontend/Public/MetasoundFrontendNodeClassRegistry.h
@@ -113,6 +113,8 @@ namespace Metasound::Frontend
 #else
 		UE_API virtual TInstancedStruct<FMetaSoundFrontendNodeConfiguration> CreateFrontendNodeConfiguration() const;
 #endif
+		/** Return whether this entry's node configuration is compatible with the given node configuration. */
+		virtual bool IsCompatibleNodeConfiguration(TConstStructView<FMetaSoundFrontendNodeConfiguration> InNodeConfiguration) const = 0;
 
 #if WITH_EDITORONLY_DATA
 		/** Name of the plugin registering this node. */
@@ -361,7 +363,10 @@ namespace Metasound::Frontend
 
 		virtual bool FindFrontendClassFromRegistered(const FNodeClassRegistryKey& InKey, FMetasoundFrontendClass& OutClass) const = 0;
 
-		/** Return the node extension associated with the node. If there is no extension assoicated with the node, the returned TInstancedStruct will be invalid. */
+		/** Returns true if the provided node configuration is compatible (same derived class) with the registered configuration for the node represented by the registry key. */
+		virtual bool IsCompatibleNodeConfiguration(const FNodeClassRegistryKey& InKey, TConstStructView<FMetaSoundFrontendNodeConfiguration> InNodeConfiguration) const = 0;
+
+		/** Return the node extension associated with the node. If there is no extension associated with the node, the returned TInstancedStruct will be invalid. */
 		virtual TInstancedStruct<FMetaSoundFrontendNodeConfiguration> CreateFrontendNodeConfiguration(const FNodeClassRegistryKey& InKey) const = 0;
 
 		virtual bool FindImplementedInterfacesFromRegistered(const FNodeClassRegistryKey& InKey, TSet<FMetasoundFrontendVersion>& OutInterfaceVersions) const = 0;
diff --git a/Engine/Plugins/Runtime/Metasound/Source/MetasoundFrontend/Public/MetasoundNodeRegistrationMacro.h b/Engine/Plugins/Runtime/Metasound/Source/MetasoundFrontend/Public/MetasoundNodeRegistrationMacro.h
--- a/Engine/Plugins/Runtime/Metasound/Source/MetasoundFrontend/Public/MetasoundNodeRegistrationMacro.h
+++ b/Engine/Plugins/Runtime/Metasound/Source/MetasoundFrontend/Public/MetasoundNodeRegistrationMacro.h
@@ -51,6 +51,7 @@ namespace Metasound::Frontend
 			UE_API virtual FVertexInterface GetDefaultVertexInterface() const override;
 
 			UE_API virtual TInstancedStruct<FMetaSoundFrontendNodeConfiguration> CreateFrontendNodeConfiguration() const override;
+			UE_API virtual bool IsCompatibleNodeConfiguration(TConstStructView<FMetaSoundFrontendNodeConfiguration> InNodeConfiguration) const override;
 
 #if WITH_EDITORONLY_DATA
 			UE_API virtual FName GetPluginName() const override;
@@ -118,6 +119,10 @@ namespace Metasound::Frontend
 				return TInstancedStruct<FMetaSoundFrontendNodeConfiguration>::Make<ConfigurationType>();
 			}
 
+			virtual bool IsCompatibleNodeConfiguration(TConstStructView<FMetaSoundFrontendNodeConfiguration> InNodeConfiguration) const override 
+			{
+				return InNodeConfiguration.GetScriptStruct() == ConfigurationType::StaticStruct();
+			}
 		};
 
 		// A partial template specialization for scenario where no node extension is provided
