Fix static analysis warning - possible module by zero

#rnx

diff --git a/Engine/Plugins/Runtime/Metasound/Source/MetasoundGraphCore/Private/MetasoundOperatorSettings.cpp b/Engine/Plugins/Runtime/Metasound/Source/MetasoundGraphCore/Private/MetasoundOperatorSettings.cpp
--- a/Engine/Plugins/Runtime/Metasound/Source/MetasoundGraphCore/Private/MetasoundOperatorSettings.cpp
+++ b/Engine/Plugins/Runtime/Metasound/Source/MetasoundGraphCore/Private/MetasoundOperatorSettings.cpp
@@ -94,7 +94,7 @@ namespace Metasound
 
 	int32 FOperatorSettings::RoundToAligned(int32 Alignment, int32 InNum) const
 	{
-		check(Alignment >= 0);
+		check(Alignment > 0);
 
 		// Need to add (Alignment / 2) in order to perform rounding operation.
 		// Otherwise, this function would be more like "FloorToAligned(...)"
