MetaSound Node Configuration: node update transform API and implementation in autoupdate

diff --git a/Engine/Plugins/Runtime/Metasound/Source/MetasoundEditor/Private/MetasoundEditorGraphNode.cpp b/Engine/Plugins/Runtime/Metasound/Source/MetasoundEditor/Private/MetasoundEditorGraphNode.cpp
--- a/Engine/Plugins/Runtime/Metasound/Source/MetasoundEditor/Private/MetasoundEditorGraphNode.cpp
+++ b/Engine/Plugins/Runtime/Metasound/Source/MetasoundEditor/Private/MetasoundEditorGraphNode.cpp
@@ -1155,9 +1155,8 @@ bool UMetasoundEditorGraphExternalNode::CanAutoUpdate() const
 {
 	using namespace Metasound::Frontend;
 
-	FClassInterfaceUpdates InterfaceUpdates;
 	FMetaSoundFrontendDocumentBuilder& DocBuilder = GetBuilderChecked().GetBuilder();
-	return DocBuilder.CanAutoUpdate(GetNodeID(), InterfaceUpdates);
+	return DocBuilder.CanAutoUpdate(GetNodeID()) != EAutoUpdateEligibility::Ineligible;
 }
 
 void UMetasoundEditorGraphExternalNode::CacheBreadcrumb()
diff --git a/Engine/Plugins/Runtime/Metasound/Source/MetasoundFrontend/Private/MetasoundFrontendDocumentBuilder.cpp b/Engine/Plugins/Runtime/Metasound/Source/MetasoundFrontend/Private/MetasoundFrontendDocumentBuilder.cpp
--- a/Engine/Plugins/Runtime/Metasound/Source/MetasoundFrontend/Private/MetasoundFrontendDocumentBuilder.cpp
+++ b/Engine/Plugins/Runtime/Metasound/Source/MetasoundFrontend/Private/MetasoundFrontendDocumentBuilder.cpp
@@ -22,6 +22,7 @@
 #include "MetasoundFrontendDocumentModifyDelegates.h"
 #include "MetasoundFrontendDocumentVersioning.h"
 #include "MetasoundFrontendNodeTemplateRegistry.h"
+#include "MetasoundFrontendNodeUpdateTransform.h"
 #include "MetasoundFrontendRegistries.h"
 #include "MetasoundFrontendRegistryKey.h"
 #include "MetasoundFrontendSearchEngine.h"
@@ -1807,38 +1808,25 @@ const FMetasoundFrontendNode* FMetaSoundFrontendDocumentBuilder::AddGraphNode(co
 	return AddNodeInternal(ClassKey, FinalizeNode, PageID, InNodeID, NewNodeIndex);
 }
 
-const FMetasoundFrontendNode* FMetaSoundFrontendDocumentBuilder::AddNodeByClassName(const FMetasoundFrontendClassName& InClassName, int32 InMajorVersion, FGuid InNodeID, const FGuid* InPageID)
+const FMetasoundFrontendNode* FMetaSoundFrontendDocumentBuilder::AddNodeByClassInternal(FMetasoundFrontendClass&& InClass, FGuid InNodeID, const FGuid* InPageID)
 {
-	using namespace Metasound;
 	using namespace Metasound::Frontend;
 
 	const FMetasoundFrontendClass* Dependency = nullptr;
-	FNodeRegistryKey ClassKey;
+	const EMetasoundFrontendClassType ClassType = InClass.Metadata.GetType();
+	if (ClassType != EMetasoundFrontendClassType::External && ClassType != EMetasoundFrontendClassType::Graph)
 	{
-		FMetasoundFrontendClass RegisteredClass;
-		if (!ISearchEngine::Get().FindClassWithHighestMinorVersion(InClassName, InMajorVersion, RegisteredClass))
-		{
-			UE_LOG(LogMetaSound, Error, TEXT("Failed to add new node by class name '%s' and major version '%d': Class not found"), *InClassName.ToString(), InMajorVersion);
-			return nullptr;
-		}
-
-		const EMetasoundFrontendClassType ClassType = RegisteredClass.Metadata.GetType();
-		if (ClassType != EMetasoundFrontendClassType::External && ClassType != EMetasoundFrontendClassType::Graph)
-		{
-			UE_LOG(LogMetaSound, Warning, TEXT("Failed to add new node by class name '%s': Class is restricted type '%s' that cannot be added via this function."),
-				*InClassName.ToString(),
-				LexToString(ClassType));
-			return nullptr;
-		}
-
-		// Dependency is considered "External" when looked up or added as a dependency to a graph
-		RegisteredClass.Metadata.SetType(EMetasoundFrontendClassType::External);
-		ClassKey = FNodeRegistryKey(RegisteredClass.Metadata);
-		Dependency = FindDependency(RegisteredClass.Metadata);
-		if (!Dependency)
-		{
-			Dependency = AddDependency(MoveTemp(RegisteredClass));
-		}
+		UE_LOG(LogMetaSound, Warning, TEXT("Failed to add new node by class name '%s': Class is restricted type '%s' that cannot be added via this function."),
+			*InClass.Metadata.GetClassName().ToString(),
+			LexToString(ClassType));
+		return nullptr;
+	}
+	// Dependency is considered "External" when looked up or added as a dependency to a graph
+	InClass.Metadata.SetType(EMetasoundFrontendClassType::External);
+	Dependency = FindDependency(InClass.Metadata);
+	if (!Dependency)
+	{
+		Dependency = AddDependency(MoveTemp(InClass));
 	}
 
 	if (Dependency)
@@ -1852,54 +1840,32 @@ const FMetasoundFrontendNode* FMetaSoundFrontendDocumentBuilder::AddNodeByClassN
 	return nullptr;
 }
 
-const FMetasoundFrontendNode* FMetaSoundFrontendDocumentBuilder::AddNodeByClassName(const FMetasoundFrontendClassName& InClassName, int32 InMajorVersion, int32 InMinorVersion, FGuid InNodeID, const FGuid* InPageID)
+const FMetasoundFrontendNode* FMetaSoundFrontendDocumentBuilder::AddNodeByClassName(const FMetasoundFrontendClassName& InClassName, int32 InMajorVersion, FGuid InNodeID, const FGuid* InPageID)
 {
-	using namespace Metasound;
 	using namespace Metasound::Frontend;
-
-	const FMetasoundFrontendClass* Dependency = nullptr;
-	FNodeRegistryKey ClassKey;
+	
+	FMetasoundFrontendClass RegisteredClass;
+	if (!ISearchEngine::Get().FindClassWithHighestMinorVersion(InClassName, InMajorVersion, RegisteredClass))
 	{
-		INodeClassRegistry* NodeRegistry = INodeClassRegistry::Get();
-		if (ensure(NodeRegistry))
-		{
-			FNodeRegistryKey LookupKey(EMetasoundFrontendClassType::External, InClassName, InMajorVersion, InMinorVersion);
-			FMetasoundFrontendClass RegisteredClass;
-			if (!NodeRegistry->FindFrontendClassFromRegistered(LookupKey, RegisteredClass))
-			{
-				UE_LOG(LogMetaSound, Error, TEXT("Failed to add new node by class name '%s' and major version '%d' and minor version '%d': Class not found"), *InClassName.ToString(), InMajorVersion, InMinorVersion);
-				return nullptr;
-			}
-
-			const EMetasoundFrontendClassType ClassType = RegisteredClass.Metadata.GetType();
-			if (ClassType != EMetasoundFrontendClassType::External && ClassType != EMetasoundFrontendClassType::Graph)
-			{
-				UE_LOG(LogMetaSound, Warning, TEXT("Failed to add new node by class name '%s': Class is restricted type '%s' that cannot be added via this function."),
-					*InClassName.ToString(),
-					LexToString(ClassType));
-				return nullptr;
-			}
-
-			// Dependency is considered "External" when looked up or added as a dependency to a graph
-			RegisteredClass.Metadata.SetType(EMetasoundFrontendClassType::External);
-			ClassKey = FNodeRegistryKey(RegisteredClass.Metadata);
-			Dependency = FindDependency(RegisteredClass.Metadata);
-			if (!Dependency)
-			{
-				Dependency = AddDependency(MoveTemp(RegisteredClass));
-			}
-		}
+		UE_LOG(LogMetaSound, Error, TEXT("Failed to add new node by class name '%s' and major version '%d': Class info not registered"), *InClassName.ToString(), InMajorVersion);
+		return nullptr;
 	}
+	
+	return AddNodeByClassInternal(MoveTemp(RegisteredClass), InNodeID, InPageID);
+}
 
-	if (Dependency)
+const FMetasoundFrontendNode* FMetaSoundFrontendDocumentBuilder::AddNodeByClassName(const FMetasoundFrontendClassName& InClassName, int32 InMajorVersion, int32 InMinorVersion, FGuid InNodeID, const FGuid* InPageID)
+{
+	using namespace Metasound::Frontend;
+
+	FNodeRegistryKey LookupKey(EMetasoundFrontendClassType::External, InClassName, InMajorVersion, InMinorVersion);
+	FMetasoundFrontendClass RegisteredClass;
+	if (!INodeClassRegistry::GetChecked().FindFrontendClassFromRegistered(LookupKey, RegisteredClass))
 	{
-		auto FinalizeNode = [](const FMetasoundFrontendNode& Node, const Metasound::Frontend::FNodeRegistryKey& ClassKey) { return Node.Name; };
-		constexpr int32* NewNodeIndex = nullptr;
-		const FGuid& PageID = InPageID ? *InPageID : BuildPageID;
-		return AddNodeInternal(Dependency->Metadata, FinalizeNode, PageID, InNodeID, NewNodeIndex);
+		UE_LOG(LogMetaSound, Error, TEXT("Failed to add new node by class name '%s' and major version '%d' and minor version '%d': Class not found"), *InClassName.ToString(), InMajorVersion, InMinorVersion);
+		return nullptr;
 	}
-
-	return nullptr;
+	return AddNodeByClassInternal(MoveTemp(RegisteredClass), InNodeID, InPageID);
 }
 
 const FMetasoundFrontendNode* FMetaSoundFrontendDocumentBuilder::AddNodeByTemplate(const Metasound::Frontend::INodeTemplate& InTemplate, FNodeTemplateGenerateInterfaceParams Params, FGuid InNodeID, const FGuid* InPageID)
@@ -2020,6 +1986,47 @@ const FMetasoundFrontendGraph& FMetaSoundFrontendDocumentBuilder::AddGraphPage(c
 	}
 	return ToReturn;
 }
+
+void FMetaSoundFrontendDocumentBuilder::ApplyDependencyUpdateTransform(const Metasound::Frontend::FNodeClassRegistryKey& InNodeClassKey)
+{
+	using namespace Metasound::Frontend;
+
+	const TSharedPtr<const INodeUpdateTransform>& NodeUpdateTransform = INodeClassRegistry::GetChecked().FindNodeUpdateTransform(InNodeClassKey);
+	if (!NodeUpdateTransform.IsValid())
+	{
+		UE_LOG(LogMetaSound, Warning, TEXT("Dependency update transform %s not registered."), *InNodeClassKey.ToString());
+		return;
+	}
+	
+	// Apply node update transform to all nodes of this class
+	// Can't directly iterate because node updates may change the underlying node arrays
+	// Array of node id, page id pairs
+	TArray<TTuple<const FGuid, const FGuid*>> NodesToUpdate;
+	auto CollectNodesToUpdate = [&](const FMetasoundFrontendClass&, const FMetasoundFrontendNode& Node, const FGuid& PageID)
+	{
+		NodesToUpdate.Emplace(Node.GetID(), &PageID);
+	};
+
+	auto GetNodesOfClass = [&](const FMetasoundFrontendNode& Node)
+	{
+		const FMetasoundFrontendClass* Class = FindDependency(Node.ClassID);
+		if (Class)
+		{
+			const FNodeClassRegistryKey KeyToCompare(Class->Metadata);
+			return KeyToCompare == InNodeClassKey;
+		}
+		return false;
+	};
+	IterateNodesByPredicate(CollectNodesToUpdate, GetNodesOfClass, /*InPageID =*/nullptr, /*bIterateAllPages=*/true);
+
+	for (const TTuple<const FGuid, const FGuid*>& NodePair : NodesToUpdate)
+	{
+		if (const FMetasoundFrontendNode* Node = FindNode(NodePair.Key, NodePair.Value))
+		{
+			NodeUpdateTransform->Update(*this, NodePair.Key, NodePair.Value);
+		}
+	}
+}
 #endif // WITH_EDITORONLY_DATA
 
 #if WITH_EDITOR
@@ -2152,22 +2159,22 @@ bool FMetaSoundFrontendDocumentBuilder::CanAddEdge(const FMetasoundFrontendEdge&
 }
 
 #if WITH_EDITORONLY_DATA
-bool FMetaSoundFrontendDocumentBuilder::CanAutoUpdate(const FGuid& InNodeID, Metasound::Frontend::FClassInterfaceUpdates& OutInterfaceUpdates, const FGuid* InPageID)
+Metasound::Frontend::EAutoUpdateEligibility FMetaSoundFrontendDocumentBuilder::CanAutoUpdate(const FGuid& InNodeID, const FGuid* InPageID) const
 {
 	using namespace Metasound::Frontend;
 	METASOUND_TRACE_CPUPROFILER_EVENT_SCOPE(FMetaSoundFrontendDocumentBuilder::CanAutoUpdate);
 
-	OutInterfaceUpdates = { };
+	FClassInterfaceUpdates InterfaceUpdates;
 	const FGuid& PageID = InPageID ? *InPageID : BuildPageID;
 	const FMetasoundFrontendNode* Node = FindNode(InNodeID, InPageID);
 	if (!Node)
 	{
-		return false;
+		return EAutoUpdateEligibility::Ineligible;
 	}
 	const FMetasoundFrontendClass* Class = FindDependency(Node->ClassID);
 	if (!Class)
 	{
-		return false;
+		return EAutoUpdateEligibility::Ineligible;
 	}
 
 	const FMetasoundFrontendClassMetadata& NodeClassMetadata = Class->Metadata;
@@ -2176,13 +2183,13 @@ bool FMetaSoundFrontendDocumentBuilder::CanAutoUpdate(const FGuid& InNodeID, Met
 	FMetasoundFrontendClass RegistryClass;
 	if (!ISearchEngine::Get().FindClassWithHighestMinorVersion(NodeClassMetadata.GetClassName().ToNodeClassName(), NodeClassMetadata.GetVersion().Major, RegistryClass))
 	{
-		return false;
+		return EAutoUpdateEligibility::Ineligible;
 	}
 
 	// 1. Document's class version is somehow higher than registries, so can't update.
 	if (RegistryClass.Metadata.GetVersion() < NodeClassMetadata.GetVersion())
 	{
-		return false;
+		return EAutoUpdateEligibility::Ineligible;
 	}
 
 	// 2. Document's class version is equal, so have to dif and check change IDs
@@ -2205,7 +2212,7 @@ bool FMetaSoundFrontendDocumentBuilder::CanAutoUpdate(const FGuid& InNodeID, Met
 					const FGuid& NodeClassInterfaceChangeID = ClassInterface.GetChangeID();
 					if (RegistryClass.GetDefaultInterface().GetChangeID() == NodeClassInterfaceChangeID)
 					{
-						return false;
+						return EAutoUpdateEligibility::Ineligible;
 					}
 				}
 			}
@@ -2213,7 +2220,7 @@ bool FMetaSoundFrontendDocumentBuilder::CanAutoUpdate(const FGuid& InNodeID, Met
 			{
 				if (!MetaSoundAutoUpdateNativeClassesOfEqualVersionCVar)
 				{
-					return false;
+					return EAutoUpdateEligibility::Ineligible;
 				}
 			}
 		}
@@ -2221,9 +2228,26 @@ bool FMetaSoundFrontendDocumentBuilder::CanAutoUpdate(const FGuid& InNodeID, Met
 		constexpr bool bUseHighestMinorVersion = true;
 		constexpr bool bForceRegenerateClassInterfaceOverride = true;
 
-		DiffAgainstRegistryInterface(InNodeID, PageID, bUseHighestMinorVersion, OutInterfaceUpdates, bForceRegenerateClassInterfaceOverride);
+		DiffAgainstRegistryInterface(InNodeID, PageID, bUseHighestMinorVersion, InterfaceUpdates, bForceRegenerateClassInterfaceOverride);
+
+		if (InterfaceUpdates.ContainsChanges())
+		{
+			return EAutoUpdateEligibility::Eligible_InterfaceChange;
+		}
+
+		// Check for registered node transform that can be autoapplied
+		const FNodeClassRegistryKey NodeClassKey(NodeClassMetadata);
 
-		return OutInterfaceUpdates.ContainsChanges();
+		if (TSharedPtr<const INodeUpdateTransform> NodeTransform = INodeClassRegistry::Get()->FindNodeUpdateTransform(NodeClassKey))
+		{
+			if (NodeTransform && NodeTransform->ShouldAutoApply())
+			{
+				return EAutoUpdateEligibility::Eligible_NodeUpdateTransform;
+			}
+		}
+		
+		// No interface change or eligible transform
+		return EAutoUpdateEligibility::Ineligible;
 	}
 
 	// 3. Document's class version is out-of-date, so dif and always return true that can auto-update
@@ -2231,9 +2255,9 @@ bool FMetaSoundFrontendDocumentBuilder::CanAutoUpdate(const FGuid& InNodeID, Met
 	// contains no changes).
 	constexpr bool bUseHighestMinorVersion = true;
 	constexpr bool bForceRegenerateClassInterfaceOverride = true;
-	DiffAgainstRegistryInterface(InNodeID, PageID, bUseHighestMinorVersion, OutInterfaceUpdates, bForceRegenerateClassInterfaceOverride);
+	DiffAgainstRegistryInterface(InNodeID, PageID, bUseHighestMinorVersion, InterfaceUpdates, bForceRegenerateClassInterfaceOverride);
 
-	return true;
+	return EAutoUpdateEligibility::Eligible_MinorVersionUpdate;
 }
 #endif // WITH_EDITORONLY_DATA
 
@@ -5505,7 +5529,7 @@ bool FMetaSoundFrontendDocumentBuilder::ReplaceDependency(const Metasound::Front
 	// Readd nodes with new class
 	for (FNodeInstanceReplacementData& NodeInstanceData : NodeData)
 	{
-		const FMetasoundFrontendNode* NewNode = AddNodeByClassName(NewNodeClassKey.ClassName, NewNodeClassKey.Version.Major, NodeInstanceData.NodeID, &NodeInstanceData.PageID);
+		const FMetasoundFrontendNode* NewNode = AddNodeByClassName(NewNodeClassKey.ClassName, NewNodeClassKey.Version.Major, NewNodeClassKey.Version.Minor, NodeInstanceData.NodeID, &NodeInstanceData.PageID);
 		bSuccess |= (NewNode != nullptr);
 	}
 
diff --git a/Engine/Plugins/Runtime/Metasound/Source/MetasoundFrontend/Private/MetasoundFrontendNodeClassRegistry.cpp b/Engine/Plugins/Runtime/Metasound/Source/MetasoundFrontend/Private/MetasoundFrontendNodeClassRegistry.cpp
--- a/Engine/Plugins/Runtime/Metasound/Source/MetasoundFrontend/Private/MetasoundFrontendNodeClassRegistry.cpp
+++ b/Engine/Plugins/Runtime/Metasound/Source/MetasoundFrontend/Private/MetasoundFrontendNodeClassRegistry.cpp
@@ -1003,6 +1003,64 @@ PRAGMA_ENABLE_DEPRECATION_WARNINGS
 
 		return bSuccess;
 	}
+
+	bool FNodeClassRegistry::RegisterNodeUpdateTransform(const FNodeClassRegistryKey& InNodeClassKey, const TSharedRef<const INodeUpdateTransform>& InNodeUpdateTransform)
+	{
+		if (IsNodeRegistered(InNodeClassKey))
+		{
+			FScopeLock Lock(&RegistryMapsCriticalSection);
+			NodeUpdateTransforms.Add(InNodeClassKey, InNodeUpdateTransform);
+			return true;
+		}
+		UE_LOG(LogMetaSound, Error, TEXT("Invalid attempt to register node update transform for node key %s that is not registered."), *InNodeClassKey.ToString());
+		return false;
+	}
+
+	bool FNodeClassRegistry::RegisterNodeUpdateTransform(const TArray<const FNodeClassRegistryKey>& InNodeClassKeys, const TSharedRef<const INodeUpdateTransform>& InNodeUpdateTransform)
+	{
+		bool bSuccess = true;
+		FScopeLock Lock(&RegistryMapsCriticalSection);
+		for (const FNodeClassRegistryKey& Key : InNodeClassKeys)
+		{
+			if (IsNodeRegistered(Key))
+			{
+				NodeUpdateTransforms.Add(Key, InNodeUpdateTransform);
+			}
+			else
+			{
+				bSuccess = false;
+				UE_LOG(LogMetaSound, Error, TEXT("Invalid attempt to register node update transform for node key %s that is not registered."), *Key.ToString());
+			}
+		}
+		return bSuccess;
+	}
+
+	bool FNodeClassRegistry::UnregisterNodeUpdateTransform(const FNodeClassRegistryKey& InNodeUpdateTransformKey)
+	{
+		FScopeLock Lock(&RegistryMapsCriticalSection);
+		return NodeUpdateTransforms.Remove(InNodeUpdateTransformKey) > 0;
+	}
+
+	bool FNodeClassRegistry::UnregisterNodeUpdateTransform(const TArray<const FNodeClassRegistryKey>& InNodeUpdateTransformKeys)
+	{
+		bool bSuccess = true;
+		FScopeLock Lock(&RegistryMapsCriticalSection);
+		for (const FNodeClassRegistryKey& Key : InNodeUpdateTransformKeys)
+		{
+			bSuccess &= NodeUpdateTransforms.Remove(Key) > 0;
+		}
+		return bSuccess;
+	}
+
+	TSharedPtr<const INodeUpdateTransform> FNodeClassRegistry::FindNodeUpdateTransform(const FNodeClassRegistryKey& InUpdateKey) const
+	{
+		const TSharedRef<const INodeUpdateTransform, ESPMode::ThreadSafe>* NodeUpdateTransform = NodeUpdateTransforms.Find(InUpdateKey);
+		if (NodeUpdateTransform)
+		{
+			return *NodeUpdateTransform;
+		}
+		return nullptr;
+	}
 #endif // WITH_EDITORONLY_DATA
 
 	FNodeClassRegistryKey FNodeClassRegistry::RegisterNodeTemplate(TUniquePtr<INodeTemplateRegistryEntry>&& InEntry)
diff --git a/Engine/Plugins/Runtime/Metasound/Source/MetasoundFrontend/Private/MetasoundFrontendNodeClassRegistryPrivate.h b/Engine/Plugins/Runtime/Metasound/Source/MetasoundFrontend/Private/MetasoundFrontendNodeClassRegistryPrivate.h
--- a/Engine/Plugins/Runtime/Metasound/Source/MetasoundFrontend/Private/MetasoundFrontendNodeClassRegistryPrivate.h
+++ b/Engine/Plugins/Runtime/Metasound/Source/MetasoundFrontend/Private/MetasoundFrontendNodeClassRegistryPrivate.h
@@ -89,6 +89,16 @@ namespace Metasound::Frontend
 #if WITH_EDITORONLY_DATA
 		virtual bool RegisterNodeMigration(const FNodeMigrationInfo& InMigrationInfo) override;
 		virtual bool UnregisterNodeMigration(const FNodeMigrationInfo& InMigrationInfo) override;
+
+		virtual bool RegisterNodeUpdateTransform(const FNodeClassRegistryKey& InNodeUpdateTransformKey, const TSharedRef<const INodeUpdateTransform>& InNodeUpdateTransform) override;
+		virtual bool RegisterNodeUpdateTransform(const TArray<const FNodeClassRegistryKey>& InNodeUpdateTransformKeys, const TSharedRef<const INodeUpdateTransform>& InNodeUpdateTransform) override;
+
+		// Unregister the association between a class name and version and a node update transform. 
+		// If other class name/versions are still associated with the transform, it will remain in the registry with that association 
+		virtual bool UnregisterNodeUpdateTransform(const FNodeClassRegistryKey& InNodeUpdateTransformKey) override;
+		virtual bool UnregisterNodeUpdateTransform(const TArray<const FNodeClassRegistryKey>& InNodeUpdateTransformKeys) override;
+
+		virtual TSharedPtr<const INodeUpdateTransform> FindNodeUpdateTransform(const FNodeClassRegistryKey& InUpdateKey) const override;
 #endif // if WITH_EDITORONLY_DATA
 
 		virtual bool IsNodeRegistered(const FNodeClassRegistryKey& InKey) const override;
@@ -191,6 +201,7 @@ namespace Metasound::Frontend
 
 #if WITH_EDITORONLY_DATA
 		TMultiMap<FNodeClassRegistryKey, FNodeMigrationInfo> NodeMigrations;
+		TMap<FNodeClassRegistryKey, TSharedRef<const INodeUpdateTransform>> NodeUpdateTransforms;
 #endif
 
 		// Registry in which we keep all information about dynamically-generated templated nodes via in C++.
diff --git a/Engine/Plugins/Runtime/Metasound/Source/MetasoundFrontend/Private/MetasoundFrontendTransform.cpp b/Engine/Plugins/Runtime/Metasound/Source/MetasoundFrontend/Private/MetasoundFrontendTransform.cpp
--- a/Engine/Plugins/Runtime/Metasound/Source/MetasoundFrontend/Private/MetasoundFrontendTransform.cpp
+++ b/Engine/Plugins/Runtime/Metasound/Source/MetasoundFrontend/Private/MetasoundFrontendTransform.cpp
@@ -521,13 +521,67 @@
 			METASOUND_TRACE_CPUPROFILER_EVENT_SCOPE(FAutoUpdateRootGraph::Transform);
 			bool bDidEdit = false;
 
+			const bool bIsPreset = InOutBuilderToTransform.IsPreset();
+			const FMetasoundFrontendGraphClass& RootGraph = InOutBuilderToTransform.GetConstDocumentChecked().RootGraph;
+			
+			// If preset, rebuild root graph 
+			if (bIsPreset)
+			{
+				FMetasoundAssetBase* PresetReferencedMetaSoundAsset = InOutBuilderToTransform.GetReferencedPresetAsset();
+				if (!PresetReferencedMetaSoundAsset)
+				{
+					METASOUND_VERSIONING_LOG(Error, TEXT("Auto-Updating preset '%s' failed: Referenced class missing."), *DebugAssetPath);
+					return false;
+				}
+				TScriptInterface<IMetaSoundDocumentInterface> RefDocInterface = PresetReferencedMetaSoundAsset->GetOwningAsset();
+				const FMetaSoundFrontendDocumentBuilder& ReferenceBuilder = IDocumentBuilderRegistry::GetChecked().FindOrBeginBuilding(RefDocInterface);
+
+				bDidEdit |= FRebuildPresetRootGraph(ReferenceBuilder).Transform(InOutBuilderToTransform);
+				if (bDidEdit)
+				{
+					FMetasoundFrontendClassMetadata PresetMetadata = InOutBuilderToTransform.GetConstDocumentChecked().RootGraph.Metadata;
+					PresetMetadata.SetType(EMetasoundFrontendClassType::External);
+					const FNodeRegistryKey RegistryKey(PresetMetadata);
+					FMetasoundAssetBase* PresetMetaSoundAsset = IMetaSoundAssetManager::GetChecked().TryLoadAssetFromKey(RegistryKey);
+					if (ensure(PresetMetaSoundAsset))
+					{
+						TScriptInterface<IMetaSoundDocumentInterface> PresetInterface = PresetMetaSoundAsset->GetOwningAsset();
+						check(PresetInterface);
+						PresetInterface->ConformObjectToDocument();
+					}
+
+					InOutBuilderToTransform.RemoveUnusedDependencies();
+					InOutBuilderToTransform.SynchronizeDependencyMetadata();
+				}
+			}
+			// Non preset
+			else
+			{
+				// Update external node dependencies until there are no more updates
+				// (A node could need a minor version update, then an update transform, 
+				// and then updated to another minor version which would require multiple passes)
+				bool bUpdated = true;
+				while (bUpdated)
+				{
+					bUpdated = false;
+					bUpdated |= UpdateExternalDependencies(InOutBuilderToTransform);
+					bDidEdit |= bUpdated;
+				}
+			}
+				
+			return bDidEdit;
+		}
+
+		bool FAutoUpdateRootGraph::UpdateExternalDependencies(FMetaSoundFrontendDocumentBuilder& InOutBuilderToTransform)
+		{
+			bool bDidEdit = false;
 			// <Old dependency, new dependency>
 			using FDependencyPair = TPair<FNodeClassRegistryKey, FNodeClassRegistryKey>;
 			TSet<FDependencyPair> DependenciesToUpdate;
-			const bool bIsPreset = InOutBuilderToTransform.IsPreset();
+			TSet<FNodeClassRegistryKey> DependenciesToTransform;
+			
+			// Collect updates 
 			const FMetasoundFrontendGraphClass& RootGraph = InOutBuilderToTransform.GetConstDocumentChecked().RootGraph;
-			
-			// Go through external nodes and find dependencies to update
 			RootGraph.IterateGraphPages([&](const FMetasoundFrontendGraph& Graph)
 			{
 				const FGuid PageID = Graph.PageID;
@@ -535,10 +589,9 @@
 				{
 					const FNodeClassRegistryKey& NodeClassKey = FNodeClassRegistryKey(Class.Metadata);
 
-					FClassInterfaceUpdates InterfaceUpdates;
 					const bool bHasUpdated = UpdatedClasses.Contains(Class.ID);
-					const bool bCanAutoUpdate = InOutBuilderToTransform.CanAutoUpdate(Node.GetID(), InterfaceUpdates, &PageID);
-					if (bHasUpdated || !bCanAutoUpdate)
+					const EAutoUpdateEligibility AutoUpdateReason = InOutBuilderToTransform.CanAutoUpdate(Node.GetID(), &PageID);
+					if (bHasUpdated || AutoUpdateReason == EAutoUpdateEligibility::Ineligible)
 					{
 						return;
 					}
@@ -554,13 +607,22 @@
 						const FMetasoundFrontendVersionNumber UpdateVersion = ClassWithHighestMinorVersion.Metadata.GetVersion();
 
 						METASOUND_VERSIONING_LOG(Display, TEXT("Auto-Updating '%s' node class '%s': Newer version '%s' found."), *DebugAssetPath, *NodeClassKey.ClassName.ToString(), *UpdateVersion.ToString());
-
 						DependenciesToUpdate.Add(FDependencyPair(NodeClassKey, FNodeClassRegistryKey(ClassWithHighestMinorVersion.Metadata)));
+						bDidEdit |= true;
+						
 					}
-					else if (InterfaceUpdates.ContainsChanges())
+					else if (AutoUpdateReason == EAutoUpdateEligibility::Eligible_InterfaceChange)
 					{
 						METASOUND_VERSIONING_LOG(Display, TEXT("Auto-Updating '%s' node class '%s (%s)': Interface change detected."), *DebugAssetPath, *NodeClassKey.ClassName.ToString(), *NodeClassKey.Version.ToString());
 						DependenciesToUpdate.Add(FDependencyPair(NodeClassKey, NodeClassKey));
+						bDidEdit |= true;
 					}
+					// Order is intentional; node must be updated to highest minor version before node update transform can be applied
+					else if (AutoUpdateReason == EAutoUpdateEligibility::Eligible_NodeUpdateTransform)
+					{
+						METASOUND_VERSIONING_LOG(Display, TEXT("Auto-Updating '%s' node class '%s (%s)': Node update transform found."), *DebugAssetPath, *NodeClassKey.ClassName.ToString(), *NodeClassKey.Version.ToString());
+						DependenciesToTransform.Add(NodeClassKey);
+						bDidEdit |= true;
+					}
 					else
 					{
@@ -569,40 +631,7 @@
 				}, EMetasoundFrontendClassType::External, &PageID);
 			});
 			
-			// If preset, rebuild root graph 
-			if (bIsPreset)
-			{
-				FMetasoundAssetBase* PresetReferencedMetaSoundAsset = InOutBuilderToTransform.GetReferencedPresetAsset();
-				if (!PresetReferencedMetaSoundAsset)
-				{
-					METASOUND_VERSIONING_LOG(Error, TEXT("Auto-Updating preset '%s' failed: Referenced class missing."), *DebugAssetPath);
-					return false;
-				}
-				TScriptInterface<IMetaSoundDocumentInterface> RefDocInterface = PresetReferencedMetaSoundAsset->GetOwningAsset();
-				FMetaSoundFrontendDocumentBuilder& ReferenceBuilder = IDocumentBuilderRegistry::GetChecked().FindOrBeginBuilding(RefDocInterface);
-
-				bDidEdit |= FRebuildPresetRootGraph(ReferenceBuilder).Transform(InOutBuilderToTransform);
-				if (bDidEdit)
-				{
-					FMetasoundFrontendClassMetadata PresetMetadata = InOutBuilderToTransform.GetConstDocumentChecked().RootGraph.Metadata;
-					PresetMetadata.SetType(EMetasoundFrontendClassType::External);
-					const FNodeRegistryKey RegistryKey(PresetMetadata);
-					FMetasoundAssetBase* PresetMetaSoundAsset = IMetaSoundAssetManager::GetChecked().TryLoadAssetFromKey(RegistryKey);
-					if (ensure(PresetMetaSoundAsset))
-					{
-						TScriptInterface<IMetaSoundDocumentInterface> PresetInterface = PresetMetaSoundAsset->GetOwningAsset();
-						check(PresetInterface);
-						PresetInterface->ConformObjectToDocument();
-					}
-						
-					InOutBuilderToTransform.RemoveUnusedDependencies();
-					InOutBuilderToTransform.SynchronizeDependencyMetadata();
-				}
-			}
-
-			bDidEdit |= !DependenciesToUpdate.IsEmpty();
-
-			// Update each dependency and log disconnections
+			// Apply dependency updates and log disconnections
 			for (const FDependencyPair& DependencyPair : DependenciesToUpdate)
 			{
 				using FVertexNameAndType = FMetaSoundFrontendDocumentBuilder::FVertexNameAndType;
@@ -631,6 +660,12 @@
 				}
 			}
 
+			// Apply dependency transforms 
+			for (const FNodeClassRegistryKey& NodeClassKey : DependenciesToTransform)
+			{
+				InOutBuilderToTransform.ApplyDependencyUpdateTransform(NodeClassKey);
+			}
+
 			InOutBuilderToTransform.RemoveUnusedDependencies();
 			InOutBuilderToTransform.SynchronizeDependencyMetadata();
 
diff --git a/Engine/Plugins/Runtime/Metasound/Source/MetasoundFrontend/Public/MetasoundFrontendDocumentBuilder.h b/Engine/Plugins/Runtime/Metasound/Source/MetasoundFrontend/Public/MetasoundFrontendDocumentBuilder.h
--- a/Engine/Plugins/Runtime/Metasound/Source/MetasoundFrontend/Public/MetasoundFrontendDocumentBuilder.h
+++ b/Engine/Plugins/Runtime/Metasound/Source/MetasoundFrontend/Public/MetasoundFrontendDocumentBuilder.h
@@ -95,6 +95,19 @@ namespace Metasound::Frontend
 		bool bSetDefaultNodeLocations = true;
 #endif // WITH_EDITORONLY_DATA
 	};
+
+#if WITH_EDITORONLY_DATA
+	UE_EXPERIMENTAL(5.7, "Node update transforms are experimental")
+	enum class UE_API EAutoUpdateEligibility : uint8
+	{
+		None = 0,
+		Ineligible,
+		Eligible_MinorVersionUpdate,
+		Eligible_InterfaceChange,
+		Eligible_NodeUpdateTransform,
+		COUNT
+	};
+#endif // WITH_EDITORONLY_DATA
 } // namespace Metasound::Frontend
 
 
@@ -214,6 +227,10 @@ public:
 #if WITH_EDITORONLY_DATA
 	// Adds a graph page to the given builder's document
 	UE_API const FMetasoundFrontendGraph& AddGraphPage(const FGuid& InPageID, bool bDuplicateLastGraph, bool bSetAsBuildGraph = true);
+
+	// Apply the registered node update transform to all nodes of the given dependency
+	UE_EXPERIMENTAL(5.7, "Node update transforms are experimental")
+	UE_API void ApplyDependencyUpdateTransform(const Metasound::Frontend::FNodeClassRegistryKey& InNodeClassKey);
 #endif // WITH_EDITORONLY_DATA
 
 #if WITH_EDITOR
@@ -227,14 +244,13 @@ public:
 	UE_API bool CanAddEdge(const FMetasoundFrontendEdge& InEdge, const FGuid* InPageID = nullptr) const;
 
 #if WITH_EDITORONLY_DATA
-	// Returns whether the node is eligible for auto-updating (i.e.
-	// has undergone minor revision or the interface has changed, but
-	// no higher major revision is available). Provides interface updates
-	// to populate with any information regarding interface updates.
-	// Can return true if the interface has changed but with only cosmetic
+	// Returns whether the node is eligible for auto-updating and reason for eligibility (i.e.
+	// has undergone minor version revision, node update transform found,
+	// or the interface has changed, but no higher major revision is available).
+	// Can return eligible if the interface has changed but with only cosmetic
 	// differences (ex. DisplayName only used in editor) but no runtime
 	// behavior has been modified.
-	UE_API bool CanAutoUpdate(const FGuid& InNodeID, Metasound::Frontend::FClassInterfaceUpdates& OutInterfaceUpdates, const FGuid* InPageID = nullptr);
+	UE_API UE_INTERNAL Metasound::Frontend::EAutoUpdateEligibility CanAutoUpdate(const FGuid& InNodeID, const FGuid* InPageID = nullptr) const;
 #endif // WITH_EDITORONLY_DATA
 
 	// Clears document completely of all graph page data (nodes, edges, & member metadata), dependencies,
@@ -737,6 +753,8 @@ private:
 	FMetasoundFrontendNode* AddNodeInternal(const FMetasoundFrontendClassMetadata& InClassMetadata, Metasound::Frontend::FFinalizeNodeFunctionRef FinalizeNode, const FGuid& InPageID, FGuid InNodeID = FGuid::NewGuid(), int32* NewNodeIndex = nullptr);
 	FMetasoundFrontendNode* AddNodeInternal(const Metasound::Frontend::FNodeRegistryKey& InClassKey, FFinalizeNodeFunctionRef FinalizeNode, const FGuid& InPageID, FGuid InNodeID = FGuid::NewGuid(), int32* NewNodeIndex = nullptr);
 
+	const FMetasoundFrontendNode* AddNodeByClassInternal(FMetasoundFrontendClass&& InClass, FGuid InNodeID, const FGuid* InPageID);
+
 	struct FInputConnectionInfo
 	{
 		// Handle of connected output vertex
diff --git a/Engine/Plugins/Runtime/Metasound/Source/MetasoundFrontend/Public/MetasoundFrontendNodeClassRegistry.h b/Engine/Plugins/Runtime/Metasound/Source/MetasoundFrontend/Public/MetasoundFrontendNodeClassRegistry.h
--- a/Engine/Plugins/Runtime/Metasound/Source/MetasoundFrontend/Public/MetasoundFrontendNodeClassRegistry.h
+++ b/Engine/Plugins/Runtime/Metasound/Source/MetasoundFrontend/Public/MetasoundFrontendNodeClassRegistry.h
@@ -3,6 +3,7 @@
 
 #include "MetasoundFrontendDocument.h"
 #include "MetasoundFrontendNodeMigration.h"
+#include "MetasoundFrontendNodeUpdateTransform.h"
 #include "MetasoundFrontendRegistryKey.h"
 #include "MetasoundNodeInterface.h"
 #include "MetasoundNodeConstructorParams.h"
@@ -403,6 +404,29 @@ namespace Metasound::Frontend
 		// Returns a list of possible nodes to use to convert from FromDataType to ToDataType.
 		// Returns an empty array if none are available.
 		virtual TArray<FConverterNodeClassInfo> GetPossibleConverterNodes(const FName& FromDataType, const FName& ToDataType) = 0;
+
+#if WITH_EDITORONLY_DATA
+		// Node Update Transforms
+		// Register a node update transform associated with a registered node class.
+		// Multiple class name/version combinations can be associated with the same node update transform
+PRAGMA_DISABLE_EXPERIMENTAL_WARNINGS
+		UE_EXPERIMENTAL(5.7, "Node update transforms are experimental")
+		virtual bool RegisterNodeUpdateTransform(const FNodeClassRegistryKey& InNodeUpdateTransformKey, const TSharedRef<const INodeUpdateTransform>& InNodeUpdateTransform) = 0;
+		UE_EXPERIMENTAL(5.7, "Node update transforms are experimental")
+		virtual bool RegisterNodeUpdateTransform(const TArray<const FNodeClassRegistryKey>& InNodeUpdateTransformKeys, const TSharedRef<const INodeUpdateTransform>& InNodeUpdateTransform) = 0;
+		
+		// Unregister the association between a class name and version and a node update transform. 
+		// If other class name/versions are still associated with the transform, it will remain in the registry with that association 
+		UE_EXPERIMENTAL(5.7, "Node update transforms are experimental")
+		virtual bool UnregisterNodeUpdateTransform(const FNodeClassRegistryKey& InNodeUpdateTransformKey) = 0;
+		UE_EXPERIMENTAL(5.7, "Node update transforms are experimental")
+		virtual bool UnregisterNodeUpdateTransform(const TArray<const FNodeClassRegistryKey>& InNodeUpdateTransformKeys) = 0;
+
+		UE_EXPERIMENTAL(5.7, "Node update transforms are experimental")
+		virtual TSharedPtr<const INodeUpdateTransform> FindNodeUpdateTransform(const FNodeClassRegistryKey& InUpdateKey) const = 0;
+PRAGMA_ENABLE_EXPERIMENTAL_WARNINGS
+
+#endif // WITH_EDITORONLY_DATA
 	};
 } // namespace Metasound::Frontend
 
diff --git a/Engine/Plugins/Runtime/Metasound/Source/MetasoundFrontend/Public/MetasoundFrontendNodeUpdateTransform.h b/Engine/Plugins/Runtime/Metasound/Source/MetasoundFrontend/Public/MetasoundFrontendNodeUpdateTransform.h
new file mode 100644
--- /dev/null
+++ b/Engine/Plugins/Runtime/Metasound/Source/MetasoundFrontend/Public/MetasoundFrontendNodeUpdateTransform.h
@@ -0,0 +1,27 @@
+// Copyright Epic Games, Inc. All Rights Reserved.
+#pragma once
+
+#include "UObject/NameTypes.h"
+
+#define UE_API METASOUNDFRONTEND_API
+
+struct FMetaSoundFrontendDocumentBuilder;
+
+namespace Metasound::Frontend
+{
+	class UE_EXPERIMENTAL(5.7, "Node update transforms are experimental") INodeUpdateTransform
+	{
+#if WITH_EDITORONLY_DATA
+	public:
+		virtual ~INodeUpdateTransform() = default;
+
+		// Update document state for a given node
+		// This API is in progress and will likely change in the future
+		UE_INTERNAL virtual void Update(FMetaSoundFrontendDocumentBuilder& InOutBuilder, const FGuid& InNodeID, const FGuid* InPageID) const = 0;
+
+		// Whether the transform should be automatically applied to node classes registered with the transform
+		virtual bool ShouldAutoApply() const = 0;
+#endif // if WITH_EDITORONLY_DATA
+	};
+} // namespace Metasound::Frontend
+#undef UE_API
diff --git a/Engine/Plugins/Runtime/Metasound/Source/MetasoundFrontend/Public/MetasoundFrontendTransform.h b/Engine/Plugins/Runtime/Metasound/Source/MetasoundFrontend/Public/MetasoundFrontendTransform.h
--- a/Engine/Plugins/Runtime/Metasound/Source/MetasoundFrontend/Public/MetasoundFrontendTransform.h
+++ b/Engine/Plugins/Runtime/Metasound/Source/MetasoundFrontend/Public/MetasoundFrontendTransform.h
@@ -219,6 +219,9 @@ namespace Metasound::Frontend
 		UE_API bool Transform(FMetaSoundFrontendDocumentBuilder& InOutBuilderToTransform);
 
 	private:
+		// Helper function to go through external nodes and update dependencies
+		bool UpdateExternalDependencies(FMetaSoundFrontendDocumentBuilder& InOutBuilderToTransform);
+
 		// Keeps track of classes already updated so node check can be avoided.
 		// Hack to avoid issue where earlier auto-update passes on pages can
 		// clear out internal change state of a class in the registry causing
