CIS fix for macros for metasound node configuration node transform update

diff --git a/Engine/Plugins/Runtime/Metasound/Source/MetasoundFrontend/Public/MetasoundFrontendDocumentBuilder.h b/Engine/Plugins/Runtime/Metasound/Source/MetasoundFrontend/Public/MetasoundFrontendDocumentBuilder.h
--- a/Engine/Plugins/Runtime/Metasound/Source/MetasoundFrontend/Public/MetasoundFrontendDocumentBuilder.h
+++ b/Engine/Plugins/Runtime/Metasound/Source/MetasoundFrontend/Public/MetasoundFrontendDocumentBuilder.h
@@ -97,8 +97,7 @@ namespace Metasound::Frontend
 	};
 
 #if WITH_EDITORONLY_DATA
-	UE_EXPERIMENTAL(5.7, "Node update transforms are experimental")
-	enum class UE_API EAutoUpdateEligibility : uint8
+	enum class UE_EXPERIMENTAL(5.7, "Node update transforms are experimental") UE_API EAutoUpdateEligibility : uint8
 	{
 		None = 0,
 		Ineligible,
