Metasound node config CIS fix part 2

diff --git a/Engine/Plugins/Runtime/Metasound/Source/MetasoundFrontend/Public/MetasoundFrontendDocumentBuilder.h b/Engine/Plugins/Runtime/Metasound/Source/MetasoundFrontend/Public/MetasoundFrontendDocumentBuilder.h
--- a/Engine/Plugins/Runtime/Metasound/Source/MetasoundFrontend/Public/MetasoundFrontendDocumentBuilder.h
+++ b/Engine/Plugins/Runtime/Metasound/Source/MetasoundFrontend/Public/MetasoundFrontendDocumentBuilder.h
@@ -249,7 +249,9 @@ public:
 	// Can return eligible if the interface has changed but with only cosmetic
 	// differences (ex. DisplayName only used in editor) but no runtime
 	// behavior has been modified.
-	UE_API UE_INTERNAL Metasound::Frontend::EAutoUpdateEligibility CanAutoUpdate(const FGuid& InNodeID, const FGuid* InPageID = nullptr) const;
+PRAGMA_DISABLE_EXPERIMENTAL_WARNINGS
+	UE_INTERNAL UE_API Metasound::Frontend::EAutoUpdateEligibility CanAutoUpdate(const FGuid& InNodeID, const FGuid* InPageID = nullptr) const;
+PRAGMA_ENABLE_EXPERIMENTAL_WARNINGS
 #endif // WITH_EDITORONLY_DATA
 
 	// Clears document completely of all graph page data (nodes, edges, & member metadata), dependencies,
