MetaSound registry lock CIS fix

diff --git a/Engine/Plugins/Runtime/Metasound/Source/MetasoundFrontend/Private/MetasoundFrontendDocumentBuilder.cpp b/Engine/Plugins/Runtime/Metasound/Source/MetasoundFrontend/Private/MetasoundFrontendDocumentBuilder.cpp
--- a/Engine/Plugins/Runtime/Metasound/Source/MetasoundFrontend/Private/MetasoundFrontendDocumentBuilder.cpp
+++ b/Engine/Plugins/Runtime/Metasound/Source/MetasoundFrontend/Private/MetasoundFrontendDocumentBuilder.cpp
@@ -2001,10 +2001,10 @@ void FMetaSoundFrontendDocumentBuilder::ApplyDependencyUpdateTransform(const Met
 	// Apply node update transform to all nodes of this class
 	// Can't directly iterate because node updates may change the underlying node arrays
 	// Array of node id, page id pairs
-	TArray<TTuple<const FGuid, const FGuid*>> NodesToUpdate;
+	TArray<TTuple<const FGuid, const FGuid>> NodesToUpdate;
 	auto CollectNodesToUpdate = [&](const FMetasoundFrontendClass&, const FMetasoundFrontendNode& Node, const FGuid& PageID)
 	{
-		NodesToUpdate.Emplace(Node.GetID(), &PageID);
+		NodesToUpdate.Emplace(Node.GetID(), PageID);
 	};
 
 	auto GetNodesOfClass = [&](const FMetasoundFrontendNode& Node)
@@ -2019,11 +2019,11 @@ void FMetaSoundFrontendDocumentBuilder::ApplyDependencyUpdateTransform(const Met
 	};
 	IterateNodesByPredicate(CollectNodesToUpdate, GetNodesOfClass, /*InPageID =*/nullptr, /*bIterateAllPages=*/true);
 
-	for (const TTuple<const FGuid, const FGuid*>& NodePair : NodesToUpdate)
+	for (const TTuple<const FGuid, const FGuid>& NodePair : NodesToUpdate)
 	{
-		if (const FMetasoundFrontendNode* Node = FindNode(NodePair.Key, NodePair.Value))
+		if (const FMetasoundFrontendNode* Node = FindNode(NodePair.Key, &NodePair.Value))
 		{
-			NodeUpdateTransform->Update(*this, NodePair.Key, NodePair.Value);
+			NodeUpdateTransform->Update(*this, NodePair.Key, &NodePair.Value);
 		}
 	}
 }
diff --git a/Engine/Plugins/Runtime/Metasound/Source/MetasoundFrontend/Private/MetasoundFrontendNodeClassRegistry.cpp b/Engine/Plugins/Runtime/Metasound/Source/MetasoundFrontend/Private/MetasoundFrontendNodeClassRegistry.cpp
--- a/Engine/Plugins/Runtime/Metasound/Source/MetasoundFrontend/Private/MetasoundFrontendNodeClassRegistry.cpp
+++ b/Engine/Plugins/Runtime/Metasound/Source/MetasoundFrontend/Private/MetasoundFrontendNodeClassRegistry.cpp
@@ -1008,7 +1008,7 @@ PRAGMA_ENABLE_DEPRECATION_WARNINGS
 	{
 		if (IsNodeRegistered(InNodeClassKey))
 		{
-			FScopeLock Lock(&RegistryMapsCriticalSection);
+			UE::TScopeLock Lock(RegistryMapsCriticalSection);
 			NodeUpdateTransforms.Add(InNodeClassKey, InNodeUpdateTransform);
 			return true;
 		}
@@ -1019,7 +1019,7 @@ PRAGMA_ENABLE_DEPRECATION_WARNINGS
 	bool FNodeClassRegistry::RegisterNodeUpdateTransform(const TArray<const FNodeClassRegistryKey>& InNodeClassKeys, const TSharedRef<const INodeUpdateTransform>& InNodeUpdateTransform)
 	{
 		bool bSuccess = true;
-		FScopeLock Lock(&RegistryMapsCriticalSection);
+		UE::TScopeLock Lock(RegistryMapsCriticalSection);
 		for (const FNodeClassRegistryKey& Key : InNodeClassKeys)
 		{
 			if (IsNodeRegistered(Key))
@@ -1037,14 +1037,14 @@ PRAGMA_ENABLE_DEPRECATION_WARNINGS
 
 	bool FNodeClassRegistry::UnregisterNodeUpdateTransform(const FNodeClassRegistryKey& InNodeUpdateTransformKey)
 	{
-		FScopeLock Lock(&RegistryMapsCriticalSection);
+		UE::TScopeLock Lock(RegistryMapsCriticalSection);
 		return NodeUpdateTransforms.Remove(InNodeUpdateTransformKey) > 0;
 	}
 
 	bool FNodeClassRegistry::UnregisterNodeUpdateTransform(const TArray<const FNodeClassRegistryKey>& InNodeUpdateTransformKeys)
 	{
 		bool bSuccess = true;
-		FScopeLock Lock(&RegistryMapsCriticalSection);
+		UE::TScopeLock Lock(RegistryMapsCriticalSection);
 		for (const FNodeClassRegistryKey& Key : InNodeUpdateTransformKeys)
 		{
 			bSuccess &= NodeUpdateTransforms.Remove(Key) > 0;
