- Fix for runtime data not being invalidated when changing preset state
- Minor clean up of delegate collection page code

#rnx

diff --git a/Engine/Plugins/Runtime/Metasound/Source/MetasoundEngine/Private/MetasoundBuilderSubsystem.cpp b/Engine/Plugins/Runtime/Metasound/Source/MetasoundEngine/Private/MetasoundBuilderSubsystem.cpp
--- a/Engine/Plugins/Runtime/Metasound/Source/MetasoundEngine/Private/MetasoundBuilderSubsystem.cpp
+++ b/Engine/Plugins/Runtime/Metasound/Source/MetasoundEngine/Private/MetasoundBuilderSubsystem.cpp
@@ -394,16 +394,11 @@ const FMetasoundFrontendGraph& UMetaSoundSourceBuilder::GetConstTargetPageGraphC
 	return RootGraph.FindConstGraphChecked(TargetPageID);
 }
 
-const UMetaSoundSource& UMetaSoundSourceBuilder::GetMetaSoundSource() const
+UMetaSoundSource& UMetaSoundSourceBuilder::GetMetaSoundSource() const
 {
 	return GetConstBuilder().CastDocumentObjectChecked<UMetaSoundSource>();
 }
 
-UMetaSoundSource& UMetaSoundSourceBuilder::GetMetaSoundSource()
-{
-	return Builder.CastDocumentObjectChecked<UMetaSoundSource>();
-}
-
 void UMetaSoundSourceBuilder::InitDelegates(Metasound::Frontend::FDocumentModifyDelegates& OutDocumentDelegates)
 {
 	Super::InitDelegates(OutDocumentDelegates);
@@ -412,6 +407,7 @@ void UMetaSoundSourceBuilder::InitDelegates(Metasound::Frontend::FDocumentModify
 	{
 		OutDocumentDelegates.PageDelegates.OnPageAdded.AddUObject(this, &UMetaSoundSourceBuilder::OnPageAdded);
 		OutDocumentDelegates.PageDelegates.OnRemovingPage.AddUObject(this, &UMetaSoundSourceBuilder::OnRemovingPage);
+		OutDocumentDelegates.OnPresetStateChanged.AddUObject(this, &UMetaSoundSourceBuilder::OnPresetStateChanged);
 
 		OutDocumentDelegates.InterfaceDelegates.OnInputAdded.AddUObject(this, &UMetaSoundSourceBuilder::OnInputAdded);
 		OutDocumentDelegates.InterfaceDelegates.OnOutputAdded.AddUObject(this, &UMetaSoundSourceBuilder::OnOutputAdded);
@@ -733,7 +729,7 @@ void UMetaSoundSourceBuilder::OnOutputAdded(int32 OutputIndex) const
 	});
 }
 
-void UMetaSoundSourceBuilder::OnPageAdded(const Metasound::Frontend::FDocumentMutatePageArgs& Args)
+void UMetaSoundSourceBuilder::OnPageAdded(const Metasound::Frontend::FDocumentMutatePageArgs& InArgs)
 {
 	using namespace Metasound::Frontend;
 
@@ -741,7 +737,14 @@ void UMetaSoundSourceBuilder::OnPageAdded(const Metasound::Frontend::FDocumentMu
 	InitTargetPageDelegates(DocDelegates);
 }
 
-void UMetaSoundSourceBuilder::OnRemovingPage(const Metasound::Frontend::FDocumentMutatePageArgs& Args)
+void UMetaSoundSourceBuilder::OnPresetStateChanged(const Metasound::Frontend::FDocumentPresetStateChangedArgs& InArgs) const
+{
+	using namespace Metasound::Frontend;
+
+	GetMetaSoundSource().InvalidateCachedRuntimeInputData();
+}
+
+void UMetaSoundSourceBuilder::OnRemovingPage(const Metasound::Frontend::FDocumentMutatePageArgs& InArgs)
 {
 	using namespace Metasound::Frontend;
 
diff --git a/Engine/Plugins/Runtime/Metasound/Source/MetasoundEngine/Public/MetasoundBuilderSubsystem.h b/Engine/Plugins/Runtime/Metasound/Source/MetasoundEngine/Public/MetasoundBuilderSubsystem.h
--- a/Engine/Plugins/Runtime/Metasound/Source/MetasoundEngine/Public/MetasoundBuilderSubsystem.h
+++ b/Engine/Plugins/Runtime/Metasound/Source/MetasoundEngine/Public/MetasoundBuilderSubsystem.h
@@ -118,8 +118,7 @@ protected:
 private:
 	UE_API const FMetasoundFrontendGraph& GetConstTargetPageGraphChecked() const;
 
-	UE_API const UMetaSoundSource& GetMetaSoundSource() const;
-	UE_API UMetaSoundSource& GetMetaSoundSource();
+	UE_API UMetaSoundSource& GetMetaSoundSource() const;
 
 	UE_API void InitTargetPageDelegates(Metasound::Frontend::FDocumentModifyDelegates& OutDocumentDelegates);
 
@@ -131,6 +130,7 @@ private:
 	UE_API void OnNodeInputLiteralSet(int32 NodeIndex, int32 VertexIndex, int32 LiteralIndex) const;
 	UE_API void OnOutputAdded(int32 OutputIndex) const;
 	UE_API void OnPageAdded(const Metasound::Frontend::FDocumentMutatePageArgs& InArgs);
+	UE_API void OnPresetStateChanged(const Metasound::Frontend::FDocumentPresetStateChangedArgs& InArgs) const;
 	UE_API void OnRemoveSwappingEdge(int32 SwapIndex, int32 LastIndex) const;
 	UE_API void OnRemovingInput(int32 InputIndex);
 	UE_API void OnRemoveSwappingNode(int32 SwapIndex, int32 LastIndex) const;
diff --git a/Engine/Plugins/Runtime/Metasound/Source/MetasoundFrontend/Private/MetasoundFrontendDocumentBuilder.cpp b/Engine/Plugins/Runtime/Metasound/Source/MetasoundFrontend/Private/MetasoundFrontendDocumentBuilder.cpp
--- a/Engine/Plugins/Runtime/Metasound/Source/MetasoundFrontend/Private/MetasoundFrontendDocumentBuilder.cpp
+++ b/Engine/Plugins/Runtime/Metasound/Source/MetasoundFrontend/Private/MetasoundFrontendDocumentBuilder.cpp
@@ -2454,6 +2454,7 @@ bool FMetaSoundFrontendDocumentBuilder::ConvertFromPreset()
 		Style.bIsGraphEditable = true;
 #endif // WITH_EDITOR
 
+		DocumentDelegates->OnPresetStateChanged.Broadcast(FDocumentPresetStateChangedArgs { });
 		return true;
 	}
 
@@ -2487,6 +2488,7 @@ bool FMetaSoundFrontendDocumentBuilder::ConvertToPreset(const FMetaSoundFrontend
 	if (RebuildPresetRootGraph.Transform(*this))
 	{
 		DocumentInterface->ConformObjectToDocument();
+		DocumentDelegates->OnPresetStateChanged.Broadcast(FDocumentPresetStateChangedArgs { });
 		return true;
 	}
 
diff --git a/Engine/Plugins/Runtime/Metasound/Source/MetasoundFrontend/Private/MetasoundFrontendDocumentModifyDelegates.cpp b/Engine/Plugins/Runtime/Metasound/Source/MetasoundFrontend/Private/MetasoundFrontendDocumentModifyDelegates.cpp
--- a/Engine/Plugins/Runtime/Metasound/Source/MetasoundFrontend/Private/MetasoundFrontendDocumentModifyDelegates.cpp
+++ b/Engine/Plugins/Runtime/Metasound/Source/MetasoundFrontend/Private/MetasoundFrontendDocumentModifyDelegates.cpp
@@ -6,84 +6,12 @@
 namespace Metasound::Frontend
 {
 	FDocumentModifyDelegates::FDocumentModifyDelegates()
-PRAGMA_DISABLE_DEPRECATION_WARNINGS
-		: NodeDelegates()
-		, EdgeDelegates()
-PRAGMA_ENABLE_DEPRECATION_WARNINGS
 	{
-	}
-
-	FDocumentModifyDelegates::FDocumentModifyDelegates(const FDocumentModifyDelegates& InModifyDelegates)
-		: OnDocumentMetadataChanged(InModifyDelegates.OnDocumentMetadataChanged)
-		, OnDependencyAdded(InModifyDelegates.OnDependencyAdded)
-		, OnRemoveSwappingDependency(InModifyDelegates.OnRemoveSwappingDependency)
-		, OnRenamingDependencyClass(InModifyDelegates.OnRenamingDependencyClass)
-		, PageDelegates(InModifyDelegates.PageDelegates)
-		, InterfaceDelegates(InModifyDelegates.InterfaceDelegates)
-		PRAGMA_DISABLE_DEPRECATION_WARNINGS
-		, NodeDelegates(InModifyDelegates.NodeDelegates)
-		, EdgeDelegates(InModifyDelegates.EdgeDelegates)
-		PRAGMA_ENABLE_DEPRECATION_WARNINGS
-		, PageNodeDelegates(InModifyDelegates.PageNodeDelegates)
-		, PageEdgeDelegates(InModifyDelegates.PageEdgeDelegates)
-	{
-	}
-
-	FDocumentModifyDelegates::FDocumentModifyDelegates(FDocumentModifyDelegates&& InModifyDelegates)
-		: OnDocumentMetadataChanged(MoveTemp(InModifyDelegates.OnDocumentMetadataChanged))
-		, OnDependencyAdded(MoveTemp(InModifyDelegates.OnDependencyAdded))
-		, OnRemoveSwappingDependency(MoveTemp(InModifyDelegates.OnRemoveSwappingDependency))
-		, OnRenamingDependencyClass(MoveTemp(InModifyDelegates.OnRenamingDependencyClass))
-		, PageDelegates(MoveTemp(InModifyDelegates.PageDelegates))
-		, InterfaceDelegates(MoveTemp(InModifyDelegates.InterfaceDelegates))
-PRAGMA_DISABLE_DEPRECATION_WARNINGS
-		, NodeDelegates(MoveTemp(InModifyDelegates.NodeDelegates))
-		, EdgeDelegates(MoveTemp(InModifyDelegates.EdgeDelegates))
-PRAGMA_ENABLE_DEPRECATION_WARNINGS
-		, PageNodeDelegates(MoveTemp(InModifyDelegates.PageNodeDelegates))
-		, PageEdgeDelegates(MoveTemp(InModifyDelegates.PageEdgeDelegates))
-	{
-	}
-
-	FDocumentModifyDelegates& FDocumentModifyDelegates::operator=(const FDocumentModifyDelegates& InModifyDelegates)
-	{
-		OnDocumentMetadataChanged = InModifyDelegates.OnDocumentMetadataChanged;
-		OnDependencyAdded = InModifyDelegates.OnDependencyAdded;
-		OnRemoveSwappingDependency = InModifyDelegates.OnRemoveSwappingDependency;
-		OnRenamingDependencyClass = InModifyDelegates.OnRenamingDependencyClass;
-		PageDelegates = InModifyDelegates.PageDelegates;
-		InterfaceDelegates = InModifyDelegates.InterfaceDelegates;
-PRAGMA_DISABLE_DEPRECATION_WARNINGS
-		NodeDelegates = InModifyDelegates.NodeDelegates;
-		EdgeDelegates = InModifyDelegates.EdgeDelegates;
-PRAGMA_ENABLE_DEPRECATION_WARNINGS
-		PageNodeDelegates = InModifyDelegates.PageNodeDelegates;
-		PageEdgeDelegates = InModifyDelegates.PageEdgeDelegates;
-		return *this;
-	}
-
-	FDocumentModifyDelegates& FDocumentModifyDelegates::operator=(FDocumentModifyDelegates&& InModifyDelegates)
-	{
-		OnDocumentMetadataChanged = MoveTemp(InModifyDelegates.OnDocumentMetadataChanged);
-		OnDependencyAdded = MoveTemp(InModifyDelegates.OnDependencyAdded);
-		OnRemoveSwappingDependency = MoveTemp(InModifyDelegates.OnRemoveSwappingDependency);
-		OnRenamingDependencyClass = MoveTemp(InModifyDelegates.OnRenamingDependencyClass);
-		PageDelegates = MoveTemp(InModifyDelegates.PageDelegates);
-		InterfaceDelegates = MoveTemp(InModifyDelegates.InterfaceDelegates);
-PRAGMA_DISABLE_DEPRECATION_WARNINGS
-		NodeDelegates = MoveTemp(InModifyDelegates.NodeDelegates);
-		EdgeDelegates = MoveTemp(InModifyDelegates.EdgeDelegates);
-PRAGMA_ENABLE_DEPRECATION_WARNINGS
-		PageNodeDelegates = MoveTemp(InModifyDelegates.PageNodeDelegates);
-		PageEdgeDelegates = MoveTemp(InModifyDelegates.PageEdgeDelegates);
-		return *this;
+		AddPageDelegates(Frontend::DefaultPageID);
 	}
 
 	FDocumentModifyDelegates::FDocumentModifyDelegates(const FMetasoundFrontendDocument& Document)
-PRAGMA_DISABLE_DEPRECATION_WARNINGS
-		: NodeDelegates()
-		, EdgeDelegates()
-PRAGMA_ENABLE_DEPRECATION_WARNINGS
+		: FDocumentModifyDelegates()
 	{
 		Document.RootGraph.IterateGraphPages([this](const FMetasoundFrontendGraph& Graph)
 		{
@@ -93,35 +21,63 @@ PRAGMA_ENABLE_DEPRECATION_WARNINGS
 
 	void FDocumentModifyDelegates::AddPageDelegates(const FGuid& InPageID)
 	{
-		PageNodeDelegates.Add(InPageID, FNodeModifyDelegates());
-		PageEdgeDelegates.Add(InPageID, FEdgeModifyDelegates());
+		bool bDelegatesAdded = false;
+
+		if (!PageNodeDelegates.Contains(InPageID))
+		{
+			bDelegatesAdded = true;
+			PageNodeDelegates.Add(InPageID, FNodeModifyDelegates());
+		}
 
-		PageDelegates.OnPageAdded.Broadcast(FDocumentMutatePageArgs{ InPageID });
+		if (!PageEdgeDelegates.Contains(InPageID))
+		{
+			checkf(bDelegatesAdded, TEXT("All page delegates should be initialized together."));
+			PageEdgeDelegates.Add(InPageID, FEdgeModifyDelegates());
+		}
+		else
+		{
+			checkf(!bDelegatesAdded, TEXT("Page delegates should be initialized together, but were not reflective of same initial state when attempting to add delegates."));
+		}
+
+		if (bDelegatesAdded)
+		{
+			PageDelegates.OnPageAdded.Broadcast(FDocumentMutatePageArgs{ InPageID });
+		}
 	}
 
 	FNodeModifyDelegates& FDocumentModifyDelegates::FindNodeDelegatesChecked(const FGuid& InPageID)
 	{
-PRAGMA_DISABLE_DEPRECATION_WARNINGS
-		return InPageID == DefaultPageID
-		? NodeDelegates
-		: PageNodeDelegates.FindChecked(InPageID);
-PRAGMA_ENABLE_DEPRECATION_WARNINGS
+		return PageNodeDelegates.FindChecked(InPageID);
 	}
 
 	FEdgeModifyDelegates& FDocumentModifyDelegates::FindEdgeDelegatesChecked(const FGuid& InPageID)
 	{
-PRAGMA_DISABLE_DEPRECATION_WARNINGS
-		return InPageID == DefaultPageID
-		? EdgeDelegates
-		: PageEdgeDelegates.FindChecked(InPageID);
-PRAGMA_ENABLE_DEPRECATION_WARNINGS
+		return PageEdgeDelegates.FindChecked(InPageID);
+	}
+
+	void FDocumentModifyDelegates::IterateGraphEdgeDelegates(TFunctionRef<void(FEdgeModifyDelegates&)> Func)
+	{
+		for (TPair<FGuid, FEdgeModifyDelegates>& Delegates : PageEdgeDelegates)
+		{
+			Func(Delegates.Value);
+		}
+	}
+
+	void FDocumentModifyDelegates::IterateGraphNodeDelegates(TFunctionRef<void(FNodeModifyDelegates&)> Func)
+	{
+		for (TPair<FGuid, FNodeModifyDelegates>& Delegates : PageNodeDelegates)
+		{
+			Func(Delegates.Value);
+		}
 	}
 
 	void FDocumentModifyDelegates::RemovePageDelegates(const FGuid& InPageID, bool bBroadcastNotify)
 	{
+		checkf(InPageID != ::Metasound::Frontend::DefaultPageID, TEXT("Cannot remove delegates for the default page. Should always contain default page delegates."));
+
 		if (bBroadcastNotify)
 		{
-			PageDelegates.OnRemovingPage.Broadcast(FDocumentMutatePageArgs{ InPageID });
+			PageDelegates.OnRemovingPage.Broadcast(FDocumentMutatePageArgs { InPageID });
 		}
 
 		PageNodeDelegates.Remove(InPageID);
diff --git a/Engine/Plugins/Runtime/Metasound/Source/MetasoundFrontend/Public/MetasoundFrontendDocumentModifyDelegates.h b/Engine/Plugins/Runtime/Metasound/Source/MetasoundFrontend/Public/MetasoundFrontendDocumentModifyDelegates.h
--- a/Engine/Plugins/Runtime/Metasound/Source/MetasoundFrontend/Public/MetasoundFrontendDocumentModifyDelegates.h
+++ b/Engine/Plugins/Runtime/Metasound/Source/MetasoundFrontend/Public/MetasoundFrontendDocumentModifyDelegates.h
@@ -24,6 +24,12 @@ namespace Metasound::Frontend
 		FGuid PageID;
 	};
 
+	struct FDocumentPresetStateChangedArgs
+	{
+	};
+
+	DECLARE_MULTICAST_DELEGATE_OneParam(FOnDocumentPresetStateChanged, const FDocumentPresetStateChangedArgs& /* Args */);
+
 	DECLARE_MULTICAST_DELEGATE_OneParam(FOnDocumentPageAdded, const FDocumentMutatePageArgs& /* Args */);
 	DECLARE_MULTICAST_DELEGATE_OneParam(FOnDocumentRemovingPage, const FDocumentMutatePageArgs& /* Args */);
 	DECLARE_MULTICAST_DELEGATE_OneParam(FOnDocumentPageSet, const FDocumentMutatePageArgs& /* Args */);
@@ -87,60 +93,33 @@ namespace Metasound::Frontend
 	struct FDocumentModifyDelegates : TSharedFromThis<FDocumentModifyDelegates>
 	{
 		UE_API FDocumentModifyDelegates();
-		UE_API FDocumentModifyDelegates(const FDocumentModifyDelegates& InModifyDelegates);
-		UE_API FDocumentModifyDelegates(FDocumentModifyDelegates&&);
 		UE_API FDocumentModifyDelegates(const FMetasoundFrontendDocument& Document);
-		UE_API FDocumentModifyDelegates& operator=(const FDocumentModifyDelegates&);
-		UE_API FDocumentModifyDelegates& operator=(FDocumentModifyDelegates&&);
+
+		~FDocumentModifyDelegates() = default;
 
 		FOnMetaSoundFrontendDocumentMutateMetadata OnDocumentMetadataChanged;
 		FOnMetaSoundFrontendDocumentMutateArray OnDependencyAdded;
 		FOnMetaSoundFrontendDocumentRemoveSwappingArray OnRemoveSwappingDependency;
 		FOnMetaSoundFrontendDocumentRenameClass OnRenamingDependencyClass;
 
+		FOnDocumentPresetStateChanged OnPresetStateChanged;
+
 		FPageModifyDelegates PageDelegates;
 		FInterfaceModifyDelegates InterfaceDelegates;
 
-		UE_DEPRECATED(5.5, "Public exposition of NodeDelegates will be removed in a future build.  Use accessor 'FindNodeDelegates' instead")
-		FNodeModifyDelegates NodeDelegates;
-
-		UE_DEPRECATED(5.5, "Public exposition of EdgeDelegates will be removed in a future build.  Use accessor 'FindEdgeDelegates' instead")
-		FEdgeModifyDelegates EdgeDelegates;
-
-		UE_API void AddPageDelegates(const FGuid& InPageID);
-		UE_API void RemovePageDelegates(const FGuid& InPageID, bool bBroadcastNotify = true);
-
 	private:
 		TSortedMap<FGuid, FNodeModifyDelegates> PageNodeDelegates;
 		TSortedMap<FGuid, FEdgeModifyDelegates> PageEdgeDelegates;
 
 	public:
+		UE_API void AddPageDelegates(const FGuid& InPageID);
+		UE_API void RemovePageDelegates(const FGuid& InPageID, bool bBroadcastNotify = true);
+
 		UE_API FNodeModifyDelegates& FindNodeDelegatesChecked(const FGuid& InPageID);
 		UE_API FEdgeModifyDelegates& FindEdgeDelegatesChecked(const FGuid& InPageID);
 
-		void IterateGraphEdgeDelegates(TFunctionRef<void(FEdgeModifyDelegates&)> Func)
-		{
-PRAGMA_DISABLE_DEPRECATION_WARNINGS
-			Func(EdgeDelegates);
-
-			for (TPair<FGuid, FEdgeModifyDelegates>& Delegates : PageEdgeDelegates)
-			{
-				Func(Delegates.Value);
-			}
-PRAGMA_ENABLE_DEPRECATION_WARNINGS
-		}
-
-		void IterateGraphNodeDelegates(TFunctionRef<void(FNodeModifyDelegates&)> Func)
-		{
-PRAGMA_DISABLE_DEPRECATION_WARNINGS
-			Func(NodeDelegates);
-
-			for (TPair<FGuid, FNodeModifyDelegates>& Delegates : PageNodeDelegates)
-			{
-				Func(Delegates.Value);
-			}
-PRAGMA_ENABLE_DEPRECATION_WARNINGS
-		}
+		UE_API void IterateGraphEdgeDelegates(TFunctionRef<void(FEdgeModifyDelegates&)> Func);
+		UE_API void IterateGraphNodeDelegates(TFunctionRef<void(FNodeModifyDelegates&)> Func);
 	};
 
 	class IDocumentBuilderTransactionListener : public TSharedFromThis<IDocumentBuilderTransactionListener>
