Fix MetaSound invalid warning on missing task removal by tracking registration tasks.
#rnx
#FailedValidations Preflight

diff --git a/Engine/Plugins/Runtime/Metasound/Source/MetasoundFrontend/Private/MetasoundFrontendNodeClassRegistry.cpp b/Engine/Plugins/Runtime/Metasound/Source/MetasoundFrontend/Private/MetasoundFrontendNodeClassRegistry.cpp
--- a/Engine/Plugins/Runtime/Metasound/Source/MetasoundFrontend/Private/MetasoundFrontendNodeClassRegistry.cpp
+++ b/Engine/Plugins/Runtime/Metasound/Source/MetasoundFrontend/Private/MetasoundFrontendNodeClassRegistry.cpp
@@ -620,16 +620,19 @@ PRAGMA_ENABLE_DEPRECATION_WARNINGS
 
 		if (bAsync)
 		{
+			// Keep track of exact registration task to ensure correct task instance is removed upon task completion. 
+			int32 RegistrationTaskID = ++RegistrationTaskCntr;
+			
 			TArray<FGuid> PageOrder(InPageOrder.GetData(), InPageOrder.Num());
 			Tasks::FTask BuildAndRegisterTask = AsyncRegistrationPipe.Launch(
 				UE_SOURCE_LOCATION,
-				[RegistryKey, ClassInfo = MoveTemp(NodeClassInfo), AssetPath, RegistryDocInterface, ProxyDataCache = MoveTemp(ProxyDataCache), PageOrder = MoveTemp(PageOrder)]() mutable
+				[RegistryKey, ClassInfo = MoveTemp(NodeClassInfo), AssetPath, RegistryDocInterface, ProxyDataCache = MoveTemp(ProxyDataCache), PageOrder = MoveTemp(PageOrder), RegistrationTaskID]() mutable
 				{
 					FNodeClassRegistry& Registry = FNodeClassRegistry::Get();
 					// Unregister the graph before re-registering
 					Registry.UnregisterGraphInternal(RegistryKey);
 					Registry.BuildAndRegisterGraphFromDocument(RegistryDocInterface->GetConstDocument(), ProxyDataCache, MoveTemp(ClassInfo), AssetPath, PageOrder);
-					Registry.RemoveRegistrationTask(RegistryKey, FNodeClassRegistryTransaction::ETransactionType::NodeRegistration);
+					Registry.RemoveRegistrationTask(RegistryKey, RegistrationTaskID, FNodeClassRegistryTransaction::ETransactionType::NodeRegistration);
 					Registry.RemoveDocumentReference(RegistryDocInterface);
 				}
 			);
@@ -639,7 +642,8 @@ PRAGMA_ENABLE_DEPRECATION_WARNINGS
 			{
 				FNodeClassRegistryTransaction::ETransactionType::NodeRegistration,
 				BuildAndRegisterTask,
-				AssetPath
+				AssetPath,
+				RegistrationTaskID
 			});
 		}
 		else
@@ -659,22 +663,18 @@ PRAGMA_ENABLE_DEPRECATION_WARNINGS
 		ActiveRegistrationTasks.FindOrAdd(InKey.NodeKey).Add(MoveTemp(TaskInfo));
 	}
 
-	void FNodeClassRegistry::RemoveRegistrationTask(const FGraphRegistryKey& InKey, FNodeClassRegistryTransaction::ETransactionType TransactionType)
+	void FNodeClassRegistry::RemoveRegistrationTask(const FGraphRegistryKey& InKey, int32 InRegistrationTaskID, FNodeClassRegistryTransaction::ETransactionType TransactionType)
 	{
 		UE::TScopeLock LockActiveReg(ActiveRegistrationTasksCriticalSection);
 
 		int32 NumRemoved = 0;
-		if (InKey.AssetPath.IsNull())  // Null provided path instructs to remove all tasks related to the underlying Node Registry Key
+		if (TArray<FActiveRegistrationTaskInfo>* TaskInfos = ActiveRegistrationTasks.Find(InKey.NodeKey))
 		{
-			NumRemoved = ActiveRegistrationTasks.Remove(InKey.NodeKey);
-		}
-		else if (TArray<FActiveRegistrationTaskInfo>* TaskInfos = ActiveRegistrationTasks.Find(InKey.NodeKey))
-		{
-			auto MatchesEntryInTask = [&InKey, &TransactionType](const FActiveRegistrationTaskInfo& Info)
+			auto MatchesEntryInTask = [&TransactionType, InRegistrationTaskID](const FActiveRegistrationTaskInfo& Info)
 			{
-				const bool bIsAssetPath = Info.AssetPath == InKey.AssetPath;
 				const bool bIsTransactionType = Info.TransactionType == TransactionType;
-				return bIsAssetPath && bIsTransactionType;
+				const bool bIsSameRegistrationTaskID = Info.RegistrationTaskID == InRegistrationTaskID;
+				return bIsTransactionType && bIsSameRegistrationTaskID;
 			};
 
 			NumRemoved = TaskInfos->RemoveAllSwap(MatchesEntryInTask, EAllowShrinking::No);
@@ -820,21 +820,25 @@ PRAGMA_ENABLE_DEPRECATION_WARNINGS
 		const bool bAsync = !(InDocumentInterface->IsActivelyBuilding() || ConsoleVariables::bDisableAsyncGraphRegistration);
 		if (bAsync)
 		{
+			// Keep track of exact registration task to ensure correct task instance is removed upon task completion. 
+			int32 RegistrationTaskID = ++RegistrationTaskCntr;
+
 			// Wait for any async tasks that are in flight which correspond to the same graph
 			WaitForAsyncGraphRegistration(InRegistryKey);
 
-			Tasks::FTask UnregisterTask = AsyncRegistrationPipe.Launch(UE_SOURCE_LOCATION, [RegistryKey = InRegistryKey]()
+			Tasks::FTask UnregisterTask = AsyncRegistrationPipe.Launch(UE_SOURCE_LOCATION, [RegistryKey = InRegistryKey, RegistrationTaskID]()
 			{
 				FNodeClassRegistry& Registry = FNodeClassRegistry::Get();
 				Registry.UnregisterGraphInternal(RegistryKey);
-				Registry.RemoveRegistrationTask(RegistryKey, FNodeClassRegistryTransaction::ETransactionType::NodeUnregistration);
+				Registry.RemoveRegistrationTask(RegistryKey, RegistrationTaskID, FNodeClassRegistryTransaction::ETransactionType::NodeUnregistration);
 			});
 
 			AddRegistrationTask(InRegistryKey, FActiveRegistrationTaskInfo
 			{
 				FNodeClassRegistryTransaction::ETransactionType::NodeUnregistration,
 				UnregisterTask,
-				InRegistryKey.AssetPath
+				InRegistryKey.AssetPath,
+				RegistrationTaskID
 			});
 		}
 		else
diff --git a/Engine/Plugins/Runtime/Metasound/Source/MetasoundFrontend/Private/MetasoundFrontendNodeClassRegistryPrivate.h b/Engine/Plugins/Runtime/Metasound/Source/MetasoundFrontend/Private/MetasoundFrontendNodeClassRegistryPrivate.h
--- a/Engine/Plugins/Runtime/Metasound/Source/MetasoundFrontend/Private/MetasoundFrontendNodeClassRegistryPrivate.h
+++ b/Engine/Plugins/Runtime/Metasound/Source/MetasoundFrontend/Private/MetasoundFrontendNodeClassRegistryPrivate.h
@@ -152,12 +152,13 @@ namespace Metasound::Frontend
 			FNodeClassRegistryTransaction::ETransactionType TransactionType = FNodeClassRegistryTransaction::ETransactionType::NodeRegistration;
 			UE::Tasks::FTask Task;
 			FTopLevelAssetPath AssetPath;
+			int32 RegistrationTaskID = 0;
 		};
 
 		void BuildAndRegisterGraphFromDocument(const FMetasoundFrontendDocument& InDocument, const FProxyDataCache& InProxyDataCache, FNodeClassInfo&& InNodeClassInfo, const FTopLevelAssetPath& AssetPath, TArrayView<const FGuid> InPageOrder);
 
 		void AddRegistrationTask(const FGraphRegistryKey& InKey, FActiveRegistrationTaskInfo&& TaskInfo);
-		void RemoveRegistrationTask(const FGraphRegistryKey& InKey, FNodeClassRegistryTransaction::ETransactionType TransactionType);
+		void RemoveRegistrationTask(const FGraphRegistryKey& InKey, int32 InRegistrationTaskID, FNodeClassRegistryTransaction::ETransactionType TransactionType);
 
 		// Adds reference to document's owning UObject to internal ObjectReferencer,
 		// indicating async registration/unregistration task(s) are depending on it.
@@ -220,6 +221,7 @@ namespace Metasound::Frontend
 		mutable FTransactionallySafeCriticalSection ObjectReferencerCriticalSection;
 
 		UE::Tasks::FPipe AsyncRegistrationPipe;
+		int32 RegistrationTaskCntr = 0;
 		TMap<FNodeClassRegistryKey, TArray<FActiveRegistrationTaskInfo>> ActiveRegistrationTasks;
 		TUniquePtr<INodeClassRegistry::IObjectReferencer> ObjectReferencer;
 	};
