Experimental Support for Polymorphic Datatypes in Metasound

#tests LowLevelTest Audio
[RN] Experimental Support for Polymorphic Datatypes in Metasound
#DisabledValidations Preflight

diff --git a/Engine/Plugins/Runtime/Metasound/Source/MetasoundEditor/Private/MetasoundEditor.cpp b/Engine/Plugins/Runtime/Metasound/Source/MetasoundEditor/Private/MetasoundEditor.cpp
--- a/Engine/Plugins/Runtime/Metasound/Source/MetasoundEditor/Private/MetasoundEditor.cpp
+++ b/Engine/Plugins/Runtime/Metasound/Source/MetasoundEditor/Private/MetasoundEditor.cpp
@@ -540,7 +540,7 @@ namespace Metasound
 					OutConnectability.Connectable = FConnectability::EConnectable::No;
 					OutConnectability.Reason = FConnectability::EReason::IncompatibleDataTypes;
 				}
-				else if (DataType0 == DataType1)
+				else if (IsCastable(DataType0, DataType1))
 				{
 					OutConnectability.Connectable = FConnectability::EConnectable::Yes;
 					OutConnectability.Reason = FConnectability::EReason::None;
diff --git a/Engine/Plugins/Runtime/Metasound/Source/MetasoundEditor/Private/MetasoundEditorGraphSchema.cpp b/Engine/Plugins/Runtime/Metasound/Source/MetasoundEditor/Private/MetasoundEditorGraphSchema.cpp
--- a/Engine/Plugins/Runtime/Metasound/Source/MetasoundEditor/Private/MetasoundEditorGraphSchema.cpp
+++ b/Engine/Plugins/Runtime/Metasound/Source/MetasoundEditor/Private/MetasoundEditorGraphSchema.cpp
@@ -589,7 +589,7 @@ PRAGMA_ENABLE_DEPRECATION_WARNINGS
 			{
 				const bool bCanConnect = Algo::AnyOf(Class.GetDefaultInterface().Outputs, [&InputTypeName, &InputAccessType](const FMetasoundFrontendClassOutput& Output)
 				{
-					return InputTypeName == Output.TypeName && FMetasoundFrontendClassVertex::CanConnectVertexAccessTypes(Output.AccessType, InputAccessType);
+					return IsCastable(InputTypeName,Output.TypeName) && FMetasoundFrontendClassVertex::CanConnectVertexAccessTypes(Output.AccessType, InputAccessType);
 				});
 				return bCanConnect;
 			}
@@ -598,7 +598,7 @@ PRAGMA_ENABLE_DEPRECATION_WARNINGS
 			{
 				const bool bCanConnect = Algo::AnyOf(Class.GetDefaultInterface().Inputs, [&OutputTypeName, &OutputAccessType](const FMetasoundFrontendClassInput& Input)
 				{
-					return OutputTypeName == Input.TypeName && FMetasoundFrontendClassVertex::CanConnectVertexAccessTypes(OutputAccessType, Input.AccessType);
+					return IsCastable(OutputTypeName, Input.TypeName) && FMetasoundFrontendClassVertex::CanConnectVertexAccessTypes(OutputAccessType, Input.AccessType);
 				});
 				return bCanConnect;
 			}
diff --git a/Engine/Plugins/Runtime/Metasound/Source/MetasoundFrontend/Private/MetasoundDataTypeRegistrationMacro.cpp b/Engine/Plugins/Runtime/Metasound/Source/MetasoundFrontend/Private/MetasoundDataTypeRegistrationMacro.cpp
--- a/Engine/Plugins/Runtime/Metasound/Source/MetasoundFrontend/Private/MetasoundDataTypeRegistrationMacro.cpp
+++ b/Engine/Plugins/Runtime/Metasound/Source/MetasoundFrontend/Private/MetasoundDataTypeRegistrationMacro.cpp
@@ -6,6 +6,7 @@ namespace Metasound::Frontend
 {
 	namespace MetasoundDataTypeRegistrationPrivate
 	{
+
 		// Base registry entry for any data type
 		FDataTypeRegistryEntryBase::FDataTypeRegistryEntryBase(const Frontend::FDataTypeRegistryInfo& Info, const TSharedPtr<Frontend::IEnumDataTypeInterface>& EnumInterface)
 			: Info(Info), EnumInterface(EnumInterface)
@@ -24,6 +25,11 @@ namespace Metasound::Frontend
 			return EnumInterface;
 		}
 
+		TSharedPtr<IPolymorphicDataTypeInterface> FDataTypeRegistryEntryBase::GetPolymorphicInterface() const
+		{
+			return PolymorphicInterface;
+		}
+		
 		const FMetasoundFrontendClass& FDataTypeRegistryEntryBase::GetFrontendInputClass() const
 		{
 			return InputClass;
diff --git a/Engine/Plugins/Runtime/Metasound/Source/MetasoundFrontend/Private/MetasoundFrontendDataTypeRegistry.cpp b/Engine/Plugins/Runtime/Metasound/Source/MetasoundFrontend/Private/MetasoundFrontendDataTypeRegistry.cpp
--- a/Engine/Plugins/Runtime/Metasound/Source/MetasoundFrontend/Private/MetasoundFrontendDataTypeRegistry.cpp
+++ b/Engine/Plugins/Runtime/Metasound/Source/MetasoundFrontend/Private/MetasoundFrontendDataTypeRegistry.cpp
@@ -1525,6 +1525,57 @@ PRAGMA_ENABLE_DEPRECATION_WARNINGS
 			}
 		}
 
+		void FPolymorphicDataTypeBase::ResolveParent(const IDataTypeRegistry& InRegistry) const
+		{
+			static const FName VoidString = TEXT("void");
+			if (!ParentDataType.IsValid() && !ParentTypeName.IsNone() && ParentTypeName != VoidString)
+			{
+				if (const IDataTypeRegistryEntry* Entry = InRegistry.FindDataTypeRegistryEntry(ParentTypeName))
+				{
+					PRAGMA_DISABLE_EXPERIMENTAL_WARNINGS					
+					ParentDataType = Entry->GetPolymorphicInterface();
+					PRAGMA_ENABLE_EXPERIMENTAL_WARNINGS					
+				}
+				UE_CLOG(!ParentDataType.IsValid(), LogMetaSound, Warning, TEXT("Failed to resolve Polymorphic Type '%s' Parent Type '%s'"),
+					*DataTypeName.ToString(), *ParentTypeName.ToString());
+			}
+		}
+
+		FPolymorphicDataTypeBase::FPolymorphicDataTypeBase(const FName InDataTypeName, const FName InParentName)
+			: DataTypeName(InDataTypeName), ParentTypeName(InParentName)
+		{
+			// Register with Core Registry.
+			PRAGMA_DISABLE_INTERNAL_WARNINGS
+			FPolyRegistry::Get().Register(
+				{
+					.TypeName = InDataTypeName,
+					.BaseTypeName = InParentName,
+					.bIsPolymorphic = true
+					}
+				);
+			PRAGMA_ENABLE_INTERNAL_WARNINGS			
+		}
+
+		bool FPolymorphicDataTypeBase::IsA(const FName InDataTypeName, const IDataTypeRegistry& InRegistry) const
+		{
+			if (DataTypeName == InDataTypeName)
+			{
+				return true;
+			}
+
+			// Resolve.
+			ResolveParent(InRegistry);
+		
+			if (const TSharedPtr<const IPolymorphicDataTypeInterface> Parent = ParentDataType.Pin())
+			{
+				PRAGMA_DISABLE_EXPERIMENTAL_WARNINGS
+				return Parent->IsA(InDataTypeName);
+				PRAGMA_ENABLE_EXPERIMENTAL_WARNINGS				
+			}
+			
+			return false;
+		}
+
 		IDataTypeRegistry& IDataTypeRegistry::Get()
 		{
 			static MetasoundFrontendDataTypeRegistryPrivate::FDataTypeRegistry Registry;
diff --git a/Engine/Plugins/Runtime/Metasound/Source/MetasoundFrontend/Private/MetasoundFrontendDocumentBuilder.cpp b/Engine/Plugins/Runtime/Metasound/Source/MetasoundFrontend/Private/MetasoundFrontendDocumentBuilder.cpp
--- a/Engine/Plugins/Runtime/Metasound/Source/MetasoundFrontend/Private/MetasoundFrontendDocumentBuilder.cpp
+++ b/Engine/Plugins/Runtime/Metasound/Source/MetasoundFrontend/Private/MetasoundFrontendDocumentBuilder.cpp
@@ -4440,7 +4440,7 @@ Metasound::Frontend::EInvalidEdgeReason FMetaSoundFrontendDocumentBuilder::IsVal
 		return EInvalidEdgeReason::MissingInput;
 	}
 
-	if (OutputVertex->TypeName != InputVertex->TypeName)
+	if (!Metasound::IsCastable(OutputVertex->TypeName, InputVertex->TypeName))
 	{
 		return EInvalidEdgeReason::MismatchedDataType;
 	}
diff --git a/Engine/Plugins/Runtime/Metasound/Source/MetasoundFrontend/Private/MetasoundFrontendInputController.cpp b/Engine/Plugins/Runtime/Metasound/Source/MetasoundFrontend/Private/MetasoundFrontendInputController.cpp
--- a/Engine/Plugins/Runtime/Metasound/Source/MetasoundFrontend/Private/MetasoundFrontendInputController.cpp
+++ b/Engine/Plugins/Runtime/Metasound/Source/MetasoundFrontend/Private/MetasoundFrontendInputController.cpp
@@ -251,7 +251,7 @@ namespace Metasound
 				OutConnectability.Connectable = FConnectability::EConnectable::No;
 				OutConnectability.Reason = FConnectability::EReason::IncompatibleAccessTypes;
 			}
-			else if (OtherDataType == DataType)
+			else if (IsCastable(OtherDataType, DataType))
 			{
 				// If data types are equal, connection can happen.
 				OutConnectability.Connectable = FConnectability::EConnectable::Yes;
@@ -295,7 +295,7 @@ namespace Metasound
 
 			if (FMetasoundFrontendGraph* Graph = GraphPtr.Get())
 			{
-				if (OtherDataType == DataType)
+				if (IsCastable(OtherDataType, DataType))
 				{
 					if (FMetasoundFrontendClassVertex::CanConnectVertexAccessTypes(InController.GetVertexAccessType(), GetVertexAccessType()))
 					{
diff --git a/Engine/Plugins/Runtime/Metasound/Source/MetasoundFrontend/Private/MetasoundFrontendNodeClassRegistry.cpp b/Engine/Plugins/Runtime/Metasound/Source/MetasoundFrontend/Private/MetasoundFrontendNodeClassRegistry.cpp
--- a/Engine/Plugins/Runtime/Metasound/Source/MetasoundFrontend/Private/MetasoundFrontendNodeClassRegistry.cpp
+++ b/Engine/Plugins/Runtime/Metasound/Source/MetasoundFrontend/Private/MetasoundFrontendNodeClassRegistry.cpp
@@ -889,8 +889,8 @@ PRAGMA_ENABLE_DEPRECATION_WARNINGS
 		TArray<TSharedRef<INodeClassRegistryEntry>> Entries;
 #endif // !NO_LOGGING
 
+		TSharedRef<INodeClassRegistryEntry, ESPMode::ThreadSafe> Entry(InEntry.Release());
 		{
-			TSharedRef<INodeClassRegistryEntry, ESPMode::ThreadSafe> Entry(InEntry.Release());
 			UE::TScopeLock Lock(RegistryMapsCriticalSection);
 			// check to see if an identical node was already registered, and log if necessary
 			// Store registry elements in map so nodes can be queried using registry key.
@@ -901,20 +901,27 @@ PRAGMA_ENABLE_DEPRECATION_WARNINGS
 #endif // !NO_LOGGING
 		}
 
-#if !NO_LOGGING
+#if !NO_LOGGING 
 		if (Entries.Num() > 1)
 		{
 			if (IMetaSoundAssetManager* AssetManager = IMetaSoundAssetManager::Get())
 			{
-				const TArray<FTopLevelAssetPath> AssetPaths = AssetManager->FindAssetPaths(FMetaSoundAssetKey(Key));
-				TArray<FString> ExistingPaths;
-				Algo::Transform(AssetPaths, ExistingPaths, [](const FTopLevelAssetPath& AssetPath) { return AssetPath.ToString(); });
-				const FString ExistingAssetPaths = FString::Join(ExistingPaths, TEXT("\n"));
-				UE_LOG(LogMetaSound, Error,
-					TEXT("Multiple node classes with key '%s' registered. Assets currently registered with class name:\n%s"),
-						* Key.ToString(),
-						*ExistingAssetPaths
-				);
+				if (AssetManager->IsAssetClass(Entry->GetFrontendClass().Metadata))
+				{
+					const TArray<FTopLevelAssetPath> AssetPaths = AssetManager->FindAssetPaths(FMetaSoundAssetKey(Key));
+					TArray<FString> ExistingPaths;
+					Algo::Transform(AssetPaths, ExistingPaths, [](const FTopLevelAssetPath& AssetPath) { return AssetPath.ToString(); });
+					const FString ExistingAssetPaths = FString::Join(ExistingPaths, TEXT("\n"));
+					UE_LOG(LogMetaSound, Error,
+						TEXT("Multiple node classes with key '%s' registered. Assets currently registered with class name:\n%s"),
+							* Key.ToString(),
+							*ExistingAssetPaths
+					);
+				}
+				else
+				{
+					UE_LOG(LogMetaSound, Error, TEXT("Multiple node classes with key '%s' registered."),*Key.ToString());
+				}
 			}
 		}
 #endif // !NO_LOGGING
diff --git a/Engine/Plugins/Runtime/Metasound/Source/MetasoundFrontend/Public/MetasoundDataTypeRegistrationMacro.h b/Engine/Plugins/Runtime/Metasound/Source/MetasoundFrontend/Public/MetasoundDataTypeRegistrationMacro.h
--- a/Engine/Plugins/Runtime/Metasound/Source/MetasoundFrontend/Public/MetasoundDataTypeRegistrationMacro.h
+++ b/Engine/Plugins/Runtime/Metasound/Source/MetasoundFrontend/Public/MetasoundDataTypeRegistrationMacro.h
@@ -245,6 +245,10 @@ namespace Metasound::Frontend
 				RegistryInfo.ProxyGeneratorClass = nullptr;
 			}
 
+			using Base = typename TPolymorphicTraits<TDataType>::BaseType;
+			RegistryInfo.bIsPolymorphic = TPolymorphicTraits<TDataType>::bIsPolymorphic;
+			RegistryInfo.bIsAbstract = TPolymorphicTraits<TDataType>::bIsAbstract;
+			RegistryInfo.ParentDataTypeName = GetMetasoundDataTypeName<Base>();
 			return RegistryInfo;
 		}
 
@@ -316,11 +320,12 @@ namespace Metasound::Frontend
 		class FDataTypeRegistryEntryBase : public IDataTypeRegistryEntry
 		{
 		public:
-			UE_API FDataTypeRegistryEntryBase(const Frontend::FDataTypeRegistryInfo& Info, const TSharedPtr<Frontend::IEnumDataTypeInterface>& EnumInterface);
-			UE_API virtual ~FDataTypeRegistryEntryBase();
-
+		 	UE_API FDataTypeRegistryEntryBase(const Frontend::FDataTypeRegistryInfo& Info, const TSharedPtr<Frontend::IEnumDataTypeInterface>& EnumInterface);
+			UE_API virtual ~FDataTypeRegistryEntryBase() override;
+			
 			UE_API virtual const Frontend::FDataTypeRegistryInfo& GetDataTypeInfo() const override;
 			UE_API virtual TSharedPtr<const Frontend::IEnumDataTypeInterface> GetEnumInterface() const override;
+			UE_API virtual TSharedPtr<IPolymorphicDataTypeInterface> GetPolymorphicInterface() const override;			
 			UE_API virtual const FMetasoundFrontendClass& GetFrontendInputClass() const override;
 			UE_API virtual TSharedPtr<const FNodeClassMetadata> GetInputClassMetadata() const override;
 			UE_API virtual const FMetasoundFrontendClass& GetFrontendConstructorInputClass() const override;
@@ -365,6 +370,7 @@ namespace Metasound::Frontend
 			FMetasoundFrontendClass VariableDeferredAccessorClass;
 			TSharedPtr<const FNodeClassMetadata> VariableDeferredAccessorClassMetadata;
 			TSharedPtr<Frontend::IEnumDataTypeInterface> EnumInterface;
+			TSharedPtr<IPolymorphicDataTypeInterface> PolymorphicInterface;
 			Frontend::IParameterAssignmentFunction RawAssignmentFunction;
 			Frontend::FLiteralAssignmentFunction LiteralAssignmentFunction = nullptr;
 
@@ -503,7 +509,25 @@ namespace Metasound::Frontend
 						}
 					}
 				}
+				UE_EXPERIMENTAL(5.7, "Polymorphic API is experimental")
+				void InitPolymorphicInterface()
+				{
+					using FPolyTraits = TPolymorphicTraits<TDataType>;
 
+					if constexpr (FPolyTraits::bIsPolymorphic)
+					{
+						using BaseType = FPolyTraits::BaseType;
+						const FName ParentTypeName = GetMetasoundDataTypeName<BaseType>();
+						const FName TypeName = GetMetasoundDataTypeName<TDataType>();
+						this->PolymorphicInterface = MakeShared<FPolymorphicDataTypeBase>(TypeName, ParentTypeName);
+						UE_LOG(LogMetaSound, Verbose, TEXT("PolymorphicInterface Type=%s, Base=%s, bIsAbstract=%d"),
+							*TypeName.ToString(), *ParentTypeName.ToString(), (int32)FPolyTraits::bIsAbstract);
+					}
+					else
+					{
+						UE_CLOG(EnumInterface.IsValid(), LogMetaSound, Display, TEXT("PolymorphicInterface null for enum %s"), *GetMetasoundDataTypeName<TDataType>().ToString());
+					}
+				}
 
 			public:
 
@@ -513,6 +537,9 @@ namespace Metasound::Frontend
 					InitRawAssignmentFunction();
 					InitLiteralAssignmentFunction();
 					InitNodeClasses();
+					PRAGMA_DISABLE_EXPERIMENTAL_WARNINGS
+					InitPolymorphicInterface();
+					PRAGMA_ENABLE_EXPERIMENTAL_WARNINGS
 				}
 
 				UE_DEPRECATED(5.6, "Create input nodes using FNodeData")
diff --git a/Engine/Plugins/Runtime/Metasound/Source/MetasoundFrontend/Public/MetasoundFrontendDataTypeRegistry.h b/Engine/Plugins/Runtime/Metasound/Source/MetasoundFrontend/Public/MetasoundFrontendDataTypeRegistry.h
--- a/Engine/Plugins/Runtime/Metasound/Source/MetasoundFrontend/Public/MetasoundFrontendDataTypeRegistry.h
+++ b/Engine/Plugins/Runtime/Metasound/Source/MetasoundFrontend/Public/MetasoundFrontendDataTypeRegistry.h
@@ -85,6 +85,13 @@ namespace Metasound
 			{
 				return bIsProxyParsable || bIsProxyArrayParsable;
 			}
+
+			// If this type is polymorphic.
+			bool bIsPolymorphic = false;
+			bool bIsAbstract = false;
+
+			// If this type is polymorphic, what is its parent.
+			FName ParentDataTypeName = {};
 		};
 
 		/** Interface for metadata of a registered MetaSound enum type. */
@@ -134,6 +141,8 @@ namespace Metasound
 			}
 		};
 
+		struct IPolymorphicDataTypeInterface;
+	
 		/** Registry entry interface for a MetaSound data type. */
 		class IDataTypeRegistryEntry
 		{
@@ -148,6 +157,9 @@ namespace Metasound
 			 */
 			virtual TSharedPtr<const IEnumDataTypeInterface> GetEnumInterface() const = 0;
 
+			UE_EXPERIMENTAL(5.7, "Metasound Polymorphic Datatypes are still in devleopement")
+			virtual TSharedPtr<IPolymorphicDataTypeInterface> GetPolymorphicInterface() const = 0;
+
 			/** Return an FMetasoundFrontendClass representing an input node of the data type. */
 			virtual const FMetasoundFrontendClass& GetFrontendInputClass() const = 0;
 
@@ -434,6 +446,36 @@ namespace Metasound
 			virtual TUniquePtr<INode> CreateVariableDeferredAccessorNode(const FName& InDataType, FNodeData InNodeData) const = 0;
 		};
 
+		/**
+		 * IPolymorphicDataTypeInterface.
+		 * Experimental API. 
+		 */	
+		struct IPolymorphicDataTypeInterface
+		{
+			virtual ~IPolymorphicDataTypeInterface() = default;
+			UE_EXPERIMENTAL(5.7, "Metasound Polymorphic Datatypes are still in devleopement")
+			UE_API virtual bool IsA(const FName InDataTypeName, const IDataTypeRegistry& InRegistry = IDataTypeRegistry::Get() ) const = 0;
+		};
+
+		/**
+		 * FPolymorphicDataTypeBase 
+		 * Experimental API. 
+		 */
+		class FPolymorphicDataTypeBase : public IPolymorphicDataTypeInterface
+		{
+		public:
+			UE_API explicit FPolymorphicDataTypeBase(const FName InDataTypeName, const FName InParentName = NAME_None);
+			virtual ~FPolymorphicDataTypeBase() override = default;
+
+			UE_EXPERIMENTAL(5.7, "Metasound Polymorphic Datatypes are still in developement")			
+			UE_API virtual bool IsA(const FName InDataTypeName, const IDataTypeRegistry& InRegistry) const override;
+		private:
+			void ResolveParent(const IDataTypeRegistry& InRegistry) const;
+			FName DataTypeName = NAME_None;
+			FName ParentTypeName = NAME_None;
+			mutable TWeakPtr<const IPolymorphicDataTypeInterface> ParentDataType;
+		};
+		
 		/** Create data references for all inputs on the FInputVertexInterfaceData 
 		 * using the default literals which live on the FInputDataVertexes.
 		 *
diff --git a/Engine/Plugins/Runtime/Metasound/Source/MetasoundGraphCore/Private/MetasoundGraph.cpp b/Engine/Plugins/Runtime/Metasound/Source/MetasoundGraphCore/Private/MetasoundGraph.cpp
--- a/Engine/Plugins/Runtime/Metasound/Source/MetasoundGraphCore/Private/MetasoundGraph.cpp
+++ b/Engine/Plugins/Runtime/Metasound/Source/MetasoundGraphCore/Private/MetasoundGraph.cpp
@@ -6,6 +6,7 @@
 #include "Containers/Array.h"
 #include "Containers/Map.h"
 #include "MetasoundOperatorBuilder.h"
+#include "MetasoundPolymorphic.h"
 #include "Templates/SharedPointer.h"
 
 namespace Metasound
@@ -179,7 +180,7 @@ PRAGMA_ENABLE_DEPRECATION_WARNINGS
 		const FInputDataVertex& ToVertex = ToVertexInterface.GetInputVertex(ToKey);
 
 
-		if (FromVertex.DataTypeName != ToVertex.DataTypeName)
+		if (!IsCastable(FromVertex.DataTypeName, ToVertex.DataTypeName))
 		{
 			return false;
 		}
diff --git a/Engine/Plugins/Runtime/Metasound/Source/MetasoundGraphCore/Private/MetasoundGraphLinter.cpp b/Engine/Plugins/Runtime/Metasound/Source/MetasoundGraphCore/Private/MetasoundGraphLinter.cpp
--- a/Engine/Plugins/Runtime/Metasound/Source/MetasoundGraphCore/Private/MetasoundGraphLinter.cpp
+++ b/Engine/Plugins/Runtime/Metasound/Source/MetasoundGraphCore/Private/MetasoundGraphLinter.cpp
@@ -66,7 +66,7 @@ namespace Metasound
 
 		for (const FDataEdge& Edge : InGraph.GetDataEdges())
 		{
-			if (Edge.From.Vertex.DataTypeName != Edge.To.Vertex.DataTypeName)
+			if (!IsCastable(Edge.From.Vertex.DataTypeName, Edge.To.Vertex.DataTypeName))
 			{
 				AddBuildError<FInvalidConnectionDataTypeError>(OutErrors, Edge);
 
diff --git a/Engine/Plugins/Runtime/Metasound/Source/MetasoundGraphCore/Private/MetasoundOperatorBuilder.cpp b/Engine/Plugins/Runtime/Metasound/Source/MetasoundGraphCore/Private/MetasoundOperatorBuilder.cpp
--- a/Engine/Plugins/Runtime/Metasound/Source/MetasoundGraphCore/Private/MetasoundOperatorBuilder.cpp
+++ b/Engine/Plugins/Runtime/Metasound/Source/MetasoundGraphCore/Private/MetasoundOperatorBuilder.cpp
@@ -471,7 +471,7 @@ namespace Metasound
 
 			if (const FAnyDataReference* DataReference = InExternalInputData.FindDataReference(Destination.Vertex.VertexName))
 			{
-				if (DataReference->GetDataTypeName() == Destination.Vertex.DataTypeName)
+				if (IsCastable(DataReference->GetDataTypeName(), Destination.Vertex.DataTypeName))
 				{
 					OperatorInfo.VertexData.GetInputs().SetVertex(Destination.Vertex.VertexName, *DataReference);
 				}
diff --git a/Engine/Plugins/Runtime/Metasound/Source/MetasoundGraphCore/Private/MetasoundPolymorphic.cpp b/Engine/Plugins/Runtime/Metasound/Source/MetasoundGraphCore/Private/MetasoundPolymorphic.cpp
new file mode 100644
--- /dev/null
+++ b/Engine/Plugins/Runtime/Metasound/Source/MetasoundGraphCore/Private/MetasoundPolymorphic.cpp
@@ -0,0 +1,60 @@
+// Copyright Epic Games, Inc. All Rights Reserved.
+
+#include "MetasoundPolymorphic.h"
+
+namespace Metasound
+{
+	FPolyRegistry& FPolyRegistry::Get()
+	{
+		static FPolyRegistry Instance;
+		return Instance;
+	}
+
+	void FPolyRegistry::Register(const FPolyTypeInfo& Info)
+	{
+		Map.Add(Info.TypeName, Info);
+	}
+
+	void FPolyRegistry::Unregister(const FPolyTypeInfo& Info)
+	{
+		Map.Remove(Info.TypeName);
+	}
+
+	const FPolyTypeInfo* FPolyRegistry::Find(const FName Type) const
+	{
+		return Map.Find(Type);
+	}
+
+	void FPolyRegistry::GetAllRegisteredTypes(TArray<FPolyTypeInfo>& Out) const
+	{
+		Map.GenerateValueArray(Out);
+	}
+
+	bool IsChildOfByName(const FName InDataType, const FName InPotentialBase)
+	{
+		PRAGMA_DISABLE_EXPERIMENTAL_WARNINGS	
+		const FPolyRegistry& Reg = FPolyRegistry::Get();
+		for (FName i = InDataType; i != NAME_None; )
+		{
+			if (i == InPotentialBase)
+			{
+				return true;
+			}
+			const FPolyTypeInfo* Info = Reg.Find(i);
+			if (!Info || !Info->bIsPolymorphic)
+			{
+				return false;
+			}
+			i = Info->BaseTypeName;
+		}
+		return false;
+		PRAGMA_ENABLE_EXPERIMENTAL_WARNINGS
+	}
+
+	bool IsCastable(const FName InType, const FName InPotentialBase)
+	{
+		PRAGMA_DISABLE_EXPERIMENTAL_WARNINGS	
+		return IsChildOfByName(InType, InPotentialBase);
+		PRAGMA_ENABLE_EXPERIMENTAL_WARNINGS
+	}
+}
\ No newline at end of file
diff --git a/Engine/Plugins/Runtime/Metasound/Source/MetasoundGraphCore/Private/MetasoundVertexData.cpp b/Engine/Plugins/Runtime/Metasound/Source/MetasoundGraphCore/Private/MetasoundVertexData.cpp
--- a/Engine/Plugins/Runtime/Metasound/Source/MetasoundGraphCore/Private/MetasoundVertexData.cpp
+++ b/Engine/Plugins/Runtime/Metasound/Source/MetasoundGraphCore/Private/MetasoundVertexData.cpp
@@ -81,7 +81,7 @@ namespace Metasound
 
 		void FInputBinding::Set(FAnyDataReference&& InAnyDataReference)
 		{
-			check(Vertex.DataTypeName == InAnyDataReference.GetDataTypeName());
+			check(IsCastable(Vertex.DataTypeName,InAnyDataReference.GetDataTypeName()));
 			Data.Emplace(MoveTemp(InAnyDataReference));
 			CheckAccessTypeCompatibility(Vertex, *Data);
 		}
@@ -128,7 +128,7 @@ namespace Metasound
 
 		void FInputBinding::Bind(FInputBinding& InBinding)
 		{
-			check(Vertex.DataTypeName == InBinding.GetVertex().DataTypeName);
+			check(IsCastable(Vertex.DataTypeName, InBinding.GetVertex().DataTypeName));
 
 			if (Data.IsSet()) 
 			{ 
@@ -166,7 +166,7 @@ namespace Metasound
 
 		void FOutputBinding::Set(FAnyDataReference&& InAnyDataReference)
 		{
-			check(Vertex.DataTypeName == InAnyDataReference.GetDataTypeName());
+			check(IsCastable(Vertex.DataTypeName, InAnyDataReference.GetDataTypeName()));
 			Data.Emplace(MoveTemp(InAnyDataReference));
 			CheckAccessTypeCompatibility(Vertex, *Data);
 		}
@@ -251,7 +251,7 @@ namespace Metasound
 		{
 			if (BindingType* Binding = Find(InBindings, InVertexName))
 			{
-				if (Binding->GetVertex().DataTypeName == InDataReference.GetDataTypeName())
+				if (IsCastable(Binding->GetVertex().DataTypeName,InDataReference.GetDataTypeName()))
 				{
 					Binding->Set(MoveTemp(InDataReference));
 				}
@@ -286,7 +286,7 @@ namespace Metasound
 				if (BindingType* ThisBinding = Find(InThisBindings, OtherVertexName))
 				{
 					const FName& OtherDataTypeName= OtherBinding.GetVertex().DataTypeName;
-					if (OtherDataTypeName == ThisBinding->GetVertex().DataTypeName)
+					if (IsCastable(OtherDataTypeName,ThisBinding->GetVertex().DataTypeName))
 					{
 						ThisBinding->Bind(OtherBinding);
 					}
diff --git a/Engine/Plugins/Runtime/Metasound/Source/MetasoundGraphCore/Public/MetasoundDataReference.h b/Engine/Plugins/Runtime/Metasound/Source/MetasoundGraphCore/Public/MetasoundDataReference.h
--- a/Engine/Plugins/Runtime/Metasound/Source/MetasoundGraphCore/Public/MetasoundDataReference.h
+++ b/Engine/Plugins/Runtime/Metasound/Source/MetasoundGraphCore/Public/MetasoundDataReference.h
@@ -7,6 +7,8 @@
 #include "Containers/UnrealString.h"
 #include "HAL/Platform.h"
 #include "Internationalization/Text.h"
+#include "MetasoundDataTypeGetTraits.h"
+#include "MetasoundPolymorphic.h"
 #include "Templates/SharedPointer.h"
 #include "Templates/UniquePtr.h"
 #include "UObject/NameTypes.h"
@@ -36,78 +38,12 @@ namespace Metasound
 	// Unique ID type which corresponds to the underlying object referred to by a data reference.
 	using FDataReferenceID = const void*;
 
-	/** Helper class to enforce specialization of TDataReferenceTypeInfo */
-	template<typename DataType>
-	struct TSpecializationHelper 
-	{
-		enum { Value = false };
-	};
-
-	/** Info for templated data reference types help perform runtime type 
-	 * verification. 
-	 */
-	template<typename DataType>
-	struct TDataReferenceTypeInfo
-	{
-		// This static assert is triggered if TDataReferenceTypeInfo is used 
-		// without specialization.
-		static_assert(TSpecializationHelper<DataType>::Value, "TDataReferenceTypeInfo is not specialized.  Use macro DECLARE_METASOUND_DATA_REFERENCE_TYPES to declare a new type, or ensure that an existing DECLARE_METASOUND_DATA_REFERENCE_TYPES exists in the include path.");
-	};
-
-	/** Return the data type FName for a registered data type. */
-	template<typename DataType>
-	const FName& GetMetasoundDataTypeName()
-	{
-		static const FName TypeName = FName(TDataReferenceTypeInfo<std::decay_t<DataType>>::TypeName);
-
-		return TypeName;
-	}
-
-	/** Return the data type string for a registered data type. */
-	template<typename DataType>
-	const FString& GetMetasoundDataTypeString()
-	{
-		static const FString TypeString = FString(TDataReferenceTypeInfo<std::decay_t<DataType>>::TypeName);
-
-		return TypeString;
-	}
-
-	/** Return the display text for a registered data type. */
-	template<typename DataType>
-	const FText& GetMetasoundDataTypeDisplayText()
-	{
-		return TDataReferenceTypeInfo<std::decay_t<DataType>>::GetTypeDisplayText();
-	}
-
-	/** Return the data type ID for a registered data type.
-	 *
-	 * This ID is runtime constant but may change between executions and builds.
-	 */
-	template<typename DataType>
-	const void* const GetMetasoundDataTypeId()
-	{
-		return TDataReferenceTypeInfo<std::decay_t<DataType>>::GetTypeId();
-	}
-
 	/** Returns array type associated with the base datatype provided(ex. 'Float:Array' if 'Float' is provided) */
 	METASOUNDGRAPHCORE_API FName CreateArrayTypeNameFromElementTypeName(const FName InTypeName);
 
 	/** Returns the base data type with the array extension(ex. 'Float' if 'Float:Array' is provided) */
 	METASOUNDGRAPHCORE_API FName CreateElementTypeNameFromArrayTypeName(const FName InArrayTypeName);
 
-	/** Specialize void data type for internal use. */
-	template<>
-	struct TDataReferenceTypeInfo<void>
-	{
-		static METASOUNDGRAPHCORE_API const TCHAR* TypeName;
-		static METASOUNDGRAPHCORE_API const void* const TypeId;
-		static METASOUNDGRAPHCORE_API const FText& GetTypeDisplayText();
-
-		private:
-
-		static const void* const TypePtr;
-	};
-
 
 	/** A Data Reference Interface.
 	 *
@@ -138,22 +74,61 @@ namespace Metasound
 	{
 		return static_cast<FDataReferenceID>(InDataReference.GetRaw());
 	}
+
+	template <typename T>
+	UE_EXPERIMENTAL(5.7, "Polymorphic API is under development")
+	bool IsChildOf(const FName InBaseType)
+	{
+		static const FName ThisName = GetMetasoundDataTypeName<T>();
+		if (InBaseType == ThisName)
+		{
+			return true;
+		}
+
+		if constexpr (TPolymorphicTraits<T>::bIsPolymorphic)
+		{
+			using Base = typename TPolymorphicTraits<T>::BaseType;
+			static_assert(!std::is_same_v<Base, T>, "Polymorphic BaseType must not equal T");
+			return IsChildOf<Base>(InBaseType);
+		}
+		else
+		{
+			return false;
+		}
+	}
 	
-	/** Test if an IDataReference contains the same data type as the template
-	 * parameter.
-	 *
-	 * @return True if the IDataReference contains the DataType. False otherwise. 
-	 */
+	
+	// Polymorphic overload
 	template<typename DataType>
+	requires TPolymorphicTraits<DataType>::bIsPolymorphic
 	bool IsDataReferenceOfType(const IDataReference& InReference)
 	{
 		static const FName TypeName = GetMetasoundDataTypeName<DataType>();
 		static const void* const TypeId = GetMetasoundDataTypeId<DataType>();
 
-		bool bEqualTypeName = InReference.GetDataTypeName() == TypeName;
-		bool bEqualTypeId = InReference.GetDataTypeId() == TypeId;
+		const bool bEqualTypeName = InReference.GetDataTypeName() == TypeName;
+		const bool bEqualTypeId = InReference.GetDataTypeId() == TypeId;
+		if (bEqualTypeName && bEqualTypeId)
+		{
+			return true;
+		}
 
-		return (bEqualTypeName && bEqualTypeId);
+		// Recurse to base type
+		using FPolyType = TPolymorphicTraits<DataType>;
+		return IsDataReferenceOfType<typename FPolyType::BaseType>(InReference);
+	}
+
+	// Non-polymorphic fallback
+	template<typename DataType>
+	requires (!TPolymorphicTraits<DataType>::bIsPolymorphic)
+	bool IsDataReferenceOfType(const IDataReference& InReference)
+	{
+		static const FName TypeName = GetMetasoundDataTypeName<DataType>();
+		static const void* const TypeId = GetMetasoundDataTypeId<DataType>();
+
+		const bool bEqualTypeName = InReference.GetDataTypeName() == TypeName;
+		const bool bEqualTypeId = InReference.GetDataTypeId() == TypeId;
+		return bEqualTypeName && bEqualTypeId;
 	}
 
 	// This enum is used as a token to explicitly delineate when we should create a new object for the reference,
@@ -162,7 +137,52 @@ namespace Metasound
 	{
 		NewObject
 	};
+	
+	namespace MetasoundDataReferencePrivate
+	{
+		template <typename Base, typename Derived, typename... ArgTypes>
+		UE_EXPERIMENTAL(5.7, "Polymorphic API is under development")
+		TSharedRef<Base, ESPMode::NotThreadSafe> MakeRefOf(ArgTypes&&... args)
+		{
+			static_assert(std::is_base_of_v<Base, Derived>, "Target must derive from Base (or be Base)");
+			static_assert(std::is_constructible_v<Derived, ArgTypes...>, "Args do not match Target's constructor");
 
+			TSharedRef<Derived, ESPMode::NotThreadSafe> DerivedRef = MakeShared<Derived, ESPMode::NotThreadSafe>(Forward<ArgTypes>(args)...);
+ 
+			return StaticCastSharedRef<Base>(DerivedRef);
+		}
+	}
+
+	template <typename DataType>
+	struct TDynamicDataReferencePolicy;
+
+	template<typename DataType>
+	requires (TPolymorphicTraits<DataType>::bIsPolymorphic)
+	struct TDynamicDataReferencePolicy<DataType>
+	{
+		explicit TDynamicDataReferencePolicy(const FName& InDataTypeName = GetMetasoundDataTypeName<DataType>(), const void* InTypeId = GetMetasoundDataTypeId<DataType>())
+			: DataTypeName(InDataTypeName), DataTypeId(InTypeId)
+		{}
+
+		template<typename QueryType> bool IsA() const { return IsChildOfByName(DataTypeName, GetMetasoundDataTypeName<QueryType>()); }
+		const FName& GetDataTypeName() const { return DataTypeName; }
+		const void* GetTypeId() const { return DataTypeId; }
+		FName DataTypeName;
+		const void* DataTypeId;
+	};
+	
+	template<typename DataType>
+	requires (!TPolymorphicTraits<DataType>::bIsPolymorphic)
+	struct TDynamicDataReferencePolicy<DataType>
+	{
+		TDynamicDataReferencePolicy() = default;
+		TDynamicDataReferencePolicy(const FName&, const void*) {}
+
+		template<typename QueryType> bool IsA() const { return GetDataTypeName() == GetMetasoundDataTypeName<QueryType>(); }
+		const FName& GetDataTypeName() const { return GetMetasoundDataTypeName<DataType>(); }
+		const void* GetTypeId() const { return GetMetasoundDataTypeId<DataType>(); }
+	};
+	
 	/** Template class for a paramter reference. 
 	 *
 	 * This fulfills the IParamterRef interface, utilizing TDataReferenceTypeInfo to
@@ -184,6 +204,18 @@ namespace Metasound
 
 		typedef TSharedRef<DataType, ESPMode::NotThreadSafe> FRefType;
 
+		explicit TDataReference(EDataRefShouldConstruct /*InToken*/, FRefType InRef)
+			: ObjectReference(MoveTemp(InRef))
+		{}
+
+		explicit TDataReference(EDataRefShouldConstruct,
+								FRefType InRef,
+								const FName InTypeName,
+								const void* InTypeId)
+			: ObjectReference(MoveTemp(InRef))
+			, ObjectReferencePolicy(InTypeName, InTypeId)
+		{}
+	
 	public:
 
 		typedef TDataReferenceTypeInfo<DataType> FInfoType;
@@ -197,6 +229,22 @@ namespace Metasound
 			return TDataReference<DataType>(EDataRefShouldConstruct::NewObject, Forward<ArgTypes>(Args)...);
 		}
 
+		// Constructs Derived but stores as Base (DataType).
+		template <typename Derived, typename... ArgTypes>
+		UE_EXPERIMENTAL(5.7, "Polymorphic API is under development")
+		static TDataReference CreateNewDerived(ArgTypes&&... Args)
+		{
+			static_assert(std::is_base_of_v<DataType, Derived>, "Derived must inherit from DataType");
+			static_assert(std::is_constructible_v<Derived, ArgTypes...>, "Args don't match a Derived constructor");
+			
+			return TDataReference(
+				EDataRefShouldConstruct::NewObject,
+				MetasoundDataReferencePrivate::MakeRefOf<DataType, Derived>(Forward<ArgTypes>(Args)...),
+				GetMetasoundDataTypeName<Derived>(),
+			    GetMetasoundDataTypeId<Derived>()
+			);
+		}
+
 		/** Enable copy constructor */
 		TDataReference(const TDataReference<DataType>& Other) = default;
 
@@ -212,14 +260,13 @@ namespace Metasound
 		/** Return the name of the underlying type. */
 		virtual const FName& GetDataTypeName() const override
 		{
-			static const FName Name = GetMetasoundDataTypeName<DataType>();
-			return Name;
+			return ObjectReferencePolicy.GetDataTypeName();	
 		}
 
 		/** Return the ID of the underlying type. */
 		virtual const void* const GetDataTypeId() const override
 		{
-			return GetMetasoundDataTypeId<DataType>();
+			return ObjectReferencePolicy.GetTypeId();
 		}
 
 		/** Return a raw pointer to the data. */
@@ -228,11 +275,35 @@ namespace Metasound
 			return &ObjectReference.Get();
 		}
 
-	protected:
+		template<typename QueryType>
+		UE_EXPERIMENTAL(5.7, "Polymorphic API is under development")
+		bool IsA() const
+		{
+			return this->ObjectReferencePolicy.template IsA<QueryType>();
+		}
+		
+		template <typename CastType>
+		UE_EXPERIMENTAL(5.7, "Polymorphic API is under development")
+		TDataReference<CastType> TryCast() const
+		{
+			if (IsA<CastType>())
+			{
+				return TDataReference(
+					EDataRefShouldConstruct::NewObject,
+					MetasoundDataReferencePrivate::MakeRefOf<DataType, CastType>(ObjectReference),
+					GetMetasoundDataTypeName<CastType>(),
+					GetMetasoundDataTypeId<CastType>()
+				);
+			}
+			return {};
+		}
 
+	protected:
+		
 		// Protected object reference is utilized by subclasses which define what
 		// access is provided to the ObjectReference. 
 		FRefType ObjectReference;
+		TDynamicDataReferencePolicy<DataType> ObjectReferencePolicy;
 	};
 
 	// Forward declare
@@ -277,6 +348,19 @@ namespace Metasound
 			return TDataValueReference<DataType>(EDataRefShouldConstruct::NewObject, Forward<ArgTypes>(Args)...);
 		}
 
+		template <typename Derived, typename... ArgTypes>
+		UE_EXPERIMENTAL(5.7, "Polymorphic API is under development")
+		static TDataValueReference CreateNewDerived(ArgTypes&&... Args)
+		{
+			static_assert(std::is_base_of_v<DataType, Derived>, "Derived must inherit from DataType");
+			static_assert(std::is_constructible_v<Derived, ArgTypes...>, "Args don't match a Derived constructor");
+			
+			return TDataValueReference(
+				EDataRefShouldConstruct::NewObject,
+				MetasoundDataReferencePrivate::MakeRefOf<DataType, Derived>(Forward<ArgTypes>(Args)...)
+			);
+		}
+		
 		/** Enable copy constructor */
 		TDataValueReference(const TDataValueReference<DataType>& Other) = default;
 
@@ -379,6 +463,19 @@ namespace Metasound
 			return TDataWriteReference<DataType>(EDataRefShouldConstruct::NewObject, Forward<ArgTypes>(Args)...);
 		}
 
+		template <typename Derived, typename... ArgTypes>
+		UE_EXPERIMENTAL(5.7, "Polymorphic API is under development")
+		static TDataWriteReference CreateNewDerived(ArgTypes&&... Args)
+		{
+			static_assert(std::is_base_of_v<DataType, Derived>, "Derived must inherit from DataType");
+			static_assert(std::is_constructible_v<Derived, ArgTypes...>, "Args don't match a Derived constructor");
+			
+			return TDataWriteReference(
+				EDataRefShouldConstruct::NewObject,
+				MetasoundDataReferencePrivate::MakeRefOf<DataType, Derived>(Forward<ArgTypes>(Args)...)
+			);
+		}
+		
 		/** Enable copy constructor */
 		TDataWriteReference(const TDataWriteReference<DataType>& Other) = default;
 
@@ -465,6 +562,19 @@ namespace Metasound
 			return TDataReadReference<DataType>(EDataRefShouldConstruct::NewObject, Forward<ArgTypes>(Args)...);
 		}
 
+		template <typename Derived, typename... ArgTypes>
+		UE_EXPERIMENTAL(5.7, "Polymorphic API is under development")
+		static TDataReadReference CreateNewDerived(ArgTypes&&... Args)
+		{
+			static_assert(std::is_base_of_v<DataType, Derived>, "Derived must inherit from DataType");
+			static_assert(std::is_constructible_v<Derived, ArgTypes...>, "Args don't match a Derived constructor");
+			
+			return TDataReadReference(
+				EDataRefShouldConstruct::NewObject,
+				MetasoundDataReferencePrivate::MakeRefOf<DataType, Derived>(Forward<ArgTypes>(Args)...)
+			);
+		}
+		
 		TDataReadReference(const TDataReadReference& Other) = default;
 
 		/** Construct a readable parameter ref from a writable parameter ref. */
@@ -611,7 +721,9 @@ namespace Metasound
 		{
 			static TDataReadReference<DataType> Get(EDataReferenceAccessType InAccessType, const IDataReference& InRef)
 			{
-				check(IsDataReferenceOfType<DataType>(InRef));
+				PRAGMA_DISABLE_EXPERIMENTAL_WARNINGS
+				check(Metasound::IsCastable(InRef.GetDataTypeName(), GetMetasoundDataTypeName<DataType>()));
+				PRAGMA_ENABLE_EXPERIMENTAL_WARNINGS
 				switch (InAccessType)
 				{
 					case EDataReferenceAccessType::Read:
@@ -653,7 +765,9 @@ namespace Metasound
 		{
 			static TDataWriteReference<DataType> Get(EDataReferenceAccessType InAccessType, const IDataReference& InRef)
 			{
-				checkf(IsDataReferenceOfType<DataType>(InRef), TEXT("Attempt to get data reference with underlying type \"%s\" when actual underlying type is \"%s\"."), *GetMetasoundDataTypeString<DataType>(), *InRef.GetDataTypeName().ToString());
+				PRAGMA_DISABLE_EXPERIMENTAL_WARNINGS
+				checkf(Metasound::IsCastable(InRef.GetDataTypeName(), GetMetasoundDataTypeName<DataType>()), TEXT("Attempt to get data reference with underlying type \"%s\" when actual underlying type is \"%s\"."), *GetMetasoundDataTypeString<DataType>(), *InRef.GetDataTypeName().ToString());
+				PRAGMA_ENABLE_EXPERIMENTAL_WARNINGS
 				checkf(EDataReferenceAccessType::Write == InAccessType, TEXT("Invalid attempt to convert a data ref with \"%s\" access to \"%s\" access"), *LexToString(InAccessType), *LexToString(EDataReferenceAccessType::Write));
 				return *static_cast<const TDataWriteReference<DataType>*>(&InRef);
 			}
diff --git a/Engine/Plugins/Runtime/Metasound/Source/MetasoundGraphCore/Public/MetasoundDataTypeGetTraits.h b/Engine/Plugins/Runtime/Metasound/Source/MetasoundGraphCore/Public/MetasoundDataTypeGetTraits.h
new file mode 100644
--- /dev/null
+++ b/Engine/Plugins/Runtime/Metasound/Source/MetasoundGraphCore/Public/MetasoundDataTypeGetTraits.h
@@ -0,0 +1,67 @@
+// Copyright Epic Games, Inc. All Rights Reserved.
+
+#pragma once
+
+#include <type_traits>
+
+#include "Containers/UnrealString.h"
+#include "HAL/Platform.h"
+#include "Internationalization/Text.h"
+#include "Templates/SharedPointer.h"
+#include "Templates/UniquePtr.h"
+#include "UObject/NameTypes.h"
+
+namespace Metasound
+{
+	template <typename T>
+	struct TDataReferenceTypeInfo;
+	
+	/** Specialize void data type for internal use. */
+	template<>
+	struct TDataReferenceTypeInfo<void>
+	{
+		static METASOUNDGRAPHCORE_API const TCHAR* TypeName;
+		static METASOUNDGRAPHCORE_API const void* const TypeId;
+		static METASOUNDGRAPHCORE_API const FText& GetTypeDisplayText();
+		static METASOUNDGRAPHCORE_API const void* GetTypeId() { return TypeId; }
+
+	private:
+
+		static const void* const TypePtr;
+	};
+		
+	/** Return the data type FName for a registered data type. */
+	template<typename DataType>
+	const FName& GetMetasoundDataTypeName()
+	{
+		static const FName TypeName = FName(TDataReferenceTypeInfo<std::decay_t<DataType>>::TypeName);
+
+		return TypeName;
+	}
+
+	/** Return the data type string for a registered data type. */
+	template<typename DataType>
+	const FString& GetMetasoundDataTypeString()
+	{
+		static const FString TypeString = FString(TDataReferenceTypeInfo<std::decay_t<DataType>>::TypeName);
+
+		return TypeString;
+	}
+
+	/** Return the display text for a registered data type. */
+	template<typename DataType>
+	const FText& GetMetasoundDataTypeDisplayText()
+	{
+		return TDataReferenceTypeInfo<std::decay_t<DataType>>::GetTypeDisplayText();
+	}
+
+	/** Return the data type ID for a registered data type.
+	 *
+	 * This ID is runtime constant but may change between executions and builds.
+	 */
+	template<typename DataType>
+	const void* const GetMetasoundDataTypeId()
+	{
+		return TDataReferenceTypeInfo<std::decay_t<DataType>>::GetTypeId();
+	}
+}
\ No newline at end of file
diff --git a/Engine/Plugins/Runtime/Metasound/Source/MetasoundGraphCore/Public/MetasoundPolymorphic.h b/Engine/Plugins/Runtime/Metasound/Source/MetasoundGraphCore/Public/MetasoundPolymorphic.h
new file mode 100644
--- /dev/null
+++ b/Engine/Plugins/Runtime/Metasound/Source/MetasoundGraphCore/Public/MetasoundPolymorphic.h
@@ -0,0 +1,90 @@
+// Copyright Epic Games, Inc. All Rights Reserved.
+
+#pragma once
+
+#include "MetasoundDataTypeGetTraits.h"
+
+namespace Metasound
+{
+	template<typename T>
+	struct TPolymorphicTraits
+	{
+		static constexpr bool bIsPolymorphic = false;
+		static constexpr bool bIsAbstract = false;
+		using BaseType = void;
+	};
+
+	struct FPolyTypeInfo
+	{
+		FName TypeName = NAME_None;
+		FName BaseTypeName = NAME_None;
+		bool bIsPolymorphic = false;
+	};
+
+	class FPolyRegistry
+	{
+	public:
+		UE_EXPERIMENTAL(5.7, "Metasound Polymorphism is still under development")
+		METASOUNDGRAPHCORE_API static FPolyRegistry& Get();
+
+		UE_EXPERIMENTAL(5.7, "Metasound Polymorphism is still under development")
+		METASOUNDGRAPHCORE_API void Register(const FPolyTypeInfo& Info);
+
+		UE_EXPERIMENTAL(5.7, "Metasound Polymorphism is still under development")
+		METASOUNDGRAPHCORE_API const FPolyTypeInfo* Find(const FName Type) const;
+
+		UE_EXPERIMENTAL(5.7, "Metasound Polymorphism is still under development")
+		void Unregister(const FPolyTypeInfo& Info);
+
+		METASOUNDGRAPHCORE_API void GetAllRegisteredTypes(TArray<FPolyTypeInfo>& Out) const;
+
+	private:
+		TMap<FName, FPolyTypeInfo> Map;
+	};
+
+	template <typename T>
+	struct TPolyRegistrar
+	{
+		TPolyRegistrar()
+		{
+			static const FName ThisName = GetMetasoundDataTypeName<T>();
+			FPolyTypeInfo Info;
+			Info.TypeName = ThisName;
+			Info.bIsPolymorphic = TPolymorphicTraits<T>::bIsPolymorphic;
+
+			if constexpr (TPolymorphicTraits<T>::bIsPolymorphic)
+			{
+				using Base = TPolymorphicTraits<T>::BaseType;
+				static const FName BaseName = GetMetasoundDataTypeName<Base>();
+				static_assert(!std::is_same_v<Base, T>, "Polymorphic BaseType must not equal T");
+				Info.BaseTypeName = BaseName;
+			}
+			else
+			{
+				Info.BaseTypeName = NAME_None;
+			}
+			PRAGMA_DISABLE_INTERNAL_WARNINGS			
+			FPolyRegistry::Get().Register(Info);
+			PRAGMA_ENABLE_INTERNAL_WARNINGS			
+		}
+	};
+
+	UE_EXPERIMENTAL(5.7, "Metasound Polymorphism is still under development")
+	METASOUNDGRAPHCORE_API bool IsChildOfByName(const FName InDataType, const FName InPotentialBase);
+
+	UE_EXPERIMENTAL(5.7, "Metasound Polymorphism is still under development")
+	METASOUNDGRAPHCORE_API bool IsCastable(const FName InType, const FName InPotentialBase);
+	
+	
+}
+
+#define DECLARE_METASOUND_POLY_TYPE(TYPE, BASE, IS_ABSTRACT)\
+	template<>\
+	struct Metasound::TPolymorphicTraits<TYPE>\
+	{\
+		static constexpr bool bIsPolymorphic = true;\
+		static constexpr bool bIsAbstract = IS_ABSTRACT;\
+		using BaseType = BASE;\
+	};
+
+#define REGISTER_METASOUND_POLY_TYPE(T) static Metasound::TPolyRegistrar<T> UE_JOIN(GPolyReg,__COUNTER__);
\ No newline at end of file
diff --git a/Engine/Plugins/Runtime/Metasound/Source/MetasoundGraphCore/Public/MetasoundVertexData.h b/Engine/Plugins/Runtime/Metasound/Source/MetasoundGraphCore/Public/MetasoundVertexData.h
--- a/Engine/Plugins/Runtime/Metasound/Source/MetasoundGraphCore/Public/MetasoundVertexData.h
+++ b/Engine/Plugins/Runtime/Metasound/Source/MetasoundGraphCore/Public/MetasoundVertexData.h
@@ -12,6 +12,7 @@
 #include "MetasoundOperatorSettings.h"
 #include "MetasoundVertex.h"
 #include "Traits/ElementType.h"
+#include "MetasoundPolymorphic.h"
 
 #include <type_traits>
 
@@ -238,7 +239,10 @@ namespace Metasound
 			template<typename DataReferenceType>
 			void Set(const DataReferenceType& InDataReference)
 			{
-				check(Vertex.DataTypeName == InDataReference.GetDataTypeName());
+				PRAGMA_DISABLE_INTERNAL_WARNINGS
+				check(Metasound::IsCastable(Vertex.DataTypeName, InDataReference.GetDataTypeName()));
+				PRAGMA_ENABLE_INTERNAL_WARNINGS
+				
 				Data.Emplace(InDataReference);
 				CheckAccessTypeCompatibility(Vertex, *Data);
 			}
@@ -371,7 +375,9 @@ namespace Metasound
 			template<typename DataReferenceType>
 			void Set(const DataReferenceType& InDataReference)
 			{
-				check(Vertex.DataTypeName == InDataReference.GetDataTypeName());
+				PRAGMA_DISABLE_INTERNAL_WARNINGS
+				check(Metasound::IsCastable( Vertex.DataTypeName, InDataReference.GetDataTypeName()));
+				PRAGMA_ENABLE_INTERNAL_WARNINGS				
 				Data.Emplace(InDataReference);
 				CheckAccessTypeCompatibility(Vertex, *Data);
 			}
