- Fix for runtime data not being invalidated when changing preset state
- Minor clean up of delegate collection page code

#was-backout 46150595
#rnx

diff --git a/Engine/Plugins/Runtime/Metasound/Source/MetasoundEngine/Private/MetasoundBuilderBase.cpp b/Engine/Plugins/Runtime/Metasound/Source/MetasoundEngine/Private/MetasoundBuilderBase.cpp
--- a/Engine/Plugins/Runtime/Metasound/Source/MetasoundEngine/Private/MetasoundBuilderBase.cpp
+++ b/Engine/Plugins/Runtime/Metasound/Source/MetasoundEngine/Private/MetasoundBuilderBase.cpp
@@ -1004,16 +1004,6 @@ UObject* UMetaSoundBuilderBase::GetReferencedPresetAsset() const
 	return nullptr;
 }
 
-void UMetaSoundBuilderBase::InitFrontendBuilder()
-{
-PRAGMA_DISABLE_DEPRECATION_WARNINGS
-	UMetaSoundBuilderDocument& DocObject = UMetaSoundBuilderDocument::Create(GetBaseMetaSoundUClass());
-PRAGMA_ENABLE_DEPRECATION_WARNINGS
-
-	Builder = FMetaSoundFrontendDocumentBuilder(&DocObject);
-	Builder.InitDocument();
-}
-
 void UMetaSoundBuilderBase::Initialize()
 {
 	using namespace Metasound::Frontend;
diff --git a/Engine/Plugins/Runtime/Metasound/Source/MetasoundEngine/Private/MetasoundBuilderSubsystem.cpp b/Engine/Plugins/Runtime/Metasound/Source/MetasoundEngine/Private/MetasoundBuilderSubsystem.cpp
--- a/Engine/Plugins/Runtime/Metasound/Source/MetasoundEngine/Private/MetasoundBuilderSubsystem.cpp
+++ b/Engine/Plugins/Runtime/Metasound/Source/MetasoundEngine/Private/MetasoundBuilderSubsystem.cpp
@@ -394,16 +394,11 @@ const FMetasoundFrontendGraph& UMetaSoundSourceBuilder::GetConstTargetPageGraphC
 	return RootGraph.FindConstGraphChecked(TargetPageID);
 }
 
-const UMetaSoundSource& UMetaSoundSourceBuilder::GetMetaSoundSource() const
+UMetaSoundSource& UMetaSoundSourceBuilder::GetMetaSoundSource() const
 {
 	return GetConstBuilder().CastDocumentObjectChecked<UMetaSoundSource>();
 }
 
-UMetaSoundSource& UMetaSoundSourceBuilder::GetMetaSoundSource()
-{
-	return Builder.CastDocumentObjectChecked<UMetaSoundSource>();
-}
-
 void UMetaSoundSourceBuilder::InitDelegates(Metasound::Frontend::FDocumentModifyDelegates& OutDocumentDelegates)
 {
 	Super::InitDelegates(OutDocumentDelegates);
@@ -412,6 +407,7 @@ void UMetaSoundSourceBuilder::InitDelegates(Metasound::Frontend::FDocumentModify
 	{
 		OutDocumentDelegates.PageDelegates.OnPageAdded.AddUObject(this, &UMetaSoundSourceBuilder::OnPageAdded);
 		OutDocumentDelegates.PageDelegates.OnRemovingPage.AddUObject(this, &UMetaSoundSourceBuilder::OnRemovingPage);
+		OutDocumentDelegates.OnPresetStateChanged.AddUObject(this, &UMetaSoundSourceBuilder::OnPresetStateChanged);
 
 		OutDocumentDelegates.InterfaceDelegates.OnInputAdded.AddUObject(this, &UMetaSoundSourceBuilder::OnInputAdded);
 		OutDocumentDelegates.InterfaceDelegates.OnOutputAdded.AddUObject(this, &UMetaSoundSourceBuilder::OnOutputAdded);
@@ -444,19 +440,16 @@ void UMetaSoundSourceBuilder::InitTargetPageDelegates(Metasound::Frontend::FDocu
 		return bComponentsStopped;
 	});
 
-	OutDocumentDelegates.IterateGraphNodeDelegates([this](FNodeModifyDelegates& NodeDelegates)
-	{
-		NodeDelegates.OnNodeAdded.RemoveAll(this);
-		NodeDelegates.OnNodeConfigurationUpdated.RemoveAll(this);
-		NodeDelegates.OnNodeInputLiteralSet.RemoveAll(this);
-		NodeDelegates.OnRemoveSwappingNode.RemoveAll(this);
-		NodeDelegates.OnRemovingNodeInputLiteral.RemoveAll(this);
-	});
-
-	OutDocumentDelegates.IterateGraphEdgeDelegates([this](FEdgeModifyDelegates& EdgeDelegates)
+	OutDocumentDelegates.IterateGraphDelegates([this](FGraphModifyDelegates& GraphDelegates)
 	{
-		EdgeDelegates.OnEdgeAdded.RemoveAll(this);
-		EdgeDelegates.OnRemoveSwappingEdge.RemoveAll(this);
+		GraphDelegates.NodeDelegates.OnNodeAdded.RemoveAll(this);
+		GraphDelegates.NodeDelegates.OnNodeConfigurationUpdated.RemoveAll(this);
+		GraphDelegates.NodeDelegates.OnNodeInputLiteralSet.RemoveAll(this);
+		GraphDelegates.NodeDelegates.OnRemoveSwappingNode.RemoveAll(this);
+		GraphDelegates.NodeDelegates.OnRemovingNodeInputLiteral.RemoveAll(this);
+
+		GraphDelegates.EdgeDelegates.OnEdgeAdded.RemoveAll(this);
+		GraphDelegates.EdgeDelegates.OnRemoveSwappingEdge.RemoveAll(this);
 	});
 
 	const IMetaSoundDocumentInterface& DocInterface = GetConstBuilder().GetConstDocumentInterfaceChecked();
@@ -468,16 +461,16 @@ void UMetaSoundSourceBuilder::InitTargetPageDelegates(Metasound::Frontend::FDocu
 		TargetPageID = GraphPage->PageID;
 	}
 
-	FEdgeModifyDelegates& EdgeDelegates = OutDocumentDelegates.FindEdgeDelegatesChecked(TargetPageID);
-	EdgeDelegates.OnEdgeAdded.AddUObject(this, &UMetaSoundSourceBuilder::OnEdgeAdded);
-	EdgeDelegates.OnRemoveSwappingEdge.AddUObject(this, &UMetaSoundSourceBuilder::OnRemoveSwappingEdge);
+	FGraphModifyDelegates& GraphDelegates = OutDocumentDelegates.FindGraphDelegatesChecked(TargetPageID);
 
-	FNodeModifyDelegates& NodeDelegates = OutDocumentDelegates.FindNodeDelegatesChecked(TargetPageID);
-	NodeDelegates.OnNodeAdded.AddUObject(this, &UMetaSoundSourceBuilder::OnNodeAdded);
-	NodeDelegates.OnNodeConfigurationUpdated.AddUObject(this, &UMetaSoundSourceBuilder::OnNodeConfigurationUpdated);
-	NodeDelegates.OnNodeInputLiteralSet.AddUObject(this, &UMetaSoundSourceBuilder::OnNodeInputLiteralSet);
-	NodeDelegates.OnRemoveSwappingNode.AddUObject(this, &UMetaSoundSourceBuilder::OnRemoveSwappingNode);
-	NodeDelegates.OnRemovingNodeInputLiteral.AddUObject(this, &UMetaSoundSourceBuilder::OnRemovingNodeInputLiteral);
+	GraphDelegates.EdgeDelegates.OnEdgeAdded.AddUObject(this, &UMetaSoundSourceBuilder::OnEdgeAdded);
+	GraphDelegates.EdgeDelegates.OnRemoveSwappingEdge.AddUObject(this, &UMetaSoundSourceBuilder::OnRemoveSwappingEdge);
+
+	GraphDelegates.NodeDelegates.OnNodeAdded.AddUObject(this, &UMetaSoundSourceBuilder::OnNodeAdded);
+	GraphDelegates.NodeDelegates.OnNodeConfigurationUpdated.AddUObject(this, &UMetaSoundSourceBuilder::OnNodeConfigurationUpdated);
+	GraphDelegates.NodeDelegates.OnNodeInputLiteralSet.AddUObject(this, &UMetaSoundSourceBuilder::OnNodeInputLiteralSet);
+	GraphDelegates.NodeDelegates.OnRemoveSwappingNode.AddUObject(this, &UMetaSoundSourceBuilder::OnRemoveSwappingNode);
+	GraphDelegates.NodeDelegates.OnRemovingNodeInputLiteral.AddUObject(this, &UMetaSoundSourceBuilder::OnRemovingNodeInputLiteral);
 }
 
 void UMetaSoundSourceBuilder::OnAssetReferenceAdded(TScriptInterface<IMetaSoundDocumentInterface> DocInterface)
@@ -733,7 +726,7 @@ void UMetaSoundSourceBuilder::OnOutputAdded(int32 OutputIndex) const
 	});
 }
 
-void UMetaSoundSourceBuilder::OnPageAdded(const Metasound::Frontend::FDocumentMutatePageArgs& Args)
+void UMetaSoundSourceBuilder::OnPageAdded(const Metasound::Frontend::FDocumentMutatePageArgs& InArgs)
 {
 	using namespace Metasound::Frontend;
 
@@ -741,7 +734,14 @@ void UMetaSoundSourceBuilder::OnPageAdded(const Metasound::Frontend::FDocumentMu
 	InitTargetPageDelegates(DocDelegates);
 }
 
-void UMetaSoundSourceBuilder::OnRemovingPage(const Metasound::Frontend::FDocumentMutatePageArgs& Args)
+void UMetaSoundSourceBuilder::OnPresetStateChanged(const Metasound::Frontend::FDocumentPresetStateChangedArgs& InArgs) const
+{
+	using namespace Metasound::Frontend;
+
+	GetMetaSoundSource().InvalidateCachedRuntimeInputData();
+}
+
+void UMetaSoundSourceBuilder::OnRemovingPage(const Metasound::Frontend::FDocumentMutatePageArgs& InArgs)
 {
 	using namespace Metasound::Frontend;
 
diff --git a/Engine/Plugins/Runtime/Metasound/Source/MetasoundEngine/Public/MetasoundBuilderBase.h b/Engine/Plugins/Runtime/Metasound/Source/MetasoundEngine/Public/MetasoundBuilderBase.h
--- a/Engine/Plugins/Runtime/Metasound/Source/MetasoundEngine/Public/MetasoundBuilderBase.h
+++ b/Engine/Plugins/Runtime/Metasound/Source/MetasoundEngine/Public/MetasoundBuilderBase.h
@@ -543,9 +543,6 @@ public:
 
 	UE_API Metasound::Frontend::FDocumentModifyDelegates& GetBuilderDelegates();
 
-	UE_DEPRECATED(5.4, "Moved to Initialize and should only be called by internal implementation (i.e. IDocumentBuilderRegistry'")
-	UE_API virtual void InitFrontendBuilder();
-
 	// Initializes and ensures all nodes have a position (required prior to exporting to an asset if expected to be viewed in the editor).
 	UE_API virtual void InitNodeLocations();
 
@@ -574,6 +571,7 @@ public:
 
 	// Resets FrontendBuilder instance, creating a transient MetaSound document that is managed by this UObject Builder.
 	UE_API void Initialize();
+
 	// Initializes builder given an existing document interface 
 	UE_API void Initialize(TScriptInterface<IMetaSoundDocumentInterface>& DocInterface);
 
diff --git a/Engine/Plugins/Runtime/Metasound/Source/MetasoundEngine/Public/MetasoundBuilderSubsystem.h b/Engine/Plugins/Runtime/Metasound/Source/MetasoundEngine/Public/MetasoundBuilderSubsystem.h
--- a/Engine/Plugins/Runtime/Metasound/Source/MetasoundEngine/Public/MetasoundBuilderSubsystem.h
+++ b/Engine/Plugins/Runtime/Metasound/Source/MetasoundEngine/Public/MetasoundBuilderSubsystem.h
@@ -118,8 +118,7 @@ protected:
 private:
 	UE_API const FMetasoundFrontendGraph& GetConstTargetPageGraphChecked() const;
 
-	UE_API const UMetaSoundSource& GetMetaSoundSource() const;
-	UE_API UMetaSoundSource& GetMetaSoundSource();
+	UE_API UMetaSoundSource& GetMetaSoundSource() const;
 
 	UE_API void InitTargetPageDelegates(Metasound::Frontend::FDocumentModifyDelegates& OutDocumentDelegates);
 
@@ -131,6 +130,7 @@ private:
 	UE_API void OnNodeInputLiteralSet(int32 NodeIndex, int32 VertexIndex, int32 LiteralIndex) const;
 	UE_API void OnOutputAdded(int32 OutputIndex) const;
 	UE_API void OnPageAdded(const Metasound::Frontend::FDocumentMutatePageArgs& InArgs);
+	UE_API void OnPresetStateChanged(const Metasound::Frontend::FDocumentPresetStateChangedArgs& InArgs) const;
 	UE_API void OnRemoveSwappingEdge(int32 SwapIndex, int32 LastIndex) const;
 	UE_API void OnRemovingInput(int32 InputIndex);
 	UE_API void OnRemoveSwappingNode(int32 SwapIndex, int32 LastIndex) const;
diff --git a/Engine/Plugins/Runtime/Metasound/Source/MetasoundFrontend/Private/MetasoundFrontendDocumentBuilder.cpp b/Engine/Plugins/Runtime/Metasound/Source/MetasoundFrontend/Private/MetasoundFrontendDocumentBuilder.cpp
--- a/Engine/Plugins/Runtime/Metasound/Source/MetasoundFrontend/Private/MetasoundFrontendDocumentBuilder.cpp
+++ b/Engine/Plugins/Runtime/Metasound/Source/MetasoundFrontend/Private/MetasoundFrontendDocumentBuilder.cpp
@@ -881,7 +881,7 @@ void FMetaSoundFrontendDocumentBuilder::AddEdge(FMetasoundFrontendEdge&& InNewEd
 	FMetasoundFrontendGraph& Graph = Document.RootGraph.FindGraphChecked(PageID);
 	Graph.Edges.Add(MoveTemp(InNewEdge));
 	const int32 NewIndex = Graph.Edges.Num() - 1;
-	DocumentDelegates->FindEdgeDelegatesChecked(PageID).OnEdgeAdded.Broadcast(NewIndex);
+	DocumentDelegates->FindGraphDelegatesChecked(PageID).EdgeDelegates.OnEdgeAdded.Broadcast(NewIndex);
 }
 
 bool FMetaSoundFrontendDocumentBuilder::AddNamedEdges(const TSet<Metasound::Frontend::FNamedEdge>& EdgesToMake, TArray<const FMetasoundFrontendEdge*>* OutNewEdges, bool bReplaceExistingConnections, const FGuid* InPageID)
@@ -1956,7 +1956,7 @@ FMetasoundFrontendNode* FMetaSoundFrontendDocumentBuilder::AddNodeInternal(const
 		// Update document builder cache
 		const int32 NewIndex = Nodes.Num() - 1;
 		const IDocumentGraphNodeCache& NodeCache = DocumentCache->GetNodeCache(InPageID);
-		DocumentDelegates->FindNodeDelegatesChecked(InPageID).OnNodeAdded.Broadcast(NewIndex);
+		DocumentDelegates->FindGraphDelegatesChecked(InPageID).NodeDelegates.OnNodeAdded.Broadcast(NewIndex);
 
 		if (NewNodeIndex)
 		{
@@ -2339,7 +2339,7 @@ bool FMetaSoundFrontendDocumentBuilder::ConformGraphInputNodeToClass(const FMeta
 				{
 					TArray<FMetasoundFrontendNode>& Nodes = Graph.Nodes;
 					FMetasoundFrontendNode& Node = Nodes[*NodeIndexPtr];
-					FNodeModifyDelegates& NodeDelegates = DocumentDelegates->FindNodeDelegatesChecked(Graph.PageID);
+					FNodeModifyDelegates& NodeDelegates = DocumentDelegates->FindGraphDelegatesChecked(Graph.PageID).NodeDelegates;
 					const int32 RemovalIndex = *NodeIndexPtr; // Have to cache as next delegate broadcast invalidates index pointer
 					NodeDelegates.OnRemoveSwappingNode.Broadcast(RemovalIndex, Nodes.Num() - 1);
 					FMetasoundFrontendNode NewNode = MoveTemp(Node);
@@ -2397,7 +2397,7 @@ bool FMetaSoundFrontendDocumentBuilder::ConformGraphOutputNodeToClass(const FMet
 				{
 					TArray<FMetasoundFrontendNode>& Nodes = Graph.Nodes;
 					FMetasoundFrontendNode& Node = Nodes[*NodeIndexPtr];
-					FNodeModifyDelegates& NodeDelegates = DocumentDelegates->FindNodeDelegatesChecked(Graph.PageID);
+					FNodeModifyDelegates& NodeDelegates = DocumentDelegates->FindGraphDelegatesChecked(Graph.PageID).NodeDelegates;
 					const int32 RemovalIndex = *NodeIndexPtr; // Have to cache as next delegate broadcast invalidates index pointer
 					NodeDelegates.OnRemoveSwappingNode.Broadcast(RemovalIndex, Nodes.Num() - 1);
 					FMetasoundFrontendNode NewNode = MoveTemp(Node);
@@ -2454,6 +2454,7 @@ bool FMetaSoundFrontendDocumentBuilder::ConvertFromPreset()
 		Style.bIsGraphEditable = true;
 #endif // WITH_EDITOR
 
+		DocumentDelegates->OnPresetStateChanged.Broadcast(FDocumentPresetStateChangedArgs { });
 		return true;
 	}
 
@@ -2487,6 +2488,7 @@ bool FMetaSoundFrontendDocumentBuilder::ConvertToPreset(const FMetaSoundFrontend
 	if (RebuildPresetRootGraph.Transform(*this))
 	{
 		DocumentInterface->ConformObjectToDocument();
+		DocumentDelegates->OnPresetStateChanged.Broadcast(FDocumentPresetStateChangedArgs { });
 		return true;
 	}
 
@@ -4141,7 +4143,9 @@ void FMetaSoundFrontendDocumentBuilder::InitDocument(const FMetasoundFrontendDoc
 		if (Document.RootGraph.GetConstGraphPages().IsEmpty())
 		{
 			Document.RootGraph.InitDefaultGraphPage();
+			DocumentDelegates->AddPageDelegates(Frontend::DefaultPageID);
 		}
+
 		InitGraphClassMetadata(bResetVersion, InNewClassName);
 	}
 	else
@@ -4150,7 +4154,9 @@ void FMetaSoundFrontendDocumentBuilder::InitDocument(const FMetasoundFrontendDoc
 		if (Document.RootGraph.GetConstGraphPages().IsEmpty())
 		{
 			Document.RootGraph.InitDefaultGraphPage();
+			DocumentDelegates->AddPageDelegates(Frontend::DefaultPageID);
 		}
+
 		FMetasoundFrontendClassMetadata& ClassMetadata = Document.RootGraph.Metadata;
 		InitGraphClassMetadata(Document.RootGraph.Metadata, bResetVersion, InNewClassName);
 
@@ -4805,7 +4811,7 @@ bool FMetaSoundFrontendDocumentBuilder::RemoveEdge(const FMetasoundFrontendEdge&
 		if (EdgeToRemove.FromNodeID == FoundEdge.FromNodeID && EdgeToRemove.FromVertexID == FoundEdge.FromVertexID)
 		{
 			const int32 LastIndex = Edges.Num() - 1;
-			DocumentDelegates->FindEdgeDelegatesChecked(PageID).OnRemoveSwappingEdge.Broadcast(Index, LastIndex);
+			DocumentDelegates->FindGraphDelegatesChecked(PageID).EdgeDelegates.OnRemoveSwappingEdge.Broadcast(Index, LastIndex);
 			Edges.RemoveAtSwap(Index, EAllowShrinking::No);
 			return true;
 		}
@@ -4956,7 +4962,7 @@ bool FMetaSoundFrontendDocumentBuilder::RemoveNodeInputDefault(const FGuid& InNo
 			const int32 LiteralIndex = Node.InputLiterals.IndexOfByPredicate(IsLiteral);
 			if (LiteralIndex != INDEX_NONE)
 			{
-				FNodeModifyDelegates& NodeDelegates = DocumentDelegates->FindNodeDelegatesChecked(PageID);
+				FNodeModifyDelegates& NodeDelegates = DocumentDelegates->FindGraphDelegatesChecked(PageID).NodeDelegates;
 				const FOnMetaSoundFrontendDocumentMutateNodeInputLiteralArray& OnRemovingNodeInputLiteral = NodeDelegates.OnRemovingNodeInputLiteral;
 				const int32 LastIndex = Node.InputLiterals.Num() - 1;
 				OnRemovingNodeInputLiteral.Broadcast(*NodeIndex, VertexIndex, LastIndex);
@@ -5049,7 +5055,7 @@ bool FMetaSoundFrontendDocumentBuilder::RemoveEdgesFromNodeOutput(const FGuid& I
 		// Copy off indices and sort descending as the edge array will be modified when notifying the cache in the loop below
 		TArray<int32> IndicesCopy(Indices.GetData(), Indices.Num());
 		Algo::Sort(IndicesCopy, [](const int32& L, const int32& R) { return L > R; });
-		FEdgeModifyDelegates& EdgeDelegates = DocumentDelegates->FindEdgeDelegatesChecked(PageID);
+		FEdgeModifyDelegates& EdgeDelegates = DocumentDelegates->FindGraphDelegatesChecked(PageID).EdgeDelegates;
 		for (int32 Index : IndicesCopy)
 		{
 #if WITH_EDITORONLY_DATA
@@ -5100,7 +5106,7 @@ bool FMetaSoundFrontendDocumentBuilder::RemoveEdgeToNodeInput(const FGuid& InNod
 		}
 #endif // WITH_EDITORONLY_DATA
 
-		const FEdgeModifyDelegates& EdgeDelegates = DocumentDelegates->FindEdgeDelegatesChecked(PageID);
+		const FEdgeModifyDelegates& EdgeDelegates = DocumentDelegates->FindGraphDelegatesChecked(PageID).EdgeDelegates;
 		const int32 LastIndex = Graph.Edges.Num() - 1;
 		EdgeDelegates.OnRemoveSwappingEdge.Broadcast(Index, LastIndex);
 		Graph.Edges.RemoveAtSwap(Index, EAllowShrinking::No);
@@ -5416,7 +5422,7 @@ bool FMetaSoundFrontendDocumentBuilder::RemoveNode(const FGuid& InNodeID, const
 
 		RemoveEdges(NodeID, InPageID);
 		const int32 LastIndex = Nodes.Num() - 1;
-		FNodeModifyDelegates& NodeDelegates = DocumentDelegates->FindNodeDelegatesChecked(PageID);
+		FNodeModifyDelegates& NodeDelegates = DocumentDelegates->FindGraphDelegatesChecked(PageID).NodeDelegates;
 		NodeDelegates.OnRemoveSwappingNode.Broadcast(Index, LastIndex);
 		Nodes.RemoveAtSwap(Index, EAllowShrinking::No);
 
@@ -6700,7 +6706,7 @@ bool FMetaSoundFrontendDocumentBuilder::SetNodeInputDefault(const FGuid& InNodeI
 				Node.InputLiterals[LiteralIndex] = MoveTemp(NewVertexLiteral);
 			}
 
-			FNodeModifyDelegates& NodeDelegates = DocumentDelegates->FindNodeDelegatesChecked(PageID);
+			FNodeModifyDelegates& NodeDelegates = DocumentDelegates->FindGraphDelegatesChecked(PageID).NodeDelegates;
 			const FOnMetaSoundFrontendDocumentMutateNodeInputLiteralArray& OnNodeInputLiteralSet = NodeDelegates.OnNodeInputLiteralSet;
 
 #if WITH_EDITORONLY_DATA
@@ -7143,7 +7149,7 @@ bool FMetaSoundFrontendDocumentBuilder::UpdateNodeInterfaceFromConfiguration(con
 	};
 
 	const bool bInterfaceUpdated = Node.Interface.Update(ClassInterface, DisconnectInput, DisconnectOutput);
-	DocumentDelegates->FindNodeDelegatesChecked(PageID).OnNodeConfigurationUpdated.Broadcast(*NodeIndex);
+	DocumentDelegates->FindGraphDelegatesChecked(PageID).NodeDelegates.OnNodeConfigurationUpdated.Broadcast(*NodeIndex);
 
 #if WITH_EDITORONLY_DATA
 	if (bInterfaceUpdated)
diff --git a/Engine/Plugins/Runtime/Metasound/Source/MetasoundFrontend/Private/MetasoundFrontendDocumentCache.cpp b/Engine/Plugins/Runtime/Metasound/Source/MetasoundFrontend/Private/MetasoundFrontendDocumentCache.cpp
--- a/Engine/Plugins/Runtime/Metasound/Source/MetasoundFrontend/Private/MetasoundFrontendDocumentCache.cpp
+++ b/Engine/Plugins/Runtime/Metasound/Source/MetasoundFrontend/Private/MetasoundFrontendDocumentCache.cpp
@@ -605,7 +605,7 @@ namespace Metasound::Frontend
 
 		// Factory style constructor as restrictions on construction of shared pointers disallows passing of this document cache's pointer to sub-caches
 		TSharedRef<FDocumentGraphNodeCache> Cache = MakeShared<FDocumentGraphNodeCache>(ParentCache, InPageID);
-		Cache->Init(OutDelegates.FindNodeDelegatesChecked(InPageID));
+		Cache->Init(OutDelegates.FindGraphDelegatesChecked(InPageID).NodeDelegates);
 		return Cache;
 	}
 
@@ -1031,7 +1031,7 @@ namespace Metasound::Frontend
 
 		// Factory style constructor as restrictions on construction of shared pointers disallows passing of this document cache's pointer to sub-caches
 		TSharedRef<FDocumentGraphEdgeCache> Cache = MakeShared<FDocumentGraphEdgeCache>(ParentCache, InPageID);
-		Cache->Init(OutDelegates.FindEdgeDelegatesChecked(InPageID));
+		Cache->Init(OutDelegates.FindGraphDelegatesChecked(InPageID).EdgeDelegates);
 		return Cache;
 	}
 
diff --git a/Engine/Plugins/Runtime/Metasound/Source/MetasoundFrontend/Private/MetasoundFrontendDocumentModifyDelegates.cpp b/Engine/Plugins/Runtime/Metasound/Source/MetasoundFrontend/Private/MetasoundFrontendDocumentModifyDelegates.cpp
--- a/Engine/Plugins/Runtime/Metasound/Source/MetasoundFrontend/Private/MetasoundFrontendDocumentModifyDelegates.cpp
+++ b/Engine/Plugins/Runtime/Metasound/Source/MetasoundFrontend/Private/MetasoundFrontendDocumentModifyDelegates.cpp
@@ -2,129 +2,95 @@
 
 #include "MetasoundFrontendDocumentModifyDelegates.h"
 
+#if WITH_EDITORONLY_DATA
+#include "MetasoundFrontendDocumentVersioning.h"
+#endif // WITH_EDITORONLY_DATA
 
 namespace Metasound::Frontend
 {
 	FDocumentModifyDelegates::FDocumentModifyDelegates()
-PRAGMA_DISABLE_DEPRECATION_WARNINGS
-		: NodeDelegates()
-		, EdgeDelegates()
-PRAGMA_ENABLE_DEPRECATION_WARNINGS
 	{
 	}
 
-	FDocumentModifyDelegates::FDocumentModifyDelegates(const FDocumentModifyDelegates& InModifyDelegates)
-		: OnDocumentMetadataChanged(InModifyDelegates.OnDocumentMetadataChanged)
-		, OnDependencyAdded(InModifyDelegates.OnDependencyAdded)
-		, OnRemoveSwappingDependency(InModifyDelegates.OnRemoveSwappingDependency)
-		, OnRenamingDependencyClass(InModifyDelegates.OnRenamingDependencyClass)
-		, PageDelegates(InModifyDelegates.PageDelegates)
-		, InterfaceDelegates(InModifyDelegates.InterfaceDelegates)
-		PRAGMA_DISABLE_DEPRECATION_WARNINGS
-		, NodeDelegates(InModifyDelegates.NodeDelegates)
-		, EdgeDelegates(InModifyDelegates.EdgeDelegates)
-		PRAGMA_ENABLE_DEPRECATION_WARNINGS
-		, PageNodeDelegates(InModifyDelegates.PageNodeDelegates)
-		, PageEdgeDelegates(InModifyDelegates.PageEdgeDelegates)
+	FDocumentModifyDelegates::FDocumentModifyDelegates(const FMetasoundFrontendDocument& Document)
+		: FDocumentModifyDelegates()
 	{
+		Document.RootGraph.IterateGraphPages([this](const FMetasoundFrontendGraph& Graph)
+		{
+			AddPageDelegates(Graph.PageID);
+		});
+
+#if WITH_EDITORONLY_DATA
+		// Modify delegates may register when a builder is constructed prior to a document
+		// being versioned. In this legacy case, add the default page ID delegates as the
+		// default page is the only version that exists on legacy documents.
+		if (Document.Metadata.Version.Number < Frontend::GetPageMigrationVersion())
+		{
+			if (Document.RootGraph.GetConstGraphPages().IsEmpty())
+			{
+				AddPageDelegates(Frontend::DefaultPageID);
+			}
+		}
+#endif // WITH_EDITORONLY_DATA
 	}
 
-	FDocumentModifyDelegates::FDocumentModifyDelegates(FDocumentModifyDelegates&& InModifyDelegates)
-		: OnDocumentMetadataChanged(MoveTemp(InModifyDelegates.OnDocumentMetadataChanged))
-		, OnDependencyAdded(MoveTemp(InModifyDelegates.OnDependencyAdded))
-		, OnRemoveSwappingDependency(MoveTemp(InModifyDelegates.OnRemoveSwappingDependency))
-		, OnRenamingDependencyClass(MoveTemp(InModifyDelegates.OnRenamingDependencyClass))
-		, PageDelegates(MoveTemp(InModifyDelegates.PageDelegates))
-		, InterfaceDelegates(MoveTemp(InModifyDelegates.InterfaceDelegates))
-PRAGMA_DISABLE_DEPRECATION_WARNINGS
-		, NodeDelegates(MoveTemp(InModifyDelegates.NodeDelegates))
-		, EdgeDelegates(MoveTemp(InModifyDelegates.EdgeDelegates))
-PRAGMA_ENABLE_DEPRECATION_WARNINGS
-		, PageNodeDelegates(MoveTemp(InModifyDelegates.PageNodeDelegates))
-		, PageEdgeDelegates(MoveTemp(InModifyDelegates.PageEdgeDelegates))
+	void FDocumentModifyDelegates::AddPageDelegates(const FGuid& InPageID)
 	{
+		if (!GraphDelegates.Contains(InPageID))
+		{
+			GraphDelegates.Add(InPageID, FGraphModifyDelegates());
+			PageDelegates.OnPageAdded.Broadcast(FDocumentMutatePageArgs{ InPageID });
+		}
 	}
 
-	FDocumentModifyDelegates& FDocumentModifyDelegates::operator=(const FDocumentModifyDelegates& InModifyDelegates)
+	FNodeModifyDelegates& FDocumentModifyDelegates::FindNodeDelegatesChecked(const FGuid& InPageID)
 	{
-		OnDocumentMetadataChanged = InModifyDelegates.OnDocumentMetadataChanged;
-		OnDependencyAdded = InModifyDelegates.OnDependencyAdded;
-		OnRemoveSwappingDependency = InModifyDelegates.OnRemoveSwappingDependency;
-		OnRenamingDependencyClass = InModifyDelegates.OnRenamingDependencyClass;
-		PageDelegates = InModifyDelegates.PageDelegates;
-		InterfaceDelegates = InModifyDelegates.InterfaceDelegates;
-PRAGMA_DISABLE_DEPRECATION_WARNINGS
-		NodeDelegates = InModifyDelegates.NodeDelegates;
-		EdgeDelegates = InModifyDelegates.EdgeDelegates;
-PRAGMA_ENABLE_DEPRECATION_WARNINGS
-		PageNodeDelegates = InModifyDelegates.PageNodeDelegates;
-		PageEdgeDelegates = InModifyDelegates.PageEdgeDelegates;
-		return *this;
+		return GraphDelegates.FindChecked(InPageID).NodeDelegates;
 	}
 
-	FDocumentModifyDelegates& FDocumentModifyDelegates::operator=(FDocumentModifyDelegates&& InModifyDelegates)
+	FEdgeModifyDelegates& FDocumentModifyDelegates::FindEdgeDelegatesChecked(const FGuid& InPageID)
 	{
-		OnDocumentMetadataChanged = MoveTemp(InModifyDelegates.OnDocumentMetadataChanged);
-		OnDependencyAdded = MoveTemp(InModifyDelegates.OnDependencyAdded);
-		OnRemoveSwappingDependency = MoveTemp(InModifyDelegates.OnRemoveSwappingDependency);
-		OnRenamingDependencyClass = MoveTemp(InModifyDelegates.OnRenamingDependencyClass);
-		PageDelegates = MoveTemp(InModifyDelegates.PageDelegates);
-		InterfaceDelegates = MoveTemp(InModifyDelegates.InterfaceDelegates);
-PRAGMA_DISABLE_DEPRECATION_WARNINGS
-		NodeDelegates = MoveTemp(InModifyDelegates.NodeDelegates);
-		EdgeDelegates = MoveTemp(InModifyDelegates.EdgeDelegates);
-PRAGMA_ENABLE_DEPRECATION_WARNINGS
-		PageNodeDelegates = MoveTemp(InModifyDelegates.PageNodeDelegates);
-		PageEdgeDelegates = MoveTemp(InModifyDelegates.PageEdgeDelegates);
-		return *this;
+		return GraphDelegates.FindChecked(InPageID).EdgeDelegates;
 	}
 
-	FDocumentModifyDelegates::FDocumentModifyDelegates(const FMetasoundFrontendDocument& Document)
-PRAGMA_DISABLE_DEPRECATION_WARNINGS
-		: NodeDelegates()
-		, EdgeDelegates()
-PRAGMA_ENABLE_DEPRECATION_WARNINGS
+	FGraphModifyDelegates& FDocumentModifyDelegates::FindGraphDelegatesChecked(const FGuid& InPageID)
 	{
-		Document.RootGraph.IterateGraphPages([this](const FMetasoundFrontendGraph& Graph)
-		{
-			AddPageDelegates(Graph.PageID);
-		});
+		return GraphDelegates.FindChecked(InPageID);
 	}
 
-	void FDocumentModifyDelegates::AddPageDelegates(const FGuid& InPageID)
+	void FDocumentModifyDelegates::IterateGraphEdgeDelegates(TFunctionRef<void(FEdgeModifyDelegates&)> Func)
 	{
-		PageNodeDelegates.Add(InPageID, FNodeModifyDelegates());
-		PageEdgeDelegates.Add(InPageID, FEdgeModifyDelegates());
-
-		PageDelegates.OnPageAdded.Broadcast(FDocumentMutatePageArgs{ InPageID });
+		IterateGraphDelegates([&](FGraphModifyDelegates& GraphDelegates)
+		{
+			Func(GraphDelegates.EdgeDelegates);
+		});
 	}
 
-	FNodeModifyDelegates& FDocumentModifyDelegates::FindNodeDelegatesChecked(const FGuid& InPageID)
+	void FDocumentModifyDelegates::IterateGraphNodeDelegates(TFunctionRef<void(FNodeModifyDelegates&)> Func)
 	{
-PRAGMA_DISABLE_DEPRECATION_WARNINGS
-		return InPageID == DefaultPageID
-		? NodeDelegates
-		: PageNodeDelegates.FindChecked(InPageID);
-PRAGMA_ENABLE_DEPRECATION_WARNINGS
+		IterateGraphDelegates([&](FGraphModifyDelegates& GraphDelegates)
+		{
+			Func(GraphDelegates.NodeDelegates);
+		});
 	}
 
-	FEdgeModifyDelegates& FDocumentModifyDelegates::FindEdgeDelegatesChecked(const FGuid& InPageID)
+	void FDocumentModifyDelegates::IterateGraphDelegates(TFunctionRef<void(FGraphModifyDelegates&)> Func)
 	{
-PRAGMA_DISABLE_DEPRECATION_WARNINGS
-		return InPageID == DefaultPageID
-		? EdgeDelegates
-		: PageEdgeDelegates.FindChecked(InPageID);
-PRAGMA_ENABLE_DEPRECATION_WARNINGS
+		for (TPair<FGuid, FGraphModifyDelegates>& Pair : GraphDelegates)
+		{
+			Func(Pair.Value);
+		}
 	}
 
 	void FDocumentModifyDelegates::RemovePageDelegates(const FGuid& InPageID, bool bBroadcastNotify)
 	{
-		if (bBroadcastNotify)
+		if (GraphDelegates.Contains(InPageID))
 		{
-			PageDelegates.OnRemovingPage.Broadcast(FDocumentMutatePageArgs{ InPageID });
+			if (bBroadcastNotify)
+			{
+				PageDelegates.OnRemovingPage.Broadcast(FDocumentMutatePageArgs{ InPageID });
+			}
+			GraphDelegates.Remove(InPageID);
 		}
-
-		PageNodeDelegates.Remove(InPageID);
-		PageEdgeDelegates.Remove(InPageID);
 	}
 } // namespace Metasound::Frontend
diff --git a/Engine/Plugins/Runtime/Metasound/Source/MetasoundFrontend/Private/MetasoundFrontendDocumentVersioning.h b/Engine/Plugins/Runtime/Metasound/Source/MetasoundFrontend/Private/MetasoundFrontendDocumentVersioning.h
--- a/Engine/Plugins/Runtime/Metasound/Source/MetasoundFrontend/Private/MetasoundFrontendDocumentVersioning.h
+++ b/Engine/Plugins/Runtime/Metasound/Source/MetasoundFrontend/Private/MetasoundFrontendDocumentVersioning.h
@@ -19,6 +19,15 @@ namespace Metasound::Frontend
 		return FMetasoundFrontendVersionNumber { 1, 16 };
 	}
 
+	// Version where page data was migrated from a singleton graph to
+	// multiple paged graphs.  Some legacy behavior requires access to
+	// this version to initialize a builder prior to versioning document
+	// data.
+	static FMetasoundFrontendVersionNumber GetPageMigrationVersion()
+	{
+		return FMetasoundFrontendVersionNumber { 1, 13 };
+	}
+
 	bool ChangeIDComparisonEnabledInAutoUpdate();
 
 	// Versions Frontend Document. Passed as AssetBase for backward compat to
diff --git a/Engine/Plugins/Runtime/Metasound/Source/MetasoundFrontend/Public/MetasoundFrontendDocumentModifyDelegates.h b/Engine/Plugins/Runtime/Metasound/Source/MetasoundFrontend/Public/MetasoundFrontendDocumentModifyDelegates.h
--- a/Engine/Plugins/Runtime/Metasound/Source/MetasoundFrontend/Public/MetasoundFrontendDocumentModifyDelegates.h
+++ b/Engine/Plugins/Runtime/Metasound/Source/MetasoundFrontend/Public/MetasoundFrontendDocumentModifyDelegates.h
@@ -24,6 +24,12 @@ namespace Metasound::Frontend
 		FGuid PageID;
 	};
 
+	struct FDocumentPresetStateChangedArgs
+	{
+	};
+
+	DECLARE_MULTICAST_DELEGATE_OneParam(FOnDocumentPresetStateChanged, const FDocumentPresetStateChangedArgs& /* Args */);
+
 	DECLARE_MULTICAST_DELEGATE_OneParam(FOnDocumentPageAdded, const FDocumentMutatePageArgs& /* Args */);
 	DECLARE_MULTICAST_DELEGATE_OneParam(FOnDocumentRemovingPage, const FDocumentMutatePageArgs& /* Args */);
 	DECLARE_MULTICAST_DELEGATE_OneParam(FOnDocumentPageSet, const FDocumentMutatePageArgs& /* Args */);
@@ -84,63 +90,51 @@ namespace Metasound::Frontend
 		FOnMetaSoundFrontendDocumentRemoveSwappingArray OnRemoveSwappingEdge;
 	};
 
+	struct FGraphModifyDelegates
+	{
+		FEdgeModifyDelegates EdgeDelegates;
+		FNodeModifyDelegates NodeDelegates;
+	};
+
 	struct FDocumentModifyDelegates : TSharedFromThis<FDocumentModifyDelegates>
 	{
 		UE_API FDocumentModifyDelegates();
-		UE_API FDocumentModifyDelegates(const FDocumentModifyDelegates& InModifyDelegates);
-		UE_API FDocumentModifyDelegates(FDocumentModifyDelegates&&);
 		UE_API FDocumentModifyDelegates(const FMetasoundFrontendDocument& Document);
-		UE_API FDocumentModifyDelegates& operator=(const FDocumentModifyDelegates&);
-		UE_API FDocumentModifyDelegates& operator=(FDocumentModifyDelegates&&);
+
+		~FDocumentModifyDelegates() = default;
 
 		FOnMetaSoundFrontendDocumentMutateMetadata OnDocumentMetadataChanged;
 		FOnMetaSoundFrontendDocumentMutateArray OnDependencyAdded;
 		FOnMetaSoundFrontendDocumentRemoveSwappingArray OnRemoveSwappingDependency;
 		FOnMetaSoundFrontendDocumentRenameClass OnRenamingDependencyClass;
 
+		FOnDocumentPresetStateChanged OnPresetStateChanged;
+
 		FPageModifyDelegates PageDelegates;
 		FInterfaceModifyDelegates InterfaceDelegates;
 
-		UE_DEPRECATED(5.5, "Public exposition of NodeDelegates will be removed in a future build.  Use accessor 'FindNodeDelegates' instead")
-		FNodeModifyDelegates NodeDelegates;
-
-		UE_DEPRECATED(5.5, "Public exposition of EdgeDelegates will be removed in a future build.  Use accessor 'FindEdgeDelegates' instead")
-		FEdgeModifyDelegates EdgeDelegates;
+	private:
+		TSortedMap<FGuid, FGraphModifyDelegates> GraphDelegates;
 
+	public:
 		UE_API void AddPageDelegates(const FGuid& InPageID);
 		UE_API void RemovePageDelegates(const FGuid& InPageID, bool bBroadcastNotify = true);
 
-	private:
-		TSortedMap<FGuid, FNodeModifyDelegates> PageNodeDelegates;
-		TSortedMap<FGuid, FEdgeModifyDelegates> PageEdgeDelegates;
-
-	public:
+		UE_DEPRECATED(5.7, TEXT("Use FindGraphDelegatesChecked instead"))
 		UE_API FNodeModifyDelegates& FindNodeDelegatesChecked(const FGuid& InPageID);
+
+		UE_DEPRECATED(5.7, TEXT("Use FindGraphDelegatesChecked instead"))
 		UE_API FEdgeModifyDelegates& FindEdgeDelegatesChecked(const FGuid& InPageID);
 
-		void IterateGraphEdgeDelegates(TFunctionRef<void(FEdgeModifyDelegates&)> Func)
-		{
-PRAGMA_DISABLE_DEPRECATION_WARNINGS
-			Func(EdgeDelegates);
-
-			for (TPair<FGuid, FEdgeModifyDelegates>& Delegates : PageEdgeDelegates)
-			{
-				Func(Delegates.Value);
-			}
-PRAGMA_ENABLE_DEPRECATION_WARNINGS
-		}
-
-		void IterateGraphNodeDelegates(TFunctionRef<void(FNodeModifyDelegates&)> Func)
-		{
-PRAGMA_DISABLE_DEPRECATION_WARNINGS
-			Func(NodeDelegates);
-
-			for (TPair<FGuid, FNodeModifyDelegates>& Delegates : PageNodeDelegates)
-			{
-				Func(Delegates.Value);
-			}
-PRAGMA_ENABLE_DEPRECATION_WARNINGS
-		}
+		UE_DEPRECATED(5.7, TEXT("Use IterateGraphDelegates instead"))
+		UE_API void IterateGraphEdgeDelegates(TFunctionRef<void(FEdgeModifyDelegates&)> Func);
+
+		UE_DEPRECATED(5.7, TEXT("Use IterateGraphDelegates instead"))
+		UE_API void IterateGraphNodeDelegates(TFunctionRef<void(FNodeModifyDelegates&)> Func);
+
+		UE_API FGraphModifyDelegates& FindGraphDelegatesChecked(const FGuid& InPageID);
+
+		UE_API void IterateGraphDelegates(TFunctionRef<void(FGraphModifyDelegates&)> Func);
 	};
 
 	class IDocumentBuilderTransactionListener : public TSharedFromThis<IDocumentBuilderTransactionListener>
