Fix strange interaction with MetaSound variables and subgraphs causing delayed trigger variables to not execute.
#rnx

diff --git a/Engine/Plugins/Runtime/Metasound/Source/MetasoundFrontend/Public/MetasoundInputNode.h b/Engine/Plugins/Runtime/Metasound/Source/MetasoundFrontend/Public/MetasoundInputNode.h
--- a/Engine/Plugins/Runtime/Metasound/Source/MetasoundFrontend/Public/MetasoundInputNode.h
+++ b/Engine/Plugins/Runtime/Metasound/Source/MetasoundFrontend/Public/MetasoundInputNode.h
@@ -100,7 +100,7 @@ namespace Metasound
 
 			virtual void BindOutputs(FOutputVertexInterfaceData& InOutVertexData) override
 			{
-				InOutVertexData.BindVertex(DataReferenceName, DataRef);
+				InOutVertexData.BindReadVertex<DataType>(DataReferenceName, DataRef);
 			}
 
 			virtual FExecuteFunction GetExecuteFunction() override
diff --git a/Engine/Plugins/Runtime/Metasound/Source/MetasoundGraphCore/Public/MetasoundVertexData.h b/Engine/Plugins/Runtime/Metasound/Source/MetasoundGraphCore/Public/MetasoundVertexData.h
--- a/Engine/Plugins/Runtime/Metasound/Source/MetasoundGraphCore/Public/MetasoundVertexData.h
+++ b/Engine/Plugins/Runtime/Metasound/Source/MetasoundGraphCore/Public/MetasoundVertexData.h
@@ -340,6 +340,12 @@ namespace Metasound
 				Set(TDataReadReference<TGetReferenceDataType<DataReferenceType>>(InOutDataReference));
 			}
 
+			template<typename DataType>
+			void BindReadFromAny(FAnyDataReference& InOutDataReference)
+			{
+				Set(InOutDataReference.GetAs<TDataReadReference<DataType>>());
+			}
+
 			template<typename DataType>
 			void BindWrite(TDataWriteReference<DataType>& InOutDataReference)
 			{
@@ -1099,9 +1105,27 @@ namespace Metasound
 		{
 			using namespace VertexDataPrivate;
 			auto BindData = [&](FOutputBinding& Binding) { Binding.BindRead(InOutDataReference); };
+
 			Apply(InVertexName, TCreateOutputBinding(InVertexName, InOutDataReference), BindData);
 		}
 
+		/** Bind a read vertex from a any reference. */
+		template<typename DataType>
+		void BindReadVertex(const FVertexName& InVertexName, FAnyDataReference& InOutDataReference)
+		{
+			using namespace VertexDataPrivate;
+			check(InOutDataReference.GetDataTypeName() == GetMetasoundDataTypeName<DataType>());
+
+			auto BindData = [&](FOutputBinding& Binding) { Binding.BindReadFromAny<DataType>(InOutDataReference); };
+			auto CreateBinding = [&]()
+			{
+				FOutputDataVertex Vertex(InVertexName, GetMetasoundDataTypeName<DataType>(), FDataVertexMetadata{}, EVertexAccessType::Reference);
+				return FOutputBinding(Vertex);
+			};
+
+			Apply(InVertexName, CreateBinding, BindData);
+		}
+
 		/** Bind a write vertex from a write reference.*/
 		template<typename DataType>
 		UE_DEPRECATED(5.6, "Please use a non-const lvalue reference for data references passed to this function. Const data references can sometimes be created by passing in temporary objects and/or triggering implicit conversion.")
