Fix LLM tracking when registering metasound nodes.
#rnx
#FailedValidations Preflight

diff --git a/Engine/Plugins/Runtime/Metasound/Source/MetasoundFrontend/Private/MetasoundFrontendModuleRegistrationMacros.cpp b/Engine/Plugins/Runtime/Metasound/Source/MetasoundFrontend/Private/MetasoundFrontendModuleRegistrationMacros.cpp
--- a/Engine/Plugins/Runtime/Metasound/Source/MetasoundFrontend/Private/MetasoundFrontendModuleRegistrationMacros.cpp
+++ b/Engine/Plugins/Runtime/Metasound/Source/MetasoundFrontend/Private/MetasoundFrontendModuleRegistrationMacros.cpp
@@ -2,6 +2,7 @@
 
 #include "MetasoundFrontendModuleRegistrationMacros.h"
 #include "MetasoundLog.h"
+#include "MetasoundTrace.h"
 
 namespace Metasound::Frontend::RegistrationPrivate
 {
@@ -23,6 +24,8 @@ namespace Metasound::Frontend::RegistrationPrivate
 
 	void FRegistrationListBase::RegisterAll(const Metasound::Frontend::FModuleInfo& InOwningModuleInfo, FRegistrationAction* Head)
 	{
+		METASOUND_LLM_SCOPE 
+
 #if WITH_EDITORONLY_DATA
 		UE_LOG(LogMetaSound, Verbose, TEXT("Running registration actions for module %s in plugin %s"), *InOwningModuleInfo.ModuleName.ToString(), *InOwningModuleInfo.PluginName.ToString());
 #endif // if WITH_EDITORONLY_DATA
@@ -40,6 +43,8 @@ namespace Metasound::Frontend::RegistrationPrivate
 
 	void FRegistrationListBase::UnregisterAll(const FModuleInfo& InOwningModuleInfo, FRegistrationAction* Head)
 	{
+		METASOUND_LLM_SCOPE 
+
 #if WITH_EDITORONLY_DATA
 		UE_LOG(LogMetaSound, Verbose, TEXT("Running unregistration actions for module %s in plugin %s"), *InOwningModuleInfo.ModuleName.ToString(), *InOwningModuleInfo.PluginName.ToString());
 #endif // if WITH_EDITORONLY_DATA
@@ -63,6 +68,8 @@ namespace Metasound::Frontend::RegistrationPrivate
 
 	void FGlobalRegistrationList::RegisterAll(const FModuleInfo& InOwningModuleInfo, FRegistrationAction* Head)
 	{
+		METASOUND_LLM_SCOPE 
+
 #if WITH_EDITORONLY_DATA
 		UE_LOG(LogMetaSound, Verbose, TEXT("Running pending registration actions for global registration list"));
 #endif // if WITH_EDITORONLY_DATA
