Fix registry tracking for seamless travel.

#FailedValidations Preflight

diff --git a/Engine/Plugins/Runtime/Metasound/Source/MetasoundFrontend/Private/MetasoundAssetBase.cpp b/Engine/Plugins/Runtime/Metasound/Source/MetasoundFrontend/Private/MetasoundAssetBase.cpp
--- a/Engine/Plugins/Runtime/Metasound/Source/MetasoundFrontend/Private/MetasoundAssetBase.cpp
+++ b/Engine/Plugins/Runtime/Metasound/Source/MetasoundFrontend/Private/MetasoundAssetBase.cpp
@@ -513,7 +513,7 @@ void FMetasoundAssetBase::UpdateAndRegisterForSerialization(const Metasound::Fro
 		TUniquePtr<INodeClassRegistryEntry> RegistryEntry = MakeUnique<AssetBasePrivate::FDocumentNodeRegistryEntryForSerialization>(Document, AssetPath);
 
 		const FNodeClassRegistryKey NodeKey = FNodeClassRegistry::Get().RegisterNode(MoveTemp(RegistryEntry));
-		GraphRegistryKey = FGraphRegistryKey { NodeKey, AssetPath };
+		GraphRegistryKey = FGraphRegistryKey { NodeKey, AssetPath, Owner->GetUniqueID()};
 	}
 
 	if (!GraphRegistryKey.IsValid())
diff --git a/Engine/Plugins/Runtime/Metasound/Source/MetasoundFrontend/Private/MetasoundFrontendNodeClassRegistry.cpp b/Engine/Plugins/Runtime/Metasound/Source/MetasoundFrontend/Private/MetasoundFrontendNodeClassRegistry.cpp
--- a/Engine/Plugins/Runtime/Metasound/Source/MetasoundFrontend/Private/MetasoundFrontendNodeClassRegistry.cpp
+++ b/Engine/Plugins/Runtime/Metasound/Source/MetasoundFrontend/Private/MetasoundFrontendNodeClassRegistry.cpp
@@ -363,12 +363,12 @@ PRAGMA_ENABLE_DEPRECATION_WARNINGS
 	} // namespace NodeClassRegistryKey
 	  
 
-	void FNodeClassRegistry::BuildAndRegisterGraphFromDocument(const FMetasoundFrontendDocument& InDocument, const FProxyDataCache& InProxyDataCache, FNodeClassInfo&& InNodeClassInfo, const FTopLevelAssetPath& AssetPath, TArrayView<const FGuid> InPageOrder)
+	void FNodeClassRegistry::BuildAndRegisterGraphFromDocument(const FMetasoundFrontendDocument& InDocument, const FProxyDataCache& InProxyDataCache, FNodeClassInfo&& InNodeClassInfo, const FGraphRegistryKey& InGraphRegistryKey, TArrayView<const FGuid> InPageOrder)
 	{
 		using namespace RegistryPrivate;
 
 		METASOUND_TRACE_CPUPROFILER_EVENT_SCOPE(Metasound::FNodeClassRegistry::BuildAndRegisterGraphFromDocument);
-		METASOUND_TRACE_CPUPROFILER_EVENT_SCOPE_TEXT(*FString::Printf(TEXT("Metasound::FNodeClassRegistry::BuildAndRegisterGraphFromDocument asset %s"), *AssetPath.GetAssetName().ToString()));
+		METASOUND_TRACE_CPUPROFILER_EVENT_SCOPE_TEXT(*FString::Printf(TEXT("Metasound::FNodeClassRegistry::BuildAndRegisterGraphFromDocument asset %s"), *InGraphRegistryKey.AssetPath.GetAssetName().ToString()));
 
 		FGuid AssetClassID;
 		if (IMetaSoundAssetManager* AssetManager = IMetaSoundAssetManager::Get())
@@ -377,15 +377,15 @@ PRAGMA_ENABLE_DEPRECATION_WARNINGS
 		}
 		else
 		{
-			UE_LOG(LogMetaSound, Warning, TEXT("No AssetManager registered, registering bespoke AssetClassID '%s' for asset '%s'"), *AssetClassID.ToString(), *AssetPath.ToString());
+			UE_LOG(LogMetaSound, Warning, TEXT("No AssetManager registered, registering bespoke AssetClassID '%s' for asset '%s'"), *AssetClassID.ToString(), *InGraphRegistryKey.AssetPath.ToString());
 			AssetClassID = FGuid::NewGuid();
 		}
 
 		// Use the asset class id for the graph id because it should be locally unique per asset.
-		TUniquePtr<FFrontendGraph> FrontendGraph = Frontend::FGraphBuilder::CreateGraph(InDocument, InProxyDataCache, AssetPath, /*GraphId=*/ AssetClassID, InPageOrder);
+		TUniquePtr<FFrontendGraph> FrontendGraph = Frontend::FGraphBuilder::CreateGraph(InDocument, InProxyDataCache, InGraphRegistryKey.AssetPath, /*GraphId=*/ AssetClassID, InPageOrder);
 		if (!FrontendGraph.IsValid())
 		{
-			UE_LOG(LogMetaSound, Error, TEXT("Failed to build MetaSound graph in asset '%s'"), *AssetPath.ToString());
+			UE_LOG(LogMetaSound, Error, TEXT("Failed to build MetaSound graph in asset '%s'"), *InGraphRegistryKey.AssetPath.ToString());
 		}
 
 		TSharedPtr<const FGraph> GraphToRegister(FrontendGraph.Release());
@@ -394,15 +394,16 @@ PRAGMA_ENABLE_DEPRECATION_WARNINGS
 			InDocument.Interfaces,
 			MoveTemp(InNodeClassInfo),
 			GraphToRegister,
-			AssetPath);
+			InGraphRegistryKey.AssetPath);
 
-		const FNodeClassRegistryKey RegistryKey = RegisterNodeInternal(MoveTemp(RegistryEntry));
+		const FNodeClassRegistryKey NodeRegistryKey = RegisterNodeInternal(MoveTemp(RegistryEntry));
+		checkf(NodeRegistryKey == InGraphRegistryKey.NodeKey, TEXT("Mismatched between expected node registry key for registered MetaSound graph. Expected %s but found %s"), *InGraphRegistryKey.NodeKey.ToString(), *NodeRegistryKey.ToString());
 
-		// Key must use the asset path provided to function and *NOT* that of the
-		// document's owning DocumentInterface object, as that may be a built/optimized
-		// transient object with a different, transient asset path.
-		const FGraphRegistryKey GraphKey { RegistryKey, AssetPath };
-		RegisterGraphInternal(GraphKey, GraphToRegister);
+		// Key must use the graph registry key provided to this function and *NOT* one created 
+		// from the document's owning DocumentInterface object, as that may have been 
+		// built/optimized using a transient object with a different, transient asset path 
+		// and different UObject UniqueID.
+		RegisterGraphInternal(InGraphRegistryKey, GraphToRegister);
 	}
 
 	FNodeClassRegistry* FNodeClassRegistry::LazySingleton = nullptr;
@@ -553,14 +554,19 @@ PRAGMA_ENABLE_DEPRECATION_WARNINGS
 		check(InDocumentInterface);
 		check(IsInGameThread());
 
+		uint32 OwningObjUniqueID = (uint32)INDEX_NONE;
+		if (const UObject* OwningObj = InDocumentInterface.GetObject())
+		{
+			OwningObjUniqueID = OwningObj->GetUniqueID();
+		}
 		const FMetasoundFrontendDocument& Document = InDocumentInterface->GetConstDocument();
 		const FTopLevelAssetPath AssetPath = InDocumentInterface->GetAssetPathChecked();
-		const FGraphRegistryKey RegistryKey { FNodeClassRegistryKey(Document.RootGraph), AssetPath };
+		const FGraphRegistryKey RegistryKey { FNodeClassRegistryKey(Document.RootGraph), AssetPath, OwningObjUniqueID };
 
 		if (!RegistryKey.IsValid())
 		{
 			// Do not attempt to build and register a MetaSound with an invalid registry key
-			UE_LOG(LogMetaSound, Warning, TEXT("Registry key is invalid when attempting to register graph for asset %s"), *AssetPath.ToString());
+			UE_LOG(LogMetaSound, Warning, TEXT("Registry key is invalid when attempting to register graph for asset %s, key %s"), *AssetPath.ToString(), *RegistryKey.ToString());
 			return RegistryKey;
 		}
 
@@ -575,8 +581,6 @@ PRAGMA_ENABLE_DEPRECATION_WARNINGS
 		// built version as the build process may in fact create a new object with a transient path.
 		const TScriptInterface<IMetaSoundDocumentInterface> RegistryDocInterface = RegistryPrivate::BuildRegistryDocument(InDocumentInterface, bAsync);
 
-		UObject* OwningObject = RegistryDocInterface.GetObject();
-		check(OwningObject);
 
 		// Proxies are created synchronously to avoid creating proxies in async tasks. Proxies
 		// are created from UObjects which need to be protected from GC and non-GT access.
@@ -626,12 +630,12 @@ PRAGMA_ENABLE_DEPRECATION_WARNINGS
 			TArray<FGuid> PageOrder(InPageOrder.GetData(), InPageOrder.Num());
 			Tasks::FTask BuildAndRegisterTask = AsyncRegistrationPipe.Launch(
 				UE_SOURCE_LOCATION,
-				[RegistryKey, ClassInfo = MoveTemp(NodeClassInfo), AssetPath, RegistryDocInterface, ProxyDataCache = MoveTemp(ProxyDataCache), PageOrder = MoveTemp(PageOrder), RegistrationTaskID]() mutable
+				[RegistryKey, ClassInfo = MoveTemp(NodeClassInfo), RegistryDocInterface, ProxyDataCache = MoveTemp(ProxyDataCache), PageOrder = MoveTemp(PageOrder), RegistrationTaskID]() mutable
 				{
 					FNodeClassRegistry& Registry = FNodeClassRegistry::Get();
 					// Unregister the graph before re-registering
 					Registry.UnregisterGraphInternal(RegistryKey);
-					Registry.BuildAndRegisterGraphFromDocument(RegistryDocInterface->GetConstDocument(), ProxyDataCache, MoveTemp(ClassInfo), AssetPath, PageOrder);
+					Registry.BuildAndRegisterGraphFromDocument(RegistryDocInterface->GetConstDocument(), ProxyDataCache, MoveTemp(ClassInfo), RegistryKey, PageOrder);
 					Registry.RemoveRegistrationTask(RegistryKey, RegistrationTaskID, FNodeClassRegistryTransaction::ETransactionType::NodeRegistration);
 					Registry.RemoveDocumentReference(RegistryDocInterface);
 				}
@@ -651,7 +655,7 @@ PRAGMA_ENABLE_DEPRECATION_WARNINGS
 			UnregisterGraphInternal(RegistryKey);
 
 			// Build and register graph synchronously
-			BuildAndRegisterGraphFromDocument(RegistryDocInterface->GetConstDocument(), ProxyDataCache, MoveTemp(NodeClassInfo), AssetPath, InPageOrder);
+			BuildAndRegisterGraphFromDocument(RegistryDocInterface->GetConstDocument(), ProxyDataCache, MoveTemp(NodeClassInfo), RegistryKey, InPageOrder);
 		}
 
 		return RegistryKey;
diff --git a/Engine/Plugins/Runtime/Metasound/Source/MetasoundFrontend/Private/MetasoundFrontendNodeClassRegistryPrivate.h b/Engine/Plugins/Runtime/Metasound/Source/MetasoundFrontend/Private/MetasoundFrontendNodeClassRegistryPrivate.h
--- a/Engine/Plugins/Runtime/Metasound/Source/MetasoundFrontend/Private/MetasoundFrontendNodeClassRegistryPrivate.h
+++ b/Engine/Plugins/Runtime/Metasound/Source/MetasoundFrontend/Private/MetasoundFrontendNodeClassRegistryPrivate.h
@@ -155,7 +155,7 @@ namespace Metasound::Frontend
 			int32 RegistrationTaskID = 0;
 		};
 
-		void BuildAndRegisterGraphFromDocument(const FMetasoundFrontendDocument& InDocument, const FProxyDataCache& InProxyDataCache, FNodeClassInfo&& InNodeClassInfo, const FTopLevelAssetPath& AssetPath, TArrayView<const FGuid> InPageOrder);
+		void BuildAndRegisterGraphFromDocument(const FMetasoundFrontendDocument& InDocument, const FProxyDataCache& InProxyDataCache, FNodeClassInfo&& InNodeClassInfo, const FGraphRegistryKey& InGraphRegistryKey, TArrayView<const FGuid> InPageOrder);
 
 		void AddRegistrationTask(const FGraphRegistryKey& InKey, FActiveRegistrationTaskInfo&& TaskInfo);
 		void RemoveRegistrationTask(const FGraphRegistryKey& InKey, int32 InRegistrationTaskID, FNodeClassRegistryTransaction::ETransactionType TransactionType);
diff --git a/Engine/Plugins/Runtime/Metasound/Source/MetasoundFrontend/Private/MetasoundFrontendRegistryKey.cpp b/Engine/Plugins/Runtime/Metasound/Source/MetasoundFrontend/Private/MetasoundFrontendRegistryKey.cpp
--- a/Engine/Plugins/Runtime/Metasound/Source/MetasoundFrontend/Private/MetasoundFrontendRegistryKey.cpp
+++ b/Engine/Plugins/Runtime/Metasound/Source/MetasoundFrontend/Private/MetasoundFrontendRegistryKey.cpp
@@ -273,6 +273,9 @@ namespace Metasound::Frontend
 		Builder.Append(AssetPath.GetPackageName().ToString());
 		Builder.Append(TEXT("/"));
 		Builder.Append(AssetPath.GetAssetName().ToString());
+		Builder.Append(TEXT(", "));
+		Builder.Append(::LexToString(ObjectID));
+
 		return Builder.ToString();
 	}
 
@@ -288,6 +291,6 @@ namespace Metasound::Frontend
 
 	bool FGraphClassRegistryKey::IsValid() const
 	{
-		return NodeKey.IsValid() && AssetPath.IsValid();
+		return NodeKey.IsValid() && AssetPath.IsValid() && (ObjectID != (uint32)INDEX_NONE);
 	}
 } // namespace Metasound::Frontend
diff --git a/Engine/Plugins/Runtime/Metasound/Source/MetasoundFrontend/Public/MetasoundFrontendRegistryKey.h b/Engine/Plugins/Runtime/Metasound/Source/MetasoundFrontend/Public/MetasoundFrontendRegistryKey.h
--- a/Engine/Plugins/Runtime/Metasound/Source/MetasoundFrontend/Public/MetasoundFrontendRegistryKey.h
+++ b/Engine/Plugins/Runtime/Metasound/Source/MetasoundFrontend/Public/MetasoundFrontendRegistryKey.h
@@ -129,6 +129,7 @@ namespace Metasound::Frontend
 	{
 		FNodeClassRegistryKey NodeKey;
 		FTopLevelAssetPath AssetPath;
+		uint32 ObjectID = (uint32)INDEX_NONE;
 
 		UE_API FString ToString() const;
 
@@ -139,12 +140,12 @@ namespace Metasound::Frontend
 
 		inline friend bool operator==(const FGraphClassRegistryKey& InLHS, const FGraphClassRegistryKey& InRHS)
 		{
-			return (InLHS.NodeKey == InRHS.NodeKey) && (InLHS.AssetPath == InRHS.AssetPath);
+			return (InLHS.NodeKey == InRHS.NodeKey) && (InLHS.AssetPath == InRHS.AssetPath) && (InLHS.ObjectID == InRHS.ObjectID);
 		}
 
 		friend inline uint32 GetTypeHash(const FGraphClassRegistryKey& InKey)
 		{
-			const int32 Hash = HashCombineFast(GetTypeHash(InKey.NodeKey), GetTypeHash(InKey.AssetPath));
+			const int32 Hash = HashCombineFast(::GetTypeHash(InKey.ObjectID), HashCombineFast(GetTypeHash(InKey.NodeKey), GetTypeHash(InKey.AssetPath)));
 			return Hash;
 		}
 	};
