Remove logging macros from polymorphic type registration, which reduces .exe template float

#DisabledValidations Preflight

diff --git a/Engine/Plugins/Runtime/Metasound/Source/MetasoundFrontend/Private/MetasoundFrontendDataTypeRegistry.cpp b/Engine/Plugins/Runtime/Metasound/Source/MetasoundFrontend/Private/MetasoundFrontendDataTypeRegistry.cpp
--- a/Engine/Plugins/Runtime/Metasound/Source/MetasoundFrontend/Private/MetasoundFrontendDataTypeRegistry.cpp
+++ b/Engine/Plugins/Runtime/Metasound/Source/MetasoundFrontend/Private/MetasoundFrontendDataTypeRegistry.cpp
@@ -1,6 +1,7 @@
 // Copyright Epic Games, Inc. All Rights Reserved.
 #include "MetasoundFrontendDataTypeRegistry.h"
 
+#include "Containers/UnrealString.h"
 #include "MetasoundLiteral.h"
 #include "MetasoundTrace.h"
 #include "Misc/ScopeLock.h"
@@ -1541,9 +1542,12 @@ PRAGMA_ENABLE_DEPRECATION_WARNINGS
 			}
 		}
 
-		FPolymorphicDataTypeBase::FPolymorphicDataTypeBase(const FName InDataTypeName, const FName InParentName)
-			: DataTypeName(InDataTypeName), ParentTypeName(InParentName)
+		FPolymorphicDataTypeBase::FPolymorphicDataTypeBase(const FName InDataTypeName, const FName InParentName, const bool InbIsAbstract)
+			: DataTypeName(InDataTypeName), ParentTypeName(InParentName), bIsAbstract(InbIsAbstract)
 		{
+			UE_LOG(LogMetaSound, Verbose, TEXT("PolymorphicInterface Type=%s, Base=%s, bIsAbstract=%d"),
+				*InDataTypeName.ToString(), *ParentTypeName.ToString(), (int32)bIsAbstract);
+			
 			// Register with Core Registry.
 			PRAGMA_DISABLE_INTERNAL_WARNINGS
 			FPolyRegistry::Get().Register(
diff --git a/Engine/Plugins/Runtime/Metasound/Source/MetasoundFrontend/Public/MetasoundDataTypeRegistrationMacro.h b/Engine/Plugins/Runtime/Metasound/Source/MetasoundFrontend/Public/MetasoundDataTypeRegistrationMacro.h
--- a/Engine/Plugins/Runtime/Metasound/Source/MetasoundFrontend/Public/MetasoundDataTypeRegistrationMacro.h
+++ b/Engine/Plugins/Runtime/Metasound/Source/MetasoundFrontend/Public/MetasoundDataTypeRegistrationMacro.h
@@ -513,19 +513,13 @@ namespace Metasound::Frontend
 				void InitPolymorphicInterface()
 				{
 					using FPolyTraits = TPolymorphicTraits<TDataType>;
-
 					if constexpr (FPolyTraits::bIsPolymorphic)
 					{
 						using BaseType = FPolyTraits::BaseType;
-						const FName ParentTypeName = GetMetasoundDataTypeName<BaseType>();
-						const FName TypeName = GetMetasoundDataTypeName<TDataType>();
-						this->PolymorphicInterface = MakeShared<FPolymorphicDataTypeBase>(TypeName, ParentTypeName);
-						UE_LOG(LogMetaSound, Verbose, TEXT("PolymorphicInterface Type=%s, Base=%s, bIsAbstract=%d"),
-							*TypeName.ToString(), *ParentTypeName.ToString(), (int32)FPolyTraits::bIsAbstract);
-					}
-					else
-					{
-						UE_CLOG(EnumInterface.IsValid(), LogMetaSound, Display, TEXT("PolymorphicInterface null for enum %s"), *GetMetasoundDataTypeName<TDataType>().ToString());
+						this->PolymorphicInterface = MakeShared<FPolymorphicDataTypeBase>(
+							GetMetasoundDataTypeName<TDataType>(),
+							GetMetasoundDataTypeName<BaseType>(),
+							FPolyTraits::bIsAbstract);
 					}
 				}
 
diff --git a/Engine/Plugins/Runtime/Metasound/Source/MetasoundFrontend/Public/MetasoundFrontendDataTypeRegistry.h b/Engine/Plugins/Runtime/Metasound/Source/MetasoundFrontend/Public/MetasoundFrontendDataTypeRegistry.h
--- a/Engine/Plugins/Runtime/Metasound/Source/MetasoundFrontend/Public/MetasoundFrontendDataTypeRegistry.h
+++ b/Engine/Plugins/Runtime/Metasound/Source/MetasoundFrontend/Public/MetasoundFrontendDataTypeRegistry.h
@@ -454,7 +454,7 @@ namespace Metasound
 		{
 			virtual ~IPolymorphicDataTypeInterface() = default;
 			UE_EXPERIMENTAL(5.7, "Metasound Polymorphic Datatypes are still in devleopement")
-			UE_API virtual bool IsA(const FName InDataTypeName, const IDataTypeRegistry& InRegistry = IDataTypeRegistry::Get() ) const = 0;
+			virtual bool IsA(const FName InDataTypeName, const IDataTypeRegistry& InRegistry = IDataTypeRegistry::Get() ) const = 0;
 		};
 
 		/**
@@ -464,7 +464,7 @@ namespace Metasound
 		class FPolymorphicDataTypeBase : public IPolymorphicDataTypeInterface
 		{
 		public:
-			UE_API explicit FPolymorphicDataTypeBase(const FName InDataTypeName, const FName InParentName = NAME_None);
+			UE_API FPolymorphicDataTypeBase(const FName InDataTypeName, const FName InParentName = NAME_None, const bool InbIsAbstract = false);
 			virtual ~FPolymorphicDataTypeBase() override = default;
 
 			UE_EXPERIMENTAL(5.7, "Metasound Polymorphic Datatypes are still in developement")			
@@ -473,6 +473,7 @@ namespace Metasound
 			void ResolveParent(const IDataTypeRegistry& InRegistry) const;
 			FName DataTypeName = NAME_None;
 			FName ParentTypeName = NAME_None;
+			bool bIsAbstract = false;
 			mutable TWeakPtr<const IPolymorphicDataTypeInterface> ParentDataType;
 		};
 		
