Fix MetaSound preset interface input override changes not persisting
- Presets should now inherit default metadata (ex. float range/clamp) from parent graph if applicable
- Fix member menu selection issue for editing interface inputs

diff --git a/Engine/Plugins/Runtime/Metasound/Source/MetasoundEditor/Private/MetasoundEditor.cpp b/Engine/Plugins/Runtime/Metasound/Source/MetasoundEditor/Private/MetasoundEditor.cpp
--- a/Engine/Plugins/Runtime/Metasound/Source/MetasoundEditor/Private/MetasoundEditor.cpp
+++ b/Engine/Plugins/Runtime/Metasound/Source/MetasoundEditor/Private/MetasoundEditor.cpp
@@ -2679,6 +2679,18 @@ namespace Metasound
 			}
 		}
 
+		UMetasoundEditorGraph* FEditor::GetMetaSoundGraph()
+		{
+			FMetasoundAssetBase* MetaSoundAsset = IMetasoundUObjectRegistry::Get().GetObjectAsAssetBase(GetMetasoundObject());
+			if (MetaSoundAsset)
+			{
+				UEdGraph* Graph = MetaSoundAsset->GetGraph();
+				return Cast<UMetasoundEditorGraph>(Graph);
+			}
+
+			return nullptr;
+		}
+
 		UMetasoundEditorGraph& FEditor::GetMetaSoundGraphChecked()
 		{
 			FMetasoundAssetBase* MetaSoundAsset = IMetasoundUObjectRegistry::Get().GetObjectAsAssetBase(GetMetasoundObject());
@@ -5202,6 +5214,7 @@ PRAGMA_ENABLE_DEPRECATION_WARNINGS
 		{
 			OutDelegates.PageDelegates.OnPageSet.AddSP(this, &FDocumentListener::OnPageSet);
 			OutDelegates.InterfaceDelegates.OnInputDefaultChanged.AddSP(this, &FDocumentListener::OnInputDefaultChanged);
+			OutDelegates.InterfaceDelegates.OnMemberMetadataSet.AddSP(this, &FDocumentListener::OnMemberMetadataSet);
 		}
 
 		void FEditor::FDocumentListener::OnPageSet(const Frontend::FDocumentMutatePageArgs& Args)
@@ -5244,44 +5257,73 @@ PRAGMA_ENABLE_DEPRECATION_WARNINGS
 		{
 			if (TSharedPtr<FEditor> ParentPtr = Parent.Pin())
 			{
+				UMetasoundEditorGraph* MetaSoundGraph = ParentPtr->GetMetaSoundGraph();
+				if (!MetaSoundGraph)
+				{
+					return;
+				}
+
 				const FMetaSoundFrontendDocumentBuilder& ConstBuilder = ParentPtr->Builder->GetConstBuilder();
 				const FMetasoundFrontendGraphClass& GraphClass = ConstBuilder.GetConstDocumentChecked().RootGraph;
 				const FMetasoundFrontendGraph& Graph = GraphClass.GetConstDefaultGraph();
 
-				const FMetasoundFrontendClassInput& GraphInput = GraphClass.GetDefaultInterface().Inputs[Index];
-				const FName InputName = GraphInput.Name;
-				UMetasoundEditorGraph& MetaSoundGraph = ParentPtr->GetMetaSoundGraphChecked();
-				UMetasoundEditorGraphInput* EditorGraphInput = MetaSoundGraph.FindInput(InputName);
-				
-				// Literal does not need to be updated if the data type doesn't match,
-				// as it should be getting overwritten by the new literal type after this delegate call
-				if (EditorGraphInput && GraphInput.TypeName == EditorGraphInput->GetDataType())
+				const TArray<FMetasoundFrontendClassInput>& GraphInputs = GraphClass.GetDefaultInterface().Inputs;
+				if (GraphInputs.IsValidIndex(Index))
 				{
-					UMetasoundEditorGraphMemberDefaultLiteral* Literal = EditorGraphInput->GetLiteral();
-					if (Literal)
+					const FMetasoundFrontendClassInput& GraphInput = GraphClass.GetDefaultInterface().Inputs[Index];
+					const FName InputName = GraphInput.Name;
+					UMetasoundEditorGraphInput* EditorGraphInput = MetaSoundGraph->FindInput(InputName);
+
+					// Literal does not need to be updated if the data type doesn't match,
+					// as it should be getting overwritten by the new literal type after this delegate call
+					if (EditorGraphInput && GraphInput.TypeName == EditorGraphInput->GetDataType())
 					{
-						// Synchronize editor literal and update preview instance
-						Literal->Synchronize();
-						if (GEditor && MetaSoundGraph.IsPreviewing())
+						UMetasoundEditorGraphMemberDefaultLiteral* Literal = EditorGraphInput->GetLiteral();
+						if (Literal)
 						{
-							UAudioComponent* PreviewComponent = GEditor->GetPreviewAudioComponent();
-							check(PreviewComponent);
-
-							if (TScriptInterface<IAudioParameterControllerInterface> ParamInterface = PreviewComponent)
+							// Synchronize editor literal and update preview instance
+							Literal->Synchronize();
+							if (GEditor && MetaSoundGraph->IsPreviewing())
 							{
-								Literal->UpdatePreviewInstance(InputName, ParamInterface);
+								UAudioComponent* PreviewComponent = GEditor->GetPreviewAudioComponent();
+								check(PreviewComponent);
+
+								if (TScriptInterface<IAudioParameterControllerInterface> ParamInterface = PreviewComponent)
+								{
+									Literal->UpdatePreviewInstance(InputName, ParamInterface);
+								}
 							}
+
+							// Mark dirty (needed for when this is being called in response to a change from Builder BP API for example)
+							UObject* MetaSound = ParentPtr->GetMetasoundObject();
+							check(MetaSound);
+							MetaSound->MarkPackageDirty();
 						}
-						
-						// Mark dirty (needed for when this is being called in response to a change from Builder BP API for example)
-						UObject* MetaSound = ParentPtr->GetMetasoundObject();
-						check(MetaSound);
-						MetaSound->MarkPackageDirty();
 					}
 				}
 			}
 		}
 
+		void FEditor::FDocumentListener::OnMemberMetadataSet(FGuid MemberID)
+		{
+			TRACE_CPUPROFILER_EVENT_SCOPE(Metasound::Editor::FEditor::OnMemberMetadataSet);
+
+			if (TSharedPtr<FEditor> ParentPtr = Parent.Pin())
+			{
+				UMetasoundEditorGraph* MetaSoundGraph = ParentPtr->GetMetaSoundGraph();
+				if (!MetaSoundGraph)
+				{
+					return;
+				}
+
+				UMetasoundEditorGraphMember* Member = MetaSoundGraph->FindMember(MemberID);
+				if (Member)
+				{
+					Member->InitializeLiteral(/*bForceRebind=*/ true);
+				}
+			}
+		}
+
 		bool FEditor::CanPromoteToInput()
 		{
 			if (MetasoundGraphEditor.IsValid())
diff --git a/Engine/Plugins/Runtime/Metasound/Source/MetasoundEditor/Private/MetasoundEditor.h b/Engine/Plugins/Runtime/Metasound/Source/MetasoundEditor/Private/MetasoundEditor.h
--- a/Engine/Plugins/Runtime/Metasound/Source/MetasoundEditor/Private/MetasoundEditor.h
+++ b/Engine/Plugins/Runtime/Metasound/Source/MetasoundEditor/Private/MetasoundEditor.h
@@ -282,6 +282,7 @@ namespace Metasound::Editor
 		FGraphAppearanceInfo GetGraphAppearance() const;
 
 		UMetasoundEditorGraph& GetMetaSoundGraphChecked();
+		UMetasoundEditorGraph* GetMetaSoundGraph();
 
 		FText GetGraphStatusDescription() const;
 		const FSlateIcon& GetPlayIcon() const;
@@ -626,6 +627,7 @@ namespace Metasound::Editor
 			virtual void OnBuilderReloaded(Frontend::FDocumentModifyDelegates& OutDelegates) override;
 			void OnPageSet(const Frontend::FDocumentMutatePageArgs& Args);
 			void OnInputDefaultChanged(int32 Index);
+			void OnMemberMetadataSet(FGuid MemberID);
 		};
 
 		TSharedPtr<FDocumentListener> DocListener;
diff --git a/Engine/Plugins/Runtime/Metasound/Source/MetasoundEditor/Private/MetasoundEditorGraph.cpp b/Engine/Plugins/Runtime/Metasound/Source/MetasoundEditor/Private/MetasoundEditorGraph.cpp
--- a/Engine/Plugins/Runtime/Metasound/Source/MetasoundEditor/Private/MetasoundEditorGraph.cpp
+++ b/Engine/Plugins/Runtime/Metasound/Source/MetasoundEditor/Private/MetasoundEditorGraph.cpp
@@ -166,7 +166,7 @@ const UMetasoundEditorGraph* UMetasoundEditorGraphMember::GetOwningGraph() const
 	return Cast<const UMetasoundEditorGraph>(GetOuter());
 }
 
-void UMetasoundEditorGraphMember::InitializeLiteral()
+void UMetasoundEditorGraphMember::InitializeLiteral(bool bForceRebind)
 {
 	using namespace Metasound;
 	using namespace Metasound::Editor;
@@ -183,7 +183,7 @@ void UMetasoundEditorGraphMember::InitializeLiteral()
 		LiteralClass = UMetasoundEditorGraphMemberDefaultLiteral::StaticClass();
 	}
 
-	if (!Literal || Literal->GetClass() != LiteralClass)
+	if (bForceRebind || !Literal || Literal->GetClass() != LiteralClass)
 	{
 		FMetaSoundFrontendDocumentBuilder& Builder = GetFrontendBuilderChecked();
 		const bool bIsNew = UMetaSoundEditorSubsystem::GetChecked().BindMemberMetadata(Builder, *this, LiteralClass);
diff --git a/Engine/Plugins/Runtime/Metasound/Source/MetasoundEditor/Private/MetasoundEditorGraphMemberDefaults.cpp b/Engine/Plugins/Runtime/Metasound/Source/MetasoundEditor/Private/MetasoundEditorGraphMemberDefaults.cpp
--- a/Engine/Plugins/Runtime/Metasound/Source/MetasoundEditor/Private/MetasoundEditorGraphMemberDefaults.cpp
+++ b/Engine/Plugins/Runtime/Metasound/Source/MetasoundEditor/Private/MetasoundEditorGraphMemberDefaults.cpp
@@ -614,11 +614,55 @@ void UMetasoundEditorGraphMemberDefaultFloat::InitDefault(const FGuid& InPageID)
 	MemberDefaultsPrivate::InitDefault<float>(*this, InPageID);
 }
 
+void UMetasoundEditorGraphMemberDefaultFloat::PostLoad()
+{
+	using namespace Metasound::Editor;
+	using namespace Metasound::Frontend;
+
+	Super::PostLoad();
+
+	FMetasoundAssetBase& MetaSound = FGraphBuilder::GetOutermostMetaSoundChecked(*this);
+	const FMetaSoundFrontendDocumentBuilder& Builder = IDocumentBuilderRegistry::GetChecked().FindOrBeginBuilding(MetaSound.GetOwningAsset());
+	const bool bIsPreset = Builder.IsPreset();
+
+	// Set default input widget type from settings. Widgets do not apply to presets
+	if (bIsPreset)
+	{
+		if (WidgetType != EMetasoundMemberDefaultWidget::None)
+		{
+			WidgetType = EMetasoundMemberDefaultWidget::None;
+		}
+	}
+}
+
 void UMetasoundEditorGraphMemberDefaultFloat::Initialize()
 {
-	if (const UMetasoundEditorSettings* EditorSettings = ::GetDefault<UMetasoundEditorSettings>())
+	using namespace Metasound::Editor;
+	using namespace Metasound::Frontend;
+
+	FMetasoundAssetBase& MetaSound = FGraphBuilder::GetOutermostMetaSoundChecked(*this);
+	const FMetaSoundFrontendDocumentBuilder& Builder = IDocumentBuilderRegistry::GetChecked().FindOrBeginBuilding(MetaSound.GetOwningAsset());
+	const bool bIsPreset = Builder.IsPreset();
+
+	// Set default input widget type from settings. Widgets do not apply to presets
+	if (!bIsPreset)
 	{
-		WidgetType = EditorSettings->DefaultInputWidgetType;
+		if (const UMetasoundEditorSettings* EditorSettings = ::GetDefault<UMetasoundEditorSettings>())
+		{
+			WidgetType = EditorSettings->DefaultInputWidgetType;
+		}
+	}
+	// Fixup for preset input widget clamp default value.
+	// When rebuilding a preset, clamp default was defaulted to false on the parent graph 
+	// even if a range was applied because a widget was in use.
+	// In that case, we want to make clamp default true and clear out the widget setting.
+	else
+	{
+		if (WidgetType != EMetasoundMemberDefaultWidget::None)
+		{
+			ClampDefault = true;
+			WidgetType = EMetasoundMemberDefaultWidget::None;
+		}
 	}
 }
 
diff --git a/Engine/Plugins/Runtime/Metasound/Source/MetasoundEditor/Private/MetasoundEditorSubsystem.cpp b/Engine/Plugins/Runtime/Metasound/Source/MetasoundEditor/Private/MetasoundEditorSubsystem.cpp
--- a/Engine/Plugins/Runtime/Metasound/Source/MetasoundEditor/Private/MetasoundEditorSubsystem.cpp
+++ b/Engine/Plugins/Runtime/Metasound/Source/MetasoundEditor/Private/MetasoundEditorSubsystem.cpp
@@ -193,6 +193,11 @@ bool UMetaSoundEditorSubsystem::BindMemberMetadata(
 			check(ExistingClass);
 			if (ExistingClass->IsChildOf(LiteralClass))
 			{
+				if (!InMember.Literal || InMember.Literal != Literal)
+				{
+					InMember.Literal = CastChecked<UMetasoundEditorGraphMemberDefaultLiteral>(Literal);
+					return true;
+				}
 				InMember.Literal = CastChecked<UMetasoundEditorGraphMemberDefaultLiteral>(Literal);
 				return false;
 			}
diff --git a/Engine/Plugins/Runtime/Metasound/Source/MetasoundEditor/Public/MetasoundEditorGraph.h b/Engine/Plugins/Runtime/Metasound/Source/MetasoundEditor/Public/MetasoundEditorGraph.h
--- a/Engine/Plugins/Runtime/Metasound/Source/MetasoundEditor/Public/MetasoundEditorGraph.h
+++ b/Engine/Plugins/Runtime/Metasound/Source/MetasoundEditor/Public/MetasoundEditorGraph.h
@@ -241,8 +241,9 @@ public:
 	/** Returns literal associated with the given member */
 	UMetasoundEditorGraphMemberDefaultLiteral* GetLiteral() const { return Literal; }
 
-	/** Creates new literal if there is none and/or conforms literal object type to member's DataType */
-	UE_API void InitializeLiteral();
+	/** Creates new literal if there is none or conforms literal object type to member's DataType if not matching.
+	  * If force rebind is true, will always attempt to reinitialize with updated data from frontend literal. */
+	UE_API void InitializeLiteral(bool bForceRebind=false);
 
 protected:
 	/** Default literal value of member */
diff --git a/Engine/Plugins/Runtime/Metasound/Source/MetasoundEditor/Public/MetasoundEditorGraphMemberDefaults.h b/Engine/Plugins/Runtime/Metasound/Source/MetasoundEditor/Public/MetasoundEditorGraphMemberDefaults.h
--- a/Engine/Plugins/Runtime/Metasound/Source/MetasoundEditor/Public/MetasoundEditorGraphMemberDefaults.h
+++ b/Engine/Plugins/Runtime/Metasound/Source/MetasoundEditor/Public/MetasoundEditorGraphMemberDefaults.h
@@ -373,7 +373,7 @@ public:
 	virtual bool TryFindDefault(FMetasoundFrontendLiteral& OutLiteral, const FGuid* InPageID = nullptr) const override;
 	virtual void UpdatePreviewInstance(const Metasound::FVertexName& InParameterName, TScriptInterface<IAudioParameterControllerInterface>& InParameterInterface) const override;
 	virtual void PostEditChangeChainProperty(FPropertyChangedChainEvent& InPropertyChangedEvent) override;
-
+	virtual void PostLoad() override;
 	// Set Range to reasonable limit given current Default value
 	void SetInitialRange();
 
diff --git a/Engine/Plugins/Runtime/Metasound/Source/MetasoundFrontend/Private/MetasoundFrontendDocumentBuilder.cpp b/Engine/Plugins/Runtime/Metasound/Source/MetasoundFrontend/Private/MetasoundFrontendDocumentBuilder.cpp
--- a/Engine/Plugins/Runtime/Metasound/Source/MetasoundFrontend/Private/MetasoundFrontendDocumentBuilder.cpp
+++ b/Engine/Plugins/Runtime/Metasound/Source/MetasoundFrontend/Private/MetasoundFrontendDocumentBuilder.cpp
@@ -296,21 +296,49 @@ namespace Metasound::Frontend
 					}
 				}
 
+				auto SetNodeID = [this](FMetasoundFrontendClassVertex& InVertex, bool bIsInput)
+				{
+					FGuid VertexID = FDocumentIDGenerator::Get().CreateVertexID(Document);
+					FGuid NodeID = FDocumentIDGenerator::Get().CreateNodeID(Document);
+
+					if (Options.ReferencedBuilder)
+					{
+						if (bIsInput)
+						{
+							const FMetasoundFrontendClassInput* GraphInput = Options.ReferencedBuilder->FindGraphInput(InVertex.Name);
+							if (GraphInput && GraphInput->TypeName == InVertex.TypeName)
+							{
+								NodeID = GraphInput->NodeID;
+								VertexID = GraphInput->VertexID;
+							}
+						}
+						else // Output
+						{
+							const FMetasoundFrontendClassOutput* GraphOutput = Options.ReferencedBuilder->FindGraphOutput(InVertex.Name);
+							if (GraphOutput && GraphOutput->TypeName == InVertex.TypeName)
+							{
+								NodeID = GraphOutput->NodeID;
+								VertexID = GraphOutput->VertexID;
+							}
+						}
+					}
+					InVertex.NodeID = NodeID;
+					InVertex.VertexID = VertexID;
+				};
+
 				for (const FMetasoundFrontendInterface& ToInterface : Options.InterfacesToAdd)
 				{
-					Algo::Transform(ToInterface.Inputs, InputsToAdd, [this, &ToInterface](const FMetasoundFrontendClassInput& Input)
+					Algo::Transform(ToInterface.Inputs, InputsToAdd, [this, &ToInterface, &SetNodeID](const FMetasoundFrontendClassInput& Input)
 					{
 						FMetasoundFrontendClassInput NewInput = Input;
-						NewInput.NodeID = FDocumentIDGenerator::Get().CreateNodeID(Document);
-						NewInput.VertexID = FDocumentIDGenerator::Get().CreateVertexID(Document);
+						SetNodeID(NewInput, true);
 						return FInputInterfacePair { MoveTemp(NewInput), &ToInterface };
 					});
 
-					Algo::Transform(ToInterface.Outputs, OutputsToAdd, [this, &ToInterface](const FMetasoundFrontendClassOutput& Output)
+					Algo::Transform(ToInterface.Outputs, OutputsToAdd, [this, &ToInterface, &SetNodeID](const FMetasoundFrontendClassOutput& Output)
 					{
 						FMetasoundFrontendClassOutput NewOutput = Output;
-						NewOutput.NodeID = FDocumentIDGenerator::Get().CreateNodeID(Document);
-						NewOutput.VertexID = FDocumentIDGenerator::Get().CreateVertexID(Document);
+						SetNodeID(NewOutput, false);
 						return FOutputInterfacePair { MoveTemp(NewOutput), &ToInterface };
 					});
 				}
@@ -685,19 +713,21 @@ namespace Metasound::Frontend
 		static_assert(static_cast<uint32>(EInvalidEdgeReason::COUNT) == 5, "Potential missing case coverage for EInvalidEdgeReason");
 	}
 
-	FModifyInterfaceOptions::FModifyInterfaceOptions(const TArray<FMetasoundFrontendInterface>& InInterfacesToRemove, const TArray<FMetasoundFrontendInterface>& InInterfacesToAdd)
+	FModifyInterfaceOptions::FModifyInterfaceOptions(const TArray<FMetasoundFrontendInterface>& InInterfacesToRemove, const TArray<FMetasoundFrontendInterface>& InInterfacesToAdd, const FMetaSoundFrontendDocumentBuilder* InReferencedBuilder)
 		: InterfacesToRemove(InInterfacesToRemove)
 		, InterfacesToAdd(InInterfacesToAdd)
+		, ReferencedBuilder(InReferencedBuilder)
 	{
 	}
 
-	FModifyInterfaceOptions::FModifyInterfaceOptions(TArray<FMetasoundFrontendInterface>&& InInterfacesToRemove, TArray<FMetasoundFrontendInterface>&& InInterfacesToAdd)
+	FModifyInterfaceOptions::FModifyInterfaceOptions(TArray<FMetasoundFrontendInterface>&& InInterfacesToRemove, TArray<FMetasoundFrontendInterface>&& InInterfacesToAdd, const FMetaSoundFrontendDocumentBuilder* InReferencedBuilder)
 		: InterfacesToRemove(MoveTemp(InInterfacesToRemove))
 		, InterfacesToAdd(MoveTemp(InInterfacesToAdd))
+		, ReferencedBuilder(InReferencedBuilder)
 	{
 	}
 
-	FModifyInterfaceOptions::FModifyInterfaceOptions(const TArray<FMetasoundFrontendVersion>& InInterfaceVersionsToRemove, const TArray<FMetasoundFrontendVersion>& InInterfaceVersionsToAdd)
+	FModifyInterfaceOptions::FModifyInterfaceOptions(const TArray<FMetasoundFrontendVersion>& InInterfaceVersionsToRemove, const TArray<FMetasoundFrontendVersion>& InInterfaceVersionsToAdd, const FMetaSoundFrontendDocumentBuilder* InReferencedBuilder)
 	{
 		Algo::Transform(InInterfaceVersionsToRemove, InterfacesToRemove, [](const FMetasoundFrontendVersion& Version)
 		{
@@ -720,6 +750,7 @@ namespace Metasound::Frontend
 			}
 			return Interface;
 		});
+		ReferencedBuilder = InReferencedBuilder;
 	}
 } // namespace Metasound::Frontend
 
@@ -1716,7 +1747,7 @@ const FMetasoundFrontendNode* FMetaSoundFrontendDocumentBuilder::AddGraphVariabl
 	return nullptr;
 }
 
-bool FMetaSoundFrontendDocumentBuilder::AddInterface(FName InterfaceName, bool bAddUserModifiableInterfaceOnly)
+bool FMetaSoundFrontendDocumentBuilder::AddInterfaceInternal(FName InterfaceName, bool bAddUserModifiableInterfaceOnly, const FMetaSoundFrontendDocumentBuilder* ReferenceBuilder)
 {
 	using namespace Metasound::Frontend;
 
@@ -1760,7 +1791,7 @@ bool FMetaSoundFrontendDocumentBuilder::AddInterface(FName InterfaceName, bool b
 			// Remove old versions and add new interfaces simultaneously to support node reconnection
 			TArray<FMetasoundFrontendInterface> InterfacesToAdd;
 			InterfacesToAdd.Add(Entry->GetInterface());
-			FModifyInterfaceOptions Options(MoveTemp(PreviousVersionInterfaces), MoveTemp(InterfacesToAdd));
+			FModifyInterfaceOptions Options(MoveTemp(PreviousVersionInterfaces), MoveTemp(InterfacesToAdd), ReferenceBuilder);
 			return ModifyInterfaces(MoveTemp(Options));
 		}
 	}
@@ -1768,6 +1799,16 @@ bool FMetaSoundFrontendDocumentBuilder::AddInterface(FName InterfaceName, bool b
 	return false;
 }
 
+bool FMetaSoundFrontendDocumentBuilder::AddInterface(FName InterfaceName)
+{
+	return AddInterfaceInternal(InterfaceName);
+}
+
+bool FMetaSoundFrontendDocumentBuilder::AddInterface(FName InterfaceName, bool bAddUserModifiableInterfaceOnly, const FMetaSoundFrontendDocumentBuilder* ReferenceBuilder)
+{
+	return AddInterfaceInternal(InterfaceName, bAddUserModifiableInterfaceOnly, ReferenceBuilder);
+}
+
 const FMetasoundFrontendNode* FMetaSoundFrontendDocumentBuilder::AddGraphNode(const FMetasoundFrontendGraphClass& InGraphClass, FGuid InNodeID, const FGuid* InPageID)
 {
 	using FNodeRegistryKey = Metasound::Frontend::FNodeRegistryKey;
@@ -2311,6 +2352,7 @@ void FMetaSoundFrontendDocumentBuilder::ClearDocument(TSharedRef<Metasound::Fron
 #if WITH_EDITORONLY_DATA
 bool FMetaSoundFrontendDocumentBuilder::ClearMemberMetadata(const FGuid& InMemberID)
 {
+	DocumentDelegates->InterfaceDelegates.OnRemovingMemberMetadata.Broadcast(InMemberID);
 	return GetDocumentChecked().Metadata.MemberMetadata.Remove(InMemberID) > 0;
 }
 #endif // WITH_EDITORONLY_DATA
@@ -3142,10 +3184,10 @@ const FMetasoundFrontendVariable* FMetaSoundFrontendDocumentBuilder::FindGraphVa
 }
 
 #if WITH_EDITOR
-UMetaSoundFrontendMemberMetadata* FMetaSoundFrontendDocumentBuilder::FindMemberMetadata(const FGuid& InMemberID)
+UMetaSoundFrontendMemberMetadata* FMetaSoundFrontendDocumentBuilder::FindMemberMetadata(const FGuid& InMemberID) const
 {
-	FMetasoundFrontendDocument& Document = GetDocumentChecked();
-	TMap<FGuid, TObjectPtr<UMetaSoundFrontendMemberMetadata>>& LiteralMetadata = Document.Metadata.MemberMetadata;
+	const FMetasoundFrontendDocument& Document = GetConstDocumentChecked();
+	const TMap<FGuid, TObjectPtr<UMetaSoundFrontendMemberMetadata>>& LiteralMetadata = Document.Metadata.MemberMetadata;
 	TObjectPtr<UMetaSoundFrontendMemberMetadata> ToReturn = LiteralMetadata.FindRef(InMemberID);
 	return ToReturn;
 }
@@ -6619,6 +6661,7 @@ void FMetaSoundFrontendDocumentBuilder::SetMemberMetadata(UMetaSoundFrontendMemb
 	TMap<FGuid, TObjectPtr<UMetaSoundFrontendMemberMetadata>>& LiteralMetadata = GetDocumentChecked().Metadata.MemberMetadata;
 	LiteralMetadata.Remove(NewMetadata.MemberID);
 	LiteralMetadata.Add(NewMetadata.MemberID, &NewMetadata);
+	DocumentDelegates->InterfaceDelegates.OnMemberMetadataSet.Broadcast(NewMetadata.MemberID);
 }
 
 bool FMetaSoundFrontendDocumentBuilder::SetNodeComment(const FGuid& InNodeID, FString&& InNewComment, const FGuid* InPageID)
diff --git a/Engine/Plugins/Runtime/Metasound/Source/MetasoundFrontend/Private/MetasoundFrontendTransform.cpp b/Engine/Plugins/Runtime/Metasound/Source/MetasoundFrontend/Private/MetasoundFrontendTransform.cpp
--- a/Engine/Plugins/Runtime/Metasound/Source/MetasoundFrontend/Private/MetasoundFrontendTransform.cpp
+++ b/Engine/Plugins/Runtime/Metasound/Source/MetasoundFrontend/Private/MetasoundFrontendTransform.cpp
@@ -717,7 +717,8 @@ namespace Metasound
 
 #if WITH_EDITORONLY_DATA
 			// Cache off member metadata so it be can be readded if necessary after the graph is cleared 
-			const FMemberIDToMetadataMap CachedMemberMetadata = DocumentToTransform.Metadata.MemberMetadata;
+			FMemberIDToMetadataMap CachedMemberMetadata;
+			CacheMemberMetadata(InOutBuilderToTransform, InputsInheritingDefault, CachedMemberMetadata);
 #endif // WITH_EDITORONLY_DATA
 
 			FGuid PresetNodeID;
@@ -741,11 +742,12 @@ namespace Metasound
 
 			// Ensure preset interfaces match those found in referenced graph.  Referenced graph is assumed to be
 			// well-formed (i.e. all inputs/outputs/environment variables declared by interfaces are present, and
-			// of proper name & data type).
+			// of proper name & data type). Pass in parent builder to copy member ids so added member ids can be consistent
+			FMetasoundFrontendDocument& Document = const_cast<FMetasoundFrontendDocument&>(InOutBuilderToTransform.GetConstDocumentChecked());
 			const TSet<FMetasoundFrontendVersion>& RefInterfaceVersions = ParentDocument.Interfaces;
 			for (const FMetasoundFrontendVersion& Version : RefInterfaceVersions)
 			{
-				InOutBuilderToTransform.AddInterface(Version.Name, /*bAddUserModifiableInterfaceOnly = */false);
+				InOutBuilderToTransform.AddInterface(Version.Name, /*bAddUserModifiableInterfaceOnly = */false, ParentBuilder);
 			}
 
 			// Add referenced node
@@ -786,7 +788,60 @@ namespace Metasound
 				}
 			}
 		}
+		
+		void FRebuildPresetRootGraph::CacheMemberMetadata(FMetaSoundFrontendDocumentBuilder& InOutBuilderToTransform, const TSet<FName>& InInputsInheritingDefault, FMemberIDToMetadataMap& InOutCachedMemberMetadata) const
+		{
+			const FMetasoundFrontendDocument& ParentDocument = ParentBuilder->GetConstDocumentChecked();
+
+			auto CreateNewLiteral = [&](UMetaSoundFrontendMemberMetadata* TemplateObject, FGuid MemberID) -> UMetaSoundFrontendMemberMetadata*
+			{
+				if (TemplateObject)
+				{
+					UMetaSoundFrontendMemberMetadata* NewMetadata = NewObject<UMetaSoundFrontendMemberMetadata>(&InOutBuilderToTransform.CastDocumentObjectChecked<UObject>(), TemplateObject->GetClass(), FName(), RF_Transactional, TemplateObject);
+					check(NewMetadata);
+					NewMetadata->MemberID = MemberID;
+					return NewMetadata;
+				}
+				return nullptr;
+			};
+			
+			for (const FMetasoundFrontendClassInput& ParentClassInput : ParentDocument.RootGraph.GetDefaultInterface().Inputs)
+			{
+				UMetaSoundFrontendMemberMetadata* MemberMetadata = nullptr;
+				// Member id is id in InOutBuilder, which may not be the same as the node id of the class input of ParentBuilder
+				FGuid MemberID;
+				const FMetasoundFrontendClassInput* GraphInput = InOutBuilderToTransform.FindGraphInput(ParentClassInput.Name);
+
+				if (GraphInput && GraphInput->TypeName == ParentClassInput.TypeName)
+				{
+					MemberID = GraphInput->NodeID;
+					// If the input vertex already exists in the parent graph,
+					// check if parent should be used or not from set of managed
+					// input names.
+					if (InInputsInheritingDefault.Contains(ParentClassInput.Name))
+					{
+						UMetaSoundFrontendMemberMetadata* ParentMetadata = ParentBuilder->FindMemberMetadata(ParentClassInput.NodeID);
+						MemberMetadata = CreateNewLiteral(ParentMetadata, MemberID);
+					}
+					else
+					{
+						// Use existing defaults 
+						MemberMetadata = InOutBuilderToTransform.FindMemberMetadata(GraphInput->NodeID);
+					}
+				}
+				else
+				{
+					UMetaSoundFrontendMemberMetadata* ParentMetadata = ParentBuilder->FindMemberMetadata(ParentClassInput.NodeID);
+					MemberID = ParentClassInput.NodeID;
+					MemberMetadata = CreateNewLiteral(ParentMetadata, MemberID);
+				}
 
+				if (MemberMetadata)
+				{
+					InOutCachedMemberMetadata.Emplace(MemberID, MemberMetadata);
+				}
+			}
+		}
 #endif // WITH_EDITORONLY_DATA
 
 		void FRebuildPresetRootGraph::AddAndConnectInputs(const TArray<FMetasoundFrontendClassInput>& InClassInputs, FMetaSoundFrontendDocumentBuilder& InOutBuilderToTransform, const FGuid& InReferencedNodeID) const
@@ -797,20 +852,36 @@ namespace Metasound
 
 			TArray<const FMetasoundFrontendVertex*> ReferencedNodeInputVertices;
 			TArray<const FMetasoundFrontendNode*> InputTemplateNodes;
-
+			const TSet<FName>* InputsInheritingDefault = InOutBuilderToTransform.GetGraphInputsInheritingDefault();
+			check(InputsInheritingDefault);
+			
 			for (const FMetasoundFrontendClassInput& ClassInput : InClassInputs)
 			{
-				// Input may have already been added if interface member
-				const FMetasoundFrontendNode* InputNode = InOutBuilderToTransform.FindGraphInputNode(ClassInput.Name);
+				const FName InputName = ClassInput.Name;
+				// Input may have already been added if interface member, but defaults need to be overridden 
+				const FMetasoundFrontendNode* InputNode = InOutBuilderToTransform.FindGraphInputNode(InputName);
 				if (!InputNode)
 				{
 					InputNode = InOutBuilderToTransform.AddGraphInput(ClassInput);
 				}
+				else
+				{
+					// Defaults must be set as either the parent default or existing default may be different than the registered interface default
+					// Setting defaults will set the input as non inheriting even if we're just copying the value 
+					// so cache off bool and set back if necessary
+					const bool bInheritsDefault = InputsInheritingDefault->Contains(InputName);
+					InOutBuilderToTransform.SetGraphInputDefaults(InputName, ClassInput.GetDefaults());
+					if (bInheritsDefault)
+					{
+						InOutBuilderToTransform.SetGraphInputInheritsDefault(InputName, /*bInputInheritsDefault=*/true);
+					}
+				}
+
 				const FGuid InputNodeID = InputNode->GetID();
 				check(InputNode);
-				const FMetasoundFrontendVertex* InputNodeOutputVertex = InOutBuilderToTransform.FindNodeOutput(InputNode->GetID(), ClassInput.Name);
+				const FMetasoundFrontendVertex* InputNodeOutputVertex = InOutBuilderToTransform.FindNodeOutput(InputNode->GetID(), InputName);
 				check(InputNodeOutputVertex);
-				const FMetasoundFrontendVertex* ReferencedNodeInputVertex = InOutBuilderToTransform.FindNodeInput(InReferencedNodeID, ClassInput.Name);
+				const FMetasoundFrontendVertex* ReferencedNodeInputVertex = InOutBuilderToTransform.FindNodeInput(InReferencedNodeID, InputName);
 				check(ReferencedNodeInputVertex);
 				
 				ReferencedNodeInputVertices.Add(ReferencedNodeInputVertex);
@@ -948,6 +1019,7 @@ namespace Metasound
 			const FMetasoundFrontendGraphClass& ParentRootGraph = ParentBuilder->GetConstDocumentChecked().RootGraph;
 			const FMetasoundFrontendDocument& ParentDocument = ParentBuilder->GetConstDocumentChecked();
 			const FMetasoundFrontendDocument& DocumentToTransform = InDocumentToTransformBuilder.GetConstDocumentChecked();
+
 			const TSet<FName>* ExistingInputsInheritingDefaultPtr = InDocumentToTransformBuilder.GetGraphInputsInheritingDefault();
 			check(ExistingInputsInheritingDefaultPtr);
 			TSet<FName> InputsInheritingDefault = *ExistingInputsInheritingDefaultPtr;
diff --git a/Engine/Plugins/Runtime/Metasound/Source/MetasoundFrontend/Public/MetasoundFrontendDocumentBuilder.h b/Engine/Plugins/Runtime/Metasound/Source/MetasoundFrontend/Public/MetasoundFrontendDocumentBuilder.h
--- a/Engine/Plugins/Runtime/Metasound/Source/MetasoundFrontend/Public/MetasoundFrontendDocumentBuilder.h
+++ b/Engine/Plugins/Runtime/Metasound/Source/MetasoundFrontend/Public/MetasoundFrontendDocumentBuilder.h
@@ -78,9 +78,9 @@ namespace Metasound::Frontend
 
 	struct FModifyInterfaceOptions
 	{
-		UE_API FModifyInterfaceOptions(const TArray<FMetasoundFrontendInterface>& InInterfacesToRemove, const TArray<FMetasoundFrontendInterface>& InInterfacesToAdd);
-		UE_API FModifyInterfaceOptions(TArray<FMetasoundFrontendInterface>&& InInterfacesToRemove, TArray<FMetasoundFrontendInterface>&& InInterfacesToAdd);
-		UE_API FModifyInterfaceOptions(const TArray<FMetasoundFrontendVersion>& InInterfaceVersionsToRemove, const TArray<FMetasoundFrontendVersion>& InInterfaceVersionsToAdd);
+		UE_API FModifyInterfaceOptions(const TArray<FMetasoundFrontendInterface>& InInterfacesToRemove, const TArray<FMetasoundFrontendInterface>& InInterfacesToAdd, const FMetaSoundFrontendDocumentBuilder* InReferencedBuilder =nullptr);
+		UE_API FModifyInterfaceOptions(TArray<FMetasoundFrontendInterface>&& InInterfacesToRemove, TArray<FMetasoundFrontendInterface>&& InInterfacesToAdd, const FMetaSoundFrontendDocumentBuilder* InReferencedBuilder = nullptr);
+		UE_API FModifyInterfaceOptions(const TArray<FMetasoundFrontendVersion>& InInterfaceVersionsToRemove, const TArray<FMetasoundFrontendVersion>& InInterfaceVersionsToAdd, const FMetaSoundFrontendDocumentBuilder* InReferencedBuilder = nullptr);
 
 		TArray<FMetasoundFrontendInterface> InterfacesToRemove;
 		TArray<FMetasoundFrontendInterface> InterfacesToAdd;
@@ -94,6 +94,9 @@ namespace Metasound::Frontend
 #if WITH_EDITORONLY_DATA
 		bool bSetDefaultNodeLocations = true;
 #endif // WITH_EDITORONLY_DATA
+
+		// Optional builder to copy node ids from instead of regenerating
+		const FMetaSoundFrontendDocumentBuilder* ReferencedBuilder = nullptr;
 	};
 
 #if WITH_EDITORONLY_DATA
@@ -213,8 +216,12 @@ public:
 	UE_API const FMetasoundFrontendNode* AddGraphOutput(FMetasoundFrontendClassOutput ClassOutput, const FGuid* InPageID = nullptr);
 
 	// Adds an interface and associated inputs and outputs
+	UE_API bool AddInterface(FName InterfaceName);
+
+	// Internal AddInterface call for use in rebuild preset root graph
 	// If bAddUserModifiableInterfaceOnly is true, only user modifiable interfaces can be added (see FMetasoundFrontendInterfaceUClassOptions)
-	UE_API bool AddInterface(FName InterfaceName, bool bAddUserModifiableInterfaceOnly = true);
+	// Reference builder is used to copy node ids for added members so they aren't just regenerated with NewGuid
+	UE_INTERNAL bool AddInterface(FName InterfaceName, bool bAddUserModifiableInterfaceOnly, const FMetaSoundFrontendDocumentBuilder* ReferenceBuilder);
 
 	UE_API const FMetasoundFrontendNode* AddNodeByClassName(const FMetasoundFrontendClassName& InClassName, int32 InMajorVersion = 1, FGuid InNodeID = FGuid::NewGuid(), const FGuid* InPageID = nullptr);
 
@@ -305,7 +312,7 @@ PRAGMA_ENABLE_EXPERIMENTAL_WARNINGS
 	UE_API const FMetaSoundFrontendGraphComment* FindGraphComment(const FGuid& InCommentID, const FGuid* InPageID = nullptr) const;
 	UE_API FMetaSoundFrontendGraphComment* FindGraphComment(const FGuid& InCommentID, const FGuid* InPageID = nullptr);
 	UE_API FMetaSoundFrontendGraphComment& FindOrAddGraphComment(const FGuid& InCommentID, const FGuid* InPageID = nullptr);
-	UE_API UMetaSoundFrontendMemberMetadata* FindMemberMetadata(const FGuid& InMemberID);
+	UE_API UMetaSoundFrontendMemberMetadata* FindMemberMetadata(const FGuid& InMemberID) const;
 #endif // WITH_EDITORONLY_DATA
 
 	static UE_API bool FindDeclaredInterfaces(const FMetasoundFrontendDocument& InDocument, TArray<const Metasound::Frontend::IInterfaceRegistryEntry*>& OutInterfaces);
@@ -751,6 +758,9 @@ private:
 
 	TInstancedStruct<FMetaSoundFrontendNodeConfiguration> CreateFrontendNodeConfiguration(const FMetasoundFrontendClassMetadata& InClassMetadata) const;
 	TInstancedStruct<FMetaSoundFrontendNodeConfiguration> CreateFrontendNodeConfiguration(EMetasoundFrontendClassType InClassType, const Metasound::Frontend::FNodeRegistryKey& InClassKey) const;
+	
+	bool AddInterfaceInternal(FName InterfaceName, bool bAddUserModifiableInterfaceOnly = true, const FMetaSoundFrontendDocumentBuilder* ReferenceBuilder = nullptr);
+
 	FMetasoundFrontendNode* AddNodeInternal(const FMetasoundFrontendClassMetadata& InClassMetadata, Metasound::Frontend::FFinalizeNodeFunctionRef FinalizeNode, const FGuid& InPageID, FGuid InNodeID = FGuid::NewGuid(), int32* NewNodeIndex = nullptr);
 	FMetasoundFrontendNode* AddNodeInternal(const Metasound::Frontend::FNodeRegistryKey& InClassKey, FFinalizeNodeFunctionRef FinalizeNode, const FGuid& InPageID, FGuid InNodeID = FGuid::NewGuid(), int32* NewNodeIndex = nullptr);
 
@@ -830,6 +840,8 @@ private:
 	FMetasoundFrontendDocument& GetDocumentChecked() const;
 	IMetaSoundDocumentInterface& GetDocumentInterfaceChecked() const;
 
+	bool ModifyInterfacesInternal(Metasound::Frontend::FModifyInterfaceOptions&& InOptions, const FMetaSoundFrontendDocumentBuilder* InBuilder = nullptr);
+
 	void RemoveSwapDependencyInternal(int32 Index);
 
 	bool SpliceVariableNodeFromStack(const FGuid& InNodeID, const FGuid& InPageID);
diff --git a/Engine/Plugins/Runtime/Metasound/Source/MetasoundFrontend/Public/MetasoundFrontendDocumentModifyDelegates.h b/Engine/Plugins/Runtime/Metasound/Source/MetasoundFrontend/Public/MetasoundFrontendDocumentModifyDelegates.h
--- a/Engine/Plugins/Runtime/Metasound/Source/MetasoundFrontend/Public/MetasoundFrontendDocumentModifyDelegates.h
+++ b/Engine/Plugins/Runtime/Metasound/Source/MetasoundFrontend/Public/MetasoundFrontendDocumentModifyDelegates.h
@@ -17,6 +17,10 @@ DECLARE_MULTICAST_DELEGATE_TwoParams(FOnMetaSoundFrontendDocumentRenameClass, co
 DECLARE_MULTICAST_DELEGATE_ThreeParams(FOnMetaSoundFrontendDocumentMutateNodeInputLiteralArray, int32 /* NodeIndex */, int32 /* VertexIndex */, int32 /* LiteralIndex */);
 DECLARE_MULTICAST_DELEGATE_TwoParams(FOnMetaSoundFrontendDocumentRenameVertex, FName /* OldName */, FName /* NewName */);
 
+#if WITH_EDITOR
+DECLARE_MULTICAST_DELEGATE_OneParam(FOnMetaSoundFrontendDocumentMutateMember, FGuid /* Member ID*/);
+#endif // WITH_EDITOR
+
 namespace Metasound::Frontend
 {
 	struct FDocumentMutatePageArgs
@@ -71,6 +75,9 @@ namespace Metasound::Frontend
 		FOnMetaSoundFrontendDocumentMutateArray OnOutputDescriptionChanged;
 		FOnMetaSoundFrontendDocumentMutateArray OnOutputSortOrderIndexChanged;
 		FOnMetaSoundFrontendDocumentMutateArray OnOutputIsAdvancedDisplayChanged;
+
+		FOnMetaSoundFrontendDocumentMutateMember OnMemberMetadataSet;
+		FOnMetaSoundFrontendDocumentMutateMember OnRemovingMemberMetadata;
 #endif // WITH_EDITOR
 	};
 
diff --git a/Engine/Plugins/Runtime/Metasound/Source/MetasoundFrontend/Public/MetasoundFrontendTransform.h b/Engine/Plugins/Runtime/Metasound/Source/MetasoundFrontend/Public/MetasoundFrontendTransform.h
--- a/Engine/Plugins/Runtime/Metasound/Source/MetasoundFrontend/Public/MetasoundFrontendTransform.h
+++ b/Engine/Plugins/Runtime/Metasound/Source/MetasoundFrontend/Public/MetasoundFrontendTransform.h
@@ -190,6 +190,7 @@ namespace Metasound::Frontend
 #if WITH_EDITORONLY_DATA
 		using FMemberIDToMetadataMap = TMap<FGuid, TObjectPtr<UMetaSoundFrontendMemberMetadata>>;
 		void AddMemberMetadata(const FMemberIDToMetadataMap& InCachedMemberMetadata, FMetaSoundFrontendDocumentBuilder& InOutBuilderToTransform) const;
+		void CacheMemberMetadata(FMetaSoundFrontendDocumentBuilder& InOutBuilderToTransform, const TSet<FName>& InInputsInheritingDefault, FMemberIDToMetadataMap& InOutCachedMemberMetadata) const;
 #endif // WITH_EDITORONLY_DATA
 
 		const FMetaSoundFrontendDocumentBuilder* ParentBuilder = nullptr;
