Fix for bus reader node getting stopped while metasound is playing and reading from the bus

diff --git a/Engine/Plugins/Runtime/Metasound/Source/MetasoundEngine/Private/MetasoundAudioBusReaderNode.cpp b/Engine/Plugins/Runtime/Metasound/Source/MetasoundEngine/Private/MetasoundAudioBusReaderNode.cpp
--- a/Engine/Plugins/Runtime/Metasound/Source/MetasoundEngine/Private/MetasoundAudioBusReaderNode.cpp
+++ b/Engine/Plugins/Runtime/Metasound/Source/MetasoundEngine/Private/MetasoundAudioBusReaderNode.cpp
@@ -3,6 +3,7 @@
 #include "AudioBusSubsystem.h"
 #include "AudioDevice.h"
 #include "DSP/AlignedBuffer.h"
+#include "DSP/FloatArrayMath.h"
 #include "Internationalization/Text.h"
 #include "Sound/AudioBus.h"
 
@@ -201,21 +202,50 @@ namespace Metasound
 				SamplesPopped = AudioBusPatchOutput->PopAudio(InterleavedBuffer.GetData(), NumSamplesToPop, false);
 			}
 
-			if (SamplesPopped < NumSamplesToPop && !bWasUnderrunReported)
-			{
-				UE_LOG(LogMetaSound, Warning, TEXT("Underrun detected in audio bus reader node."));
-				bWasUnderrunReported = true;
-			}
-
 			const uint32 MinChannels = FMath::Min(NumChannels, AudioBusChannels);
 			for (uint32 ChannelIndex = 0; ChannelIndex < MinChannels; ++ChannelIndex)
 			{
 				float* AudioOutputBufferPtr = AudioOutputs[ChannelIndex]->GetData();
-				for (int32 FrameIndex = 0; FrameIndex < BlockSizeFrames; ++FrameIndex)
+				for (int32 FrameIndex = 0; FrameIndex < SamplesPopped; ++FrameIndex)
 				{
 					AudioOutputBufferPtr[FrameIndex] = InterleavedBuffer[FrameIndex * AudioBusChannels + ChannelIndex];
 				}
 			}
+
+			if (SamplesPopped < NumSamplesToPop)
+			{
+				if (!bWasUnderrunReported)
+				{
+					UE_LOG(LogMetaSound, Warning, TEXT("Underrun detected in audio bus reader node. Check to see if bus was stopped."));
+					bWasUnderrunReported = true;
+				}
+
+				// If we have popped some samples then use those and fade out to prevent discontinutiies
+				if (SamplesPopped > 0)
+				{
+					int32 NumFramesToZero = BlockSizeFrames - SamplesPopped;
+
+					for (uint32 ChannelIndex = 0; ChannelIndex < MinChannels; ++ChannelIndex)
+					{
+						float* AudioOutputBufferPtr = AudioOutputs[ChannelIndex]->GetData();
+						TArrayView<float> BufferToFade = TArrayView<float>(AudioOutputBufferPtr, SamplesPopped);
+						Audio::ArrayFade(BufferToFade, 1.0f, 0.0f);
+
+
+						// zero out the remainder of the buffer
+						FMemory::Memzero(AudioOutputBufferPtr + SamplesPopped, sizeof(float) * NumFramesToZero);
+					}
+				}
+				// If we have not popped any, just make sure our buffers are zero'd
+				else
+				{
+					for (uint32 ChannelIndex = 0; ChannelIndex < MinChannels; ++ChannelIndex)
+					{
+						AudioOutputs[ChannelIndex]->Zero();
+					}
+				}
+
+			}
 		}
 
 		void CreatePatchOutput()
