Fix for math operations for Reset.

Audio-float multiply was incorrectly initializing it's initial float value and interpolating on first execute.
Added Reset() pathway for math ops. Implemented Calcualte for all but Audio-float multiply

#rn Fix of MetaSound audio-float multiply node resulting in a first-execute pop.

diff --git a/Engine/Plugins/Runtime/Metasound/Source/MetasoundStandardNodes/Private/MetasoundMathNodes.cpp b/Engine/Plugins/Runtime/Metasound/Source/MetasoundStandardNodes/Private/MetasoundMathNodes.cpp
--- a/Engine/Plugins/Runtime/Metasound/Source/MetasoundStandardNodes/Private/MetasoundMathNodes.cpp
+++ b/Engine/Plugins/Runtime/Metasound/Source/MetasoundStandardNodes/Private/MetasoundMathNodes.cpp
@@ -239,9 +239,10 @@ namespace Metasound
 			return Info;
 		}
 
+
 		void Reset(const IOperator::FResetParams& InParams)
 		{
-			TMathOpClass::Calculate(InstanceData, PrimaryOperandRef, AdditionalOperandRefs, ValueRef);
+			TMathOpClass::Reset(InstanceData, PrimaryOperandRef, AdditionalOperandRefs, ValueRef);
 		}
 
 		void Execute()
@@ -334,6 +335,11 @@ namespace Metasound
 
 			return Interface;
 		}
+	
+		static void Reset(TMathOpAdd<TDataClass, TOperandDataClass>& InInstanceData, const TDataClassReadRef& InPrimaryOperand, const TArray<TOperandDataClassReadRef>& InAdditionalOperands, TDataClassWriteRef& OutResult)
+		{
+			Calculate(InInstanceData, InPrimaryOperand, InAdditionalOperands, OutResult);
+		}
 
 		static void Calculate(TMathOpAdd<TDataClass, TOperandDataClass>& InInstanceData, const TDataClassReadRef& InPrimaryOperand, const TArray<TOperandDataClassReadRef>& InAdditionalOperands, TDataClassWriteRef& OutResult)
 		{
@@ -393,7 +399,12 @@ namespace Metasound
 
 			return Interface;
 		}
-
+		
+		static void Reset(TMathOpSubtract<TDataClass, TOperandDataClass>& InInstanceData, const TDataClassReadRef& InPrimaryOperand, const TArray<TOperandDataClassReadRef>& InAdditionalOperands, TDataClassWriteRef& OutResult)
+		{
+			Calculate(InInstanceData, InPrimaryOperand, InAdditionalOperands, OutResult);
+		}
+		
 		static void Calculate(TMathOpSubtract<TDataClass, TOperandDataClass>& InInstanceData, const TDataClassReadRef& InPrimaryOperand, const TArray<TOperandDataClassReadRef>& InAdditionalOperands, TDataClassWriteRef& OutResult)
 		{
 			if (!ensure(!InAdditionalOperands.IsEmpty()))
@@ -453,6 +464,11 @@ namespace Metasound
 			return Interface;
 		}
 
+		static void Reset(TMathOpMultiply<TDataClass, TOperandDataClass>& InInstanceData, const TDataClassReadRef& InPrimaryOperand, const TArray<TOperandDataClassReadRef>& InAdditionalOperands, TDataClassWriteRef& OutResult)
+		{
+			Calculate(InInstanceData, InPrimaryOperand, InAdditionalOperands, OutResult);
+		}
+
 		static void Calculate(TMathOpMultiply<TDataClass, TOperandDataClass>& InInstanceData, const TDataClassReadRef& InPrimaryOperand, const TArray<TOperandDataClassReadRef>& InAdditionalOperands, TDataClassWriteRef& OutResult)
 		{
 			if (!ensure(!InAdditionalOperands.IsEmpty()))
@@ -512,6 +528,11 @@ namespace Metasound
 			return Interface;
 		}
 
+		static void Reset(TMathOpDivide<TDataClass, TOperandDataClass>& InInstanceData, const TDataClassReadRef& InPrimaryOperand, const TArray<TOperandDataClassReadRef>& InAdditionalOperands, TDataClassWriteRef& OutResult)
+		{
+			Calculate(InInstanceData, InPrimaryOperand, InAdditionalOperands, OutResult);
+		}
+
 		static void Calculate(TMathOpDivide& InInstanceData, const TDataClassReadRef& InPrimaryOperand, const TArray<TOperandDataClassReadRef>& InAdditionalOperands, TDataClassWriteRef& OutResult)
 		{
 			if (!ensure(!InAdditionalOperands.IsEmpty()))
@@ -578,6 +599,11 @@ namespace Metasound
 			return Interface;
 		}
 
+		static void Reset(TMathOpModulo<TDataClass, TOperandDataClass>& InInstanceData, const TDataClassReadRef& InPrimaryOperand, const TArray<TOperandDataClassReadRef>& InAdditionalOperands, TDataClassWriteRef& OutResult)
+		{
+			Calculate(InInstanceData, InPrimaryOperand, InAdditionalOperands, OutResult);
+		}
+
 		static void Calculate(TMathOpModulo& InInstanceData, const TDataClassReadRef& InPrimaryOperand, const TArray<TOperandDataClassReadRef>& InAdditionalOperands, TDataClassWriteRef& OutResult)
 		{
 			if (!ensure(!InAdditionalOperands.IsEmpty()))
@@ -646,6 +672,11 @@ namespace Metasound
 			return DefaultInterface;
 		}
 
+		static void Reset(TMathOpPower<TDataClass, TOperandDataClass>& InInstanceData, const TDataClassReadRef& InPrimaryOperand, const TArray<TOperandDataClassReadRef>& InAdditionalOperands, TDataClassWriteRef& OutResult)
+		{
+			Calculate(InInstanceData, InPrimaryOperand, InAdditionalOperands, OutResult);
+		}
+
 		static void Calculate(TMathOpPower& InInstanceData, const TDataClassReadRef& InPrimaryOperand, const TArray<TOperandDataClassReadRef>& InAdditionalOperands, TDataClassWriteRef& OutResult)
 		{
 			if (!ensure(!InAdditionalOperands.IsEmpty()))
@@ -708,6 +739,11 @@ namespace Metasound
 			return DefaultInterface;
 		}
 
+		static void Reset(TMathOpLogarithm<TDataClass, TOperandDataClass>& InInstanceData, const TDataClassReadRef& InPrimaryOperand, const TArray<TOperandDataClassReadRef>& InAdditionalOperands, TDataClassWriteRef& OutResult)
+		{
+			Calculate(InInstanceData, InPrimaryOperand, InAdditionalOperands, OutResult);
+		}
+
 		static void Calculate(TMathOpLogarithm& InInstanceData, const TDataClassReadRef& InPrimaryOperand, const TArray<TOperandDataClassReadRef>& InAdditionalOperands, TDataClassWriteRef& OutResult)
 		{
 			if (!ensure(!InAdditionalOperands.IsEmpty()))
@@ -761,7 +797,7 @@ namespace Metasound
 			return Interface;
 		}
 
-		static void Calculate(TMathOpAdd<FAudioBuffer>& InInstanceData, FAudioBufferReadRef& InPrimaryOperand, const TArray<FAudioBufferReadRef>& InAdditionalOperands, FAudioBufferWriteRef& OutResult)
+		static void Calculate(TMathOpAdd<FAudioBuffer>& InInstanceData, const FAudioBufferReadRef& InPrimaryOperand, const TArray<FAudioBufferReadRef>& InAdditionalOperands, FAudioBufferWriteRef& OutResult)
 		{
 			if (!ensure(!InAdditionalOperands.IsEmpty()))
 			{
@@ -804,6 +840,12 @@ namespace Metasound
 				}
 			}
 		}
+
+		static void Reset(TMathOpAdd<FAudioBuffer>& InInstanceData, const FAudioBufferReadRef& InPrimaryOperand, const TArray<FAudioBufferReadRef>& InAdditionalOperands, FAudioBufferWriteRef& OutResult)
+		{
+			Calculate(InInstanceData, InPrimaryOperand, InAdditionalOperands, OutResult);
+		}
+
 	};
 
 	template <>
@@ -858,6 +900,12 @@ namespace Metasound
 				*OutResult += *InAdditionalOperands[i];
 			}
 		}
+
+		static void Reset(TMathOpAdd& InInstanceData, const FTimeReadRef& InPrimaryOperand, const TArray<FTimeReadRef>& InAdditionalOperands, FTimeWriteRef& OutResult)
+		{
+			Calculate(InInstanceData, InPrimaryOperand, InAdditionalOperands, OutResult);
+		}
+
 	};
 
 	template <>
@@ -904,6 +952,12 @@ namespace Metasound
 				Audio::ArrayAddConstantInplace(OutResultView, *InAdditionalOperands[i]);
 			}
 		}
+
+		static void Reset(TMathOpAdd& InInstanceData, const FAudioBufferReadRef& InPrimaryOperand, const TArray<FFloatReadRef>& InAdditionalOperands, FAudioBufferWriteRef& OutResult)
+		{
+			Calculate(InInstanceData, InPrimaryOperand, InAdditionalOperands, OutResult);
+		}
+
 	};
 
 	template <>
@@ -939,7 +993,7 @@ namespace Metasound
 			return Interface;
 		}
 
-		static void Calculate(TMathOpSubtract& InInstanceData, FAudioBufferReadRef& InPrimaryOperand, const TArray<FAudioBufferReadRef>& InAdditionalOperands, FAudioBufferWriteRef& OutResult)
+		static void Calculate(TMathOpSubtract& InInstanceData, const FAudioBufferReadRef& InPrimaryOperand, const TArray<FAudioBufferReadRef>& InAdditionalOperands, FAudioBufferWriteRef& OutResult)
 		{
 			if (!ensure(!InAdditionalOperands.IsEmpty()))
 			{
@@ -973,6 +1027,12 @@ namespace Metasound
 				Audio::ArraySubtractInPlace2(OutDataView, OperandDataView);
 			}
 		}
+
+		static void Reset(TMathOpSubtract& InInstanceData, const FAudioBufferReadRef& InPrimaryOperand, const TArray<FAudioBufferReadRef>& InAdditionalOperands, FAudioBufferWriteRef& OutResult)
+		{
+			Calculate(InInstanceData, InPrimaryOperand, InAdditionalOperands, OutResult);
+		}
+
 	};
 
 	template <>
@@ -1027,6 +1087,12 @@ namespace Metasound
 				*OutResult -= *InAdditionalOperands[i];
 			}
 		}
+
+		static void Reset(TMathOpSubtract& InInstanceData, const FTimeReadRef& InPrimaryOperand, const TArray<FTimeReadRef>& InAdditionalOperands, FTimeWriteRef& OutResult)
+		{
+			Calculate(InInstanceData, InPrimaryOperand, InAdditionalOperands, OutResult);
+		}
+
 	};
 
 	template <>
@@ -1115,6 +1181,12 @@ namespace Metasound
 				Audio::ArrayMultiplyInPlace(OperandDataView, OutDataView);
 			}
 		}
+
+		static void Reset(TMathOpMultiply& InInstanceData, const FAudioBufferReadRef& InPrimaryOperand, const TArray<FAudioBufferReadRef>& InAdditionalOperands, FAudioBufferWriteRef& OutResult)
+		{
+			Calculate(InInstanceData, InPrimaryOperand, InAdditionalOperands, OutResult);
+		}
+
 	};
 
 	template <>
@@ -1183,6 +1255,13 @@ namespace Metasound
 
 			InInstanceData.LastGain = NewGain;
 		}
+
+		static void Reset(TMathOpMultiply& InInstanceData, const FAudioBufferReadRef& InPrimaryOperand, const TArray<FFloatReadRef>& InAdditionalOperands, FAudioBufferWriteRef& OutResult)
+		{
+			InInstanceData.bInit = false;
+			OutResult->Zero();
+		}
+
 	};
 
 	template <>
@@ -1239,6 +1318,12 @@ namespace Metasound
 				*OutResult *= OperandValue;
 			}
 		}
+
+		static void Reset(TMathOpMultiply& InInstanceData, const FTimeReadRef& InPrimaryOperand, const TArray<FFloatReadRef>& InAdditionalOperands, FTimeWriteRef& OutResult)
+		{
+			Calculate(InInstanceData, InPrimaryOperand, InAdditionalOperands, OutResult);
+		}
+
 	};
 
 	template <>
@@ -1301,6 +1386,12 @@ namespace Metasound
 				*OutResult /= OperandValue;
 			}
 		}
+
+		static void Reset(TMathOpDivide& InInstanceData, const FTimeReadRef& InPrimaryOperand, const TArray<FFloatReadRef>& InAdditionalOperands, FTimeWriteRef& OutResult)
+		{
+			Calculate(InInstanceData, InPrimaryOperand, InAdditionalOperands, OutResult);
+		}
+
 	};
 
 	template <typename TDerivedClass, typename TMathOpClass, typename TDataClass, typename TOperandDataClass>
