Fix for buffer sample vs frame issue

diff --git a/Engine/Plugins/Runtime/Metasound/Source/MetasoundEngine/Private/MetasoundAudioBusReaderNode.cpp b/Engine/Plugins/Runtime/Metasound/Source/MetasoundEngine/Private/MetasoundAudioBusReaderNode.cpp
--- a/Engine/Plugins/Runtime/Metasound/Source/MetasoundEngine/Private/MetasoundAudioBusReaderNode.cpp
+++ b/Engine/Plugins/Runtime/Metasound/Source/MetasoundEngine/Private/MetasoundAudioBusReaderNode.cpp
@@ -184,15 +184,26 @@ namespace Metasound
 				}
 			}
 
+			// If still not ready, output silence and bail
 			if (InterleavedBuffer.IsEmpty())
 			{
+				for (uint32 ChannelIndex = 0; ChannelIndex < NumChannels; ++ChannelIndex)
+				{
+					AudioOutputs[ChannelIndex]->Zero();
+				}
 				return;
 			}
 
 			// Pop off the interleaved data from the audio bus
-			int32 NumSamplesToPop = BlockSizeFrames * AudioBusChannels;
+			int32 FramesToPop = BlockSizeFrames;
+			int32 NumSamplesToPop = FramesToPop * AudioBusChannels;
 			int32 SamplesPopped = -1;
 
+			if (InterleavedBuffer.Num() < NumSamplesToPop)
+			{
+				InterleavedBuffer.SetNumUninitialized(NumSamplesToPop);
+			}
+
 			if (ResampledPatchOutput.IsValid())
 			{
 				SamplesPopped = ResampledPatchOutput->PopAudio(InterleavedBuffer.GetData(), NumSamplesToPop, false);
@@ -203,15 +214,24 @@ namespace Metasound
 			}
 
 			const uint32 MinChannels = FMath::Min(NumChannels, AudioBusChannels);
-			for (uint32 ChannelIndex = 0; ChannelIndex < MinChannels; ++ChannelIndex)
+			SamplesPopped = FMath::Max(SamplesPopped, 0);
+
+			// Copy the data that was actually popped into the output buffer	
+			int32 FramesPopped = FMath::Clamp(SamplesPopped / AudioBusChannels, 0, FramesToPop);
+			if (FramesPopped > 0)
 			{
-				float* AudioOutputBufferPtr = AudioOutputs[ChannelIndex]->GetData();
-				for (int32 FrameIndex = 0; FrameIndex < SamplesPopped; ++FrameIndex)
+				for (uint32 ChannelIndex = 0; ChannelIndex < MinChannels; ++ChannelIndex)
 				{
-					AudioOutputBufferPtr[FrameIndex] = InterleavedBuffer[FrameIndex * AudioBusChannels + ChannelIndex];
+					float* RESTRICT AudioOutputBufferPtr = AudioOutputs[ChannelIndex]->GetData();
+					for (int32 FrameIndex = 0; FrameIndex < FramesPopped; ++FrameIndex)
+					{
+						const int32 InterleavedIdx = FrameIndex * AudioBusChannels + (int32)ChannelIndex;
+						AudioOutputBufferPtr[FrameIndex] = InterleavedBuffer[InterleavedIdx];
+					}
 				}
 			}
 
+			// If we popped less than we requested, do a fade out and zero remainder
 			if (SamplesPopped < NumSamplesToPop)
 			{
 				if (!bWasUnderrunReported)
@@ -220,20 +240,21 @@ namespace Metasound
 					bWasUnderrunReported = true;
 				}
 
-				// If we have popped some samples then use those and fade out to prevent discontinutiies
-				if (SamplesPopped > 0)
+				// If we have popped some frames then use those and fade out to prevent discontinutiies
+				if (FramesPopped > 0)
 				{
-					int32 NumFramesToZero = BlockSizeFrames - SamplesPopped;
-
+					const int32 NumFramesToZero = FramesToPop - FramesPopped;
 					for (uint32 ChannelIndex = 0; ChannelIndex < MinChannels; ++ChannelIndex)
 					{
-						float* AudioOutputBufferPtr = AudioOutputs[ChannelIndex]->GetData();
-						TArrayView<float> BufferToFade = TArrayView<float>(AudioOutputBufferPtr, SamplesPopped);
-						Audio::ArrayFade(BufferToFade, 1.0f, 0.0f);
-
+						float* RESTRICT AudioOutputBufferPtr = AudioOutputs[ChannelIndex]->GetData();
+						TArrayView<float> FadeView(AudioOutputBufferPtr, FramesPopped);
+						Audio::ArrayFade(FadeView, 1.0f, 0.0f);
 
 						// zero out the remainder of the buffer
-						FMemory::Memzero(AudioOutputBufferPtr + SamplesPopped, sizeof(float) * NumFramesToZero);
+						if (NumFramesToZero > 0)
+						{
+							FMemory::Memzero(AudioOutputBufferPtr + FramesPopped, sizeof(float) * NumFramesToZero);
+						}
 					}
 				}
 				// If we have not popped any, just make sure our buffers are zero'd
@@ -244,7 +265,19 @@ namespace Metasound
 						AudioOutputs[ChannelIndex]->Zero();
 					}
 				}
+			}
+			else
+			{
+				bWasUnderrunReported = false;
+			}
 
+			// Zero out channels higher than our min channels here
+			if (NumChannels > MinChannels)
+			{
+				for (uint32 ChannelIndex = MinChannels; ChannelIndex < NumChannels; ++ChannelIndex)
+				{
+					AudioOutputs[ChannelIndex]->Zero();
+				}
 			}
 		}
 
