Fix infinite recursion in MetaSounds interaction between builder delegates and preset member metadata that is now handled by preset transform
- Added protection for editor crash

diff --git a/Engine/Plugins/Runtime/Metasound/Source/MetasoundEditor/Private/MetasoundEditor.cpp b/Engine/Plugins/Runtime/Metasound/Source/MetasoundEditor/Private/MetasoundEditor.cpp
--- a/Engine/Plugins/Runtime/Metasound/Source/MetasoundEditor/Private/MetasoundEditor.cpp
+++ b/Engine/Plugins/Runtime/Metasound/Source/MetasoundEditor/Private/MetasoundEditor.cpp
@@ -1396,10 +1396,13 @@ namespace Metasound
 			if (Builder.IsValid())
 			{
 				const FMetaSoundFrontendDocumentBuilder& DocBuilder = Builder->GetConstBuilder();
-				const FMetasoundFrontendDocument& Document = DocBuilder.GetConstDocumentChecked();
-				const bool bLastGraph = Document.RootGraph.GetConstGraphPages().Num() == 1;
-				const bool bHasProjectPageValues = !bLastGraph && Document.RootGraph.FindConstGraph(Frontend::DefaultPageID) != nullptr;
-				return Editor::PageEditorEnabled(DocBuilder, bHasProjectPageValues);
+				if (DocBuilder.IsValid())
+				{
+					const FMetasoundFrontendDocument& Document = DocBuilder.GetConstDocumentChecked();
+					const bool bLastGraph = Document.RootGraph.GetConstGraphPages().Num() == 1;
+					const bool bHasProjectPageValues = !bLastGraph && Document.RootGraph.FindConstGraph(Frontend::DefaultPageID) != nullptr;
+					return Editor::PageEditorEnabled(DocBuilder, bHasProjectPageValues);
+				}
 			}
 
 			return false;
diff --git a/Engine/Plugins/Runtime/Metasound/Source/MetasoundEditor/Private/MetasoundEditorSubsystem.cpp b/Engine/Plugins/Runtime/Metasound/Source/MetasoundEditor/Private/MetasoundEditorSubsystem.cpp
--- a/Engine/Plugins/Runtime/Metasound/Source/MetasoundEditor/Private/MetasoundEditorSubsystem.cpp
+++ b/Engine/Plugins/Runtime/Metasound/Source/MetasoundEditor/Private/MetasoundEditorSubsystem.cpp
@@ -143,31 +143,7 @@ UMetasoundEditorGraphMemberDefaultLiteral* UMetaSoundEditorSubsystem::CreateMemb
 	TSubclassOf<UMetasoundEditorGraphMemberDefaultLiteral> LiteralClass) const
 {
 	using namespace Metasound::Engine;
-
-	// If preset and input inherits from default, copy member metadata from referenced graph
-	if (Builder.IsPreset())
-	{
-		const FMetasoundFrontendGraphClass& RootGraph = Builder.GetConstDocumentChecked().RootGraph;
-		const TSet<FName>& InputsInheritingDefault = RootGraph.PresetOptions.InputsInheritingDefault;
-		if (InputsInheritingDefault.Contains(InMemberName))
-		{
-			// Get referenced asset to inherit metadata from
-			FMetasoundAssetBase* ReferencedPresetAsset = Builder.GetReferencedPresetAsset();
-			if (ReferencedPresetAsset)
-			{
-				FMetaSoundFrontendDocumentBuilder& ReferencedBuilder = FDocumentBuilderRegistry::GetChecked().FindOrBeginBuilding(ReferencedPresetAsset->GetOwningAsset());
-				const FMetasoundFrontendClassInput* ClassInput = ReferencedBuilder.FindGraphInput(InMemberName);
-				if (ClassInput)
-				{
-					if (UMetaSoundFrontendMemberMetadata* ReferencedMemberMetadata = ReferencedBuilder.FindMemberMetadata(ClassInput->NodeID))
-					{
-						return NewObject<UMetasoundEditorGraphMemberDefaultLiteral>(&Builder.CastDocumentObjectChecked<UObject>(), ReferencedMemberMetadata->GetClass(), FName(), RF_Transactional, ReferencedMemberMetadata);
-					}
-				}
-			}
-		}
-	}
-	// Otherwise, create brand new member metadata
+	// Create brand new member metadata
 	return NewObject<UMetasoundEditorGraphMemberDefaultLiteral>(&Builder.CastDocumentObjectChecked<UObject>(), LiteralClass, FName(), RF_Transactional, nullptr);
 }
 
@@ -257,16 +233,6 @@ UMetaSoundFrontendMemberMetadata* UMetaSoundEditorSubsystem::FindOrCreateGraphIn
 				const FName TypeName = DocBuilder.FindGraphInput(InputName)->TypeName;
 				TSubclassOf<UMetasoundEditorGraphMemberDefaultLiteral> LiteralClass = GetLiteralClassForType(TypeName);
 
-				if (InBuilder->IsPreset())
-				{
-					// Ensure this MetaSound's dependencies are registered to lookup inherited metadata
-					// Needed in the case where this function is called from a BP,
-					// metadata wasn't created previously, and the MetaSound editor was never opened this session
-					UObject* MetaSound = InBuilder->GetConstBuilder().GetMetasoundAsset().GetOwningAsset();
-					check(MetaSound);
-					RegisterGraphWithFrontend(*MetaSound);
-				}
-
 				// Create new literal and setup
 				UMetasoundEditorGraphMemberDefaultLiteral* NewLiteral = CreateMemberMetadata(DocBuilder, InputName, LiteralClass);
 				if (NewLiteral)
diff --git a/Engine/Plugins/Runtime/Metasound/Source/MetasoundFrontend/Private/MetasoundFrontendTransform.cpp b/Engine/Plugins/Runtime/Metasound/Source/MetasoundFrontend/Private/MetasoundFrontendTransform.cpp
--- a/Engine/Plugins/Runtime/Metasound/Source/MetasoundFrontend/Private/MetasoundFrontendTransform.cpp
+++ b/Engine/Plugins/Runtime/Metasound/Source/MetasoundFrontend/Private/MetasoundFrontendTransform.cpp
@@ -743,7 +743,6 @@ namespace Metasound
 			// Ensure preset interfaces match those found in referenced graph.  Referenced graph is assumed to be
 			// well-formed (i.e. all inputs/outputs/environment variables declared by interfaces are present, and
 			// of proper name & data type). Pass in parent builder to copy member ids so added member ids can be consistent
-			FMetasoundFrontendDocument& Document = const_cast<FMetasoundFrontendDocument&>(InOutBuilderToTransform.GetConstDocumentChecked());
 			const TSet<FMetasoundFrontendVersion>& RefInterfaceVersions = ParentDocument.Interfaces;
 			for (const FMetasoundFrontendVersion& Version : RefInterfaceVersions)
 			{
