Removing unnecessary fade-in on our oscillator nodes

#rn Removed unnecessary fade from MetaSound oscillator nodes

diff --git a/Engine/Plugins/Runtime/Metasound/Source/MetasoundStandardNodes/Private/MetasoundOscillators.h b/Engine/Plugins/Runtime/Metasound/Source/MetasoundStandardNodes/Private/MetasoundOscillators.h
--- a/Engine/Plugins/Runtime/Metasound/Source/MetasoundStandardNodes/Private/MetasoundOscillators.h
+++ b/Engine/Plugins/Runtime/Metasound/Source/MetasoundStandardNodes/Private/MetasoundOscillators.h
@@ -97,8 +97,6 @@ namespace Metasound
 			float Phase = 0.f;
 			Oscillator Osc;
 			PhaseWrap Wrap;
-			float CurrentFade = 0.0f;
-			float FadeSmooth = 0.01f;
 			float CurrentFreq = -1.f;
 		
 			void operator()(const FGeneratorArgs& InArgs)
@@ -120,11 +118,8 @@ namespace Metasound
 							Wrap(Phase);
 
 							float Output = Osc(Phase, DeltaPhase, InArgs);
-							Output *= CurrentFade;
 							*Out++ = Output;
 
-							CurrentFade = CurrentFade + FadeSmooth * (1.0f - CurrentFade);
-
 							Phase += DeltaPhase;
 
 							RemainingSamplesInBlock--;
@@ -139,11 +134,8 @@ namespace Metasound
 							const float DeltaPhase = CurrentFreq * OneOverSampleRate;
 
 							float Output = Osc(Phase, DeltaPhase, InArgs);
-							Output *= CurrentFade;
 							*Out++ = Output;
 
-							CurrentFade = CurrentFade + FadeSmooth * (1.0f - CurrentFade);
-
 							Phase += DeltaPhase;
 
 							// lerp frequency based on the glide factor
@@ -166,11 +158,8 @@ namespace Metasound
 
 							float Output = Osc(Phase, DeltaPhase, InArgs);
 							Output = Audio::GetUnipolar(Output);
-							Output *= CurrentFade;
 							*Out++ = Output;
 
-							CurrentFade = CurrentFade + FadeSmooth * (1.0f - CurrentFade);
-
 							Phase += DeltaPhase;
 
 							RemainingSamplesInBlock--;
@@ -186,11 +175,8 @@ namespace Metasound
 
 							float Output = Osc(Phase, DeltaPhase, InArgs);
 							Output = Audio::GetUnipolar(Output);
-							Output *= CurrentFade;
 							*Out++ = Output;
 
-							CurrentFade = CurrentFade + FadeSmooth * (1.0f - CurrentFade);
-
 							Phase += DeltaPhase;
 
 							// lerp frequency based on the glide factor
@@ -215,8 +201,6 @@ namespace Metasound
 			Oscillator Osc;
 			PhaseWrap Wrap;
 
-			float CurrentFade = 0.0f;
-			float FadeSmooth = 0.01f;
 			float CurrentFreq = -1.f;
 
 			void operator()(const FGeneratorArgs& InArgs)
@@ -241,11 +225,8 @@ namespace Metasound
 							Wrap(Phase);
 
 							float Output = Osc(Phase, DeltaPhase, InArgs);
-							Output *= CurrentFade;
 							*Out++ = Output;
 
-							CurrentFade = CurrentFade + FadeSmooth * (1.0f - CurrentFade);
-
 							Phase += DeltaPhase;
 
 							RemainingSamplesInBlock--;
@@ -263,11 +244,8 @@ namespace Metasound
 							Wrap(Phase);
 
 							float Output = Osc(Phase, DeltaPhase, InArgs);
-							Output *= CurrentFade;
 							*Out++ = Output;
 
-							CurrentFade = CurrentFade + FadeSmooth * (1.0f - CurrentFade);
-
 							Phase += DeltaPhase;
 
 							// lerp frequency based on the glide factor
@@ -291,11 +269,8 @@ namespace Metasound
 
 							float Output = Osc(Phase, DeltaPhase, InArgs);
 							Output = Audio::GetUnipolar(Output);
-							Output *= CurrentFade;
 							*Out++ = Output;
 
-							CurrentFade = CurrentFade + FadeSmooth * (1.0f - CurrentFade);
-
 							Phase += DeltaPhase;
 
 							RemainingSamplesInBlock--;
@@ -314,11 +289,8 @@ namespace Metasound
 
 							float Output = Osc(Phase, DeltaPhase, InArgs);
 							Output = Audio::GetUnipolar(Output);
-							Output *= CurrentFade;
 							*Out++ = Output;
 
-							CurrentFade = CurrentFade + FadeSmooth * (1.0f - CurrentFade);
-
 							Phase += DeltaPhase;
 
 							// lerp frequency based on the glide factor
