Fix for delay node pitch shifting issue

#rn Fix for MetaSound delay node having a pitch shift on delay length initialization.

diff --git a/Engine/Plugins/Runtime/Metasound/Source/MetasoundStandardNodes/Private/MetasoundDelayNode.cpp b/Engine/Plugins/Runtime/Metasound/Source/MetasoundStandardNodes/Private/MetasoundDelayNode.cpp
--- a/Engine/Plugins/Runtime/Metasound/Source/MetasoundStandardNodes/Private/MetasoundDelayNode.cpp
+++ b/Engine/Plugins/Runtime/Metasound/Source/MetasoundStandardNodes/Private/MetasoundDelayNode.cpp
@@ -219,18 +219,29 @@ namespace Metasound
 			using namespace DelayNodePrivate;
 
 			FeedbackSample = 0.f;
+			PrevDelayTimeMsec = -1.0f;
 
-			float InitialDelayInSeconds = TDelay<ValueType>::GetDelayLengthSeconds(DelayTime);
-			PrevDelayTimeMsec = GetClampedDelayTimeMsec(InitialDelayInSeconds);
-
+			// Initialize the delay buffer, but don't set the delay length until first execute
 			DelayBuffer.Init(InParams.OperatorSettings.GetSampleRate(), MaxDelayTimeSeconds);
-			DelayBuffer.SetDelayMsec(PrevDelayTimeMsec);
-
 			AudioOutput->Zero();
 		}
 
 		void Execute()
 		{
+			// Wait until our first execute to set the delay length
+			// This allows Reset and intialization to propagate through first from inputs
+			// Not doing this can cause random pitching as random delay lengths etc interpolate
+			// from some unintended initial value 
+			if (PrevDelayTimeMsec < 0.0f)
+			{
+				using namespace DelayNodePrivate;
+
+				float InitialDelayInSeconds = TDelay<ValueType>::GetDelayLengthSeconds(DelayTime);
+				PrevDelayTimeMsec = GetClampedDelayTimeMsec(InitialDelayInSeconds);
+
+				DelayBuffer.SetDelayMsec(PrevDelayTimeMsec);
+			}
+
 			TriggerReset->ExecuteBlock(
 				[&](int32 StartFrame, int32 EndFrame)
 				{
@@ -340,7 +351,7 @@ namespace Metasound
 		Audio::FDelay DelayBuffer;
 
 		// The previous delay time
-		float PrevDelayTimeMsec = 0.f;
+		float PrevDelayTimeMsec = -1.f;
 
 		// Feedback sample
 		float FeedbackSample = 0.f;
