Fix MetaSound variable numeric defaults not being draggable

diff --git a/Engine/Plugins/Runtime/Metasound/Source/MetasoundFrontend/Private/MetasoundFrontendDocumentBuilder.cpp b/Engine/Plugins/Runtime/Metasound/Source/MetasoundFrontend/Private/MetasoundFrontendDocumentBuilder.cpp
--- a/Engine/Plugins/Runtime/Metasound/Source/MetasoundFrontend/Private/MetasoundFrontendDocumentBuilder.cpp
+++ b/Engine/Plugins/Runtime/Metasound/Source/MetasoundFrontend/Private/MetasoundFrontendDocumentBuilder.cpp
@@ -6552,9 +6552,6 @@ bool FMetaSoundFrontendDocumentBuilder::SetGraphVariableDefault(FName VariableNa
 		if (Frontend::IDataTypeRegistry::Get().IsLiteralTypeSupported(Variable->TypeName, InDefaultLiteral.GetType()))
 		{
 			Variable->Literal = MoveTemp(InDefaultLiteral);
-#if WITH_EDITORONLY_DATA
-			GetDocumentChecked().Metadata.ModifyContext.AddMemberIDModified(Variable->ID);
-#endif // WITH_EDITORONLY_DATA
 			return true;
 		}
 	}
