Fix crash when undoing MetaSound preset input override

diff --git a/Engine/Plugins/Runtime/Metasound/Source/MetasoundEditor/Private/MetasoundEditorGraph.cpp b/Engine/Plugins/Runtime/Metasound/Source/MetasoundEditor/Private/MetasoundEditorGraph.cpp
--- a/Engine/Plugins/Runtime/Metasound/Source/MetasoundEditor/Private/MetasoundEditorGraph.cpp
+++ b/Engine/Plugins/Runtime/Metasound/Source/MetasoundEditor/Private/MetasoundEditorGraph.cpp
@@ -695,16 +695,29 @@ void UMetasoundEditorGraphMemberDefaultLiteral::PostEditChangeChainProperty(FPro
 void UMetasoundEditorGraphMemberDefaultLiteral::PostEditUndo()
 {
 	using namespace Metasound::Editor;
-
-	Super::PostEditUndo();
+	using namespace Metasound::Frontend;
 
 	if (!IsValidChecked(this))
 	{
+		Super::PostEditUndo();
 		return;
 	}
 
+	// Builder must be reloaded to refresh caches, which may be stale after data reapplication.
+	// Must happen before Super PostEditUndo which can call PostEditChangeProperty 
+	// which can call RegisterGraphWithFrontend which depends on this data.
+	UMetasoundEditorGraphMember* Member = FindMember();
+	if (Member)
+	{
+		const FMetaSoundFrontendDocumentBuilder& Builder = Member->GetFrontendBuilderChecked();
+		const FMetasoundFrontendClassName& ClassName = Builder.GetConstDocumentChecked().RootGraph.Metadata.GetClassName();
+		IDocumentBuilderRegistry::GetChecked().ReloadBuilder(ClassName);
+	}
+
+	Super::PostEditUndo();
+
 	constexpr bool bPostTransaction = false;
-	if (UMetasoundEditorGraphMember* Member = FindMember())
+	if (Member)
 	{
 		Member->UpdateFrontendDefaultLiteral(bPostTransaction);
 	}
