Reparent MetaSound member metadata when building from builder API

diff --git a/Engine/Plugins/Runtime/Metasound/Source/MetasoundFrontend/Private/MetasoundFrontendDocumentBuilder.cpp b/Engine/Plugins/Runtime/Metasound/Source/MetasoundFrontend/Private/MetasoundFrontendDocumentBuilder.cpp
--- a/Engine/Plugins/Runtime/Metasound/Source/MetasoundFrontend/Private/MetasoundFrontendDocumentBuilder.cpp
+++ b/Engine/Plugins/Runtime/Metasound/Source/MetasoundFrontend/Private/MetasoundFrontendDocumentBuilder.cpp
@@ -4177,11 +4177,26 @@ void FMetaSoundFrontendDocumentBuilder::InitDocument(const FMetasoundFrontendDoc
 
 	FMetasoundFrontendDocument& Document = GetDocumentChecked();
 
-	// 1. Set default class Metadata.
+	// If template provided, copy & initialize from that.
 	if (InDocumentTemplate)
 	{
-		// 1a. If template provided, copy that.
+#if WITH_EDITORONLY_DATA
+		// Ensure member metadata is parented to new document owner
+		UObject* Parent = DocumentInterface.GetObject();
+		check(Parent);
+		TMap<FGuid, TObjectPtr<UMetaSoundFrontendMemberMetadata>> NewMetadata;
+		for (const TPair<FGuid, TObjectPtr<UMetaSoundFrontendMemberMetadata>>& Pair : InDocumentTemplate->Metadata.MemberMetadata)
+		{
+			NewMetadata.Add(Pair.Key, NewObject<UMetaSoundFrontendMemberMetadata>(Parent, Pair.Value->GetClass(), { }, { }, Pair.Value.Get()));
+		}
+#endif // WITH_EDITORONLY_DATA
+
 		Document = *InDocumentTemplate;
+
+#if WITH_EDITORONLY_DATA
+		Document.Metadata.MemberMetadata = MoveTemp(NewMetadata);
+#endif // WITH_EDITORONLY_DATA
+
 		if (Document.RootGraph.GetConstGraphPages().IsEmpty())
 		{
 			Document.RootGraph.InitDefaultGraphPage();
@@ -4190,9 +4205,9 @@ void FMetaSoundFrontendDocumentBuilder::InitDocument(const FMetasoundFrontendDoc
 
 		InitGraphClassMetadata(bResetVersion, InNewClassName);
 	}
+	// Initialize class using default data
 	else
 	{
-		// 1a. Initialize class using default data
 		if (Document.RootGraph.GetConstGraphPages().IsEmpty())
 		{
 			Document.RootGraph.InitDefaultGraphPage();
@@ -4203,21 +4218,20 @@ void FMetaSoundFrontendDocumentBuilder::InitDocument(const FMetasoundFrontendDoc
 		InitGraphClassMetadata(Document.RootGraph.Metadata, bResetVersion, InNewClassName);
 
 #if WITH_EDITORONLY_DATA
-		// 1b. Set default doc version Metadata
 		{
 			FMetasoundFrontendDocumentMetadata& DocMetadata = Document.Metadata;
 			DocMetadata.Version.Number = GetMaxDocumentVersion();
 		}
 #endif // WITH_EDITORONLY_DATA
 
-		// 1c. Add default interfaces for given UClass
+		// Add default interfaces for given UClass
 		{
 			TArray<FMetasoundFrontendVersion> InitVersions = ISearchEngine::Get().FindUClassDefaultInterfaceVersions(GetBuilderClassPath());
 			FModifyInterfaceOptions Options({ }, InitVersions);
 			ModifyInterfaces(MoveTemp(Options));
 		}
 
-		// 1d. Set default access flags as provided by document object class
+		// Set default access flags as provided by document object class
 		SetAccessFlags(DocumentInterface->GetDefaultAccessFlags());
 	}
 }
