Fix editor crash when MetaSound presets are autosaved

diff --git a/Engine/Plugins/Runtime/Metasound/Source/MetasoundEditor/Private/MetasoundEditorGraphConnectionDrawingPolicy.cpp b/Engine/Plugins/Runtime/Metasound/Source/MetasoundEditor/Private/MetasoundEditorGraphConnectionDrawingPolicy.cpp
--- a/Engine/Plugins/Runtime/Metasound/Source/MetasoundEditor/Private/MetasoundEditorGraphConnectionDrawingPolicy.cpp
+++ b/Engine/Plugins/Runtime/Metasound/Source/MetasoundEditor/Private/MetasoundEditorGraphConnectionDrawingPolicy.cpp
@@ -216,12 +216,15 @@ namespace Metasound
 
 			if (const FMetaSoundFrontendDocumentBuilder* Builder = Editor.GetFrontendBuilder())
 			{
-				if (const UEdGraphPin* ReroutePin = FGraphBuilder::FindReroutedOutputPin(OutputPin))
+				if (ensure(Builder->IsValid()))
 				{
-					OutputVertex = FGraphBuilder::GetPinVertex(*Builder, ReroutePin, &OutputNode);
-					if (OutputVertex && OutputNode)
+					if (const UEdGraphPin* ReroutePin = FGraphBuilder::FindReroutedOutputPin(OutputPin))
 					{
-						EdgeStyle = Builder->FindConstEdgeStyle(OutputNode->GetID(), OutputVertex->Name);
+						OutputVertex = FGraphBuilder::GetPinVertex(*Builder, ReroutePin, &OutputNode);
+						if (OutputVertex && OutputNode)
+						{
+							EdgeStyle = Builder->FindConstEdgeStyle(OutputNode->GetID(), OutputVertex->Name);
+						}
 					}
 				}
 			}
@@ -367,34 +370,37 @@ namespace Metasound
 				{
 					if (const FMetaSoundFrontendDocumentBuilder* Builder = EditorPtr->GetFrontendBuilder())
 					{
-						const FMetasoundFrontendNode* OutputNode = nullptr;
-						if (const FMetasoundFrontendVertex* OutputVertex = FGraphBuilder::GetPinVertex(*Builder, ReroutePin, &OutputNode))
+						if (ensure(Builder->IsValid()))
 						{
-							check(OutputNode);
-							const bool bIsPlaying = EditorPtr->IsPlaying();
-							const FGuid& NodeID = OutputNode->GetID();
-							if (AnyPin->PinType.PinCategory == FGraphBuilder::PinCategoryTrigger)
-							{
-								FGraphConnectionManager& ConnectionManager = EditorPtr->GetConnectionManager();
-								DrawingPolicyPrivate::DetermineWiringStyle_Envelope<FVertexAnalyzerTriggerDensity>(bIsPlaying, ConnectionManager, NodeID, OutputVertex->Name, InputPin, OutParams);
-								ScaleEnvelopeWireThickness(OutParams);
-								bThicknessSet = true;
-							}
-							else if (AnyPin->PinType.PinCategory == FGraphBuilder::PinCategoryAudio)
-							{
-								FGraphConnectionManager& ConnectionManager = EditorPtr->GetConnectionManager();
-								DrawingPolicyPrivate::DetermineWiringStyle_Envelope<FVertexAnalyzerEnvelopeFollower>(bIsPlaying, ConnectionManager, NodeID, OutputVertex->Name, InputPin, OutParams);
-								ScaleEnvelopeWireThickness(OutParams);
-								bThicknessSet = true;
-							}
-							else if (FGraphBuilder::CanInspectPin(AnyPin))
+							const FMetasoundFrontendNode* OutputNode = nullptr;
+							if (const FMetasoundFrontendVertex* OutputVertex = FGraphBuilder::GetPinVertex(*Builder, ReroutePin, &OutputNode))
 							{
-								if (const FMetasoundFrontendEdgeStyle* EdgeStyle = Builder->FindConstEdgeStyle(NodeID, OutputVertex->Name))
+								check(OutputNode);
+								const bool bIsPlaying = EditorPtr->IsPlaying();
+								const FGuid& NodeID = OutputNode->GetID();
+								if (AnyPin->PinType.PinCategory == FGraphBuilder::PinCategoryTrigger)
 								{
-									if (!EdgeStyle->LiteralColorPairs.IsEmpty())
+									FGraphConnectionManager& ConnectionManager = EditorPtr->GetConnectionManager();
+									DrawingPolicyPrivate::DetermineWiringStyle_Envelope<FVertexAnalyzerTriggerDensity>(bIsPlaying, ConnectionManager, NodeID, OutputVertex->Name, InputPin, OutParams);
+									ScaleEnvelopeWireThickness(OutParams);
+									bThicknessSet = true;
+								}
+								else if (AnyPin->PinType.PinCategory == FGraphBuilder::PinCategoryAudio)
+								{
+									FGraphConnectionManager& ConnectionManager = EditorPtr->GetConnectionManager();
+									DrawingPolicyPrivate::DetermineWiringStyle_Envelope<FVertexAnalyzerEnvelopeFollower>(bIsPlaying, ConnectionManager, NodeID, OutputVertex->Name, InputPin, OutParams);
+									ScaleEnvelopeWireThickness(OutParams);
+									bThicknessSet = true;
+								}
+								else if (FGraphBuilder::CanInspectPin(AnyPin))
+								{
+									if (const FMetasoundFrontendEdgeStyle* EdgeStyle = Builder->FindConstEdgeStyle(NodeID, OutputVertex->Name))
 									{
-										OutParams.WireThickness = AnalyzerSettings.NumericWireThickness;
-										bThicknessSet = true;
+										if (!EdgeStyle->LiteralColorPairs.IsEmpty())
+										{
+											OutParams.WireThickness = AnalyzerSettings.NumericWireThickness;
+											bThicknessSet = true;
+										}
 									}
 								}
 							}
diff --git a/Engine/Plugins/Runtime/Metasound/Source/MetasoundEditor/Private/MetasoundEditorGraphMemberDefaults.cpp b/Engine/Plugins/Runtime/Metasound/Source/MetasoundEditor/Private/MetasoundEditorGraphMemberDefaults.cpp
--- a/Engine/Plugins/Runtime/Metasound/Source/MetasoundEditor/Private/MetasoundEditorGraphMemberDefaults.cpp
+++ b/Engine/Plugins/Runtime/Metasound/Source/MetasoundEditor/Private/MetasoundEditorGraphMemberDefaults.cpp
@@ -622,8 +622,8 @@ void UMetasoundEditorGraphMemberDefaultFloat::PostLoad()
 	Super::PostLoad();
 
 	FMetasoundAssetBase& MetaSound = FGraphBuilder::GetOutermostMetaSoundChecked(*this);
-	const FMetaSoundFrontendDocumentBuilder& Builder = IDocumentBuilderRegistry::GetChecked().FindOrBeginBuilding(MetaSound.GetOwningAsset());
-	const bool bIsPreset = Builder.IsPreset();
+	const FMetasoundFrontendDocument& Document = MetaSound.GetConstDocumentChecked();
+	const bool bIsPreset = Document.RootGraph.PresetOptions.bIsPreset;
 
 	// Set default input widget type from settings. Widgets do not apply to presets
 	if (bIsPreset)
