Remove asset dirtying in MetaSound editor that should instead be managed by calling code/scripts' transactions

diff --git a/Engine/Plugins/Runtime/Metasound/Source/MetasoundEditor/Private/MetasoundEditor.cpp b/Engine/Plugins/Runtime/Metasound/Source/MetasoundEditor/Private/MetasoundEditor.cpp
--- a/Engine/Plugins/Runtime/Metasound/Source/MetasoundEditor/Private/MetasoundEditor.cpp
+++ b/Engine/Plugins/Runtime/Metasound/Source/MetasoundEditor/Private/MetasoundEditor.cpp
@@ -5296,11 +5296,6 @@ PRAGMA_ENABLE_DEPRECATION_WARNINGS
 									Literal->UpdatePreviewInstance(InputName, ParamInterface);
 								}
 							}
-
-							// Mark dirty (needed for when this is being called in response to a change from Builder BP API for example)
-							UObject* MetaSound = ParentPtr->GetMetasoundObject();
-							check(MetaSound);
-							MetaSound->MarkPackageDirty();
 						}
 					}
 				}
