MetaSound template node - node configuration, bughawk followup
#rnx

diff --git a/Engine/Plugins/Runtime/Metasound/Source/MetasoundFrontend/Private/MetasoundFrontendGraphController.cpp b/Engine/Plugins/Runtime/Metasound/Source/MetasoundFrontend/Private/MetasoundFrontendGraphController.cpp
--- a/Engine/Plugins/Runtime/Metasound/Source/MetasoundFrontend/Private/MetasoundFrontendGraphController.cpp
+++ b/Engine/Plugins/Runtime/Metasound/Source/MetasoundFrontend/Private/MetasoundFrontendGraphController.cpp
@@ -1584,8 +1584,7 @@ PRAGMA_ENABLE_DEPRECATION_WARNINGS
 				{
 					if (const FMetasoundFrontendClass* NodeClass = Class.Get())
 					{
-						// Currently template nodes do not have node configurations so an invalid one is supplied here. 
-						FMetasoundFrontendNode& Node = FindBuildGraphChecked(*GraphClass).Nodes.Emplace_GetRef(*NodeClass, TInstancedStruct<FMetaSoundFrontendNodeConfiguration>());
+						FMetasoundFrontendNode& Node = FindBuildGraphChecked(*GraphClass).Nodes.Emplace_GetRef(*NodeClass, InTemplate.CreateFrontendTemplateNodeConfiguration());
 						Node.UpdateID(InNodeGuid);
 						Node.Interface = InTemplate.GenerateNodeInterface(MoveTemp(Params));
 						FNodeAccessPtr NodePtr = GraphClassPtr.GetNodeWithNodeID(Node.GetID());
