MetaSound Search Engine Optimizations: Round 1

#rnx

diff --git a/Engine/Plugins/Runtime/Metasound/Source/MetasoundEditor/Private/MetasoundEditorGraphSchema.cpp b/Engine/Plugins/Runtime/Metasound/Source/MetasoundEditor/Private/MetasoundEditorGraphSchema.cpp
--- a/Engine/Plugins/Runtime/Metasound/Source/MetasoundEditor/Private/MetasoundEditorGraphSchema.cpp
+++ b/Engine/Plugins/Runtime/Metasound/Source/MetasoundEditor/Private/MetasoundEditorGraphSchema.cpp
@@ -411,183 +411,51 @@ PRAGMA_ENABLE_DEPRECATION_WARNINGS
 			return Algo::AnyOf(AccessorNodes, IsSourceNodeReachableDownstream);
 		}
 
-		class FAssetSchemaQueryResult : public ISchemaQueryResult
+		class FSchemaQueryResult : public ISchemaQueryResult
 		{
-		public:
-			FAssetSchemaQueryResult(Frontend::FMetaSoundAssetClassInfo InTagData)
-				: ClassInfo(MoveTemp(InTagData))
-			{
-				bIsLoaded = FSoftObjectPath(ClassInfo.AssetPath).ResolveObject() != nullptr;
-			}
-
-			virtual ~FAssetSchemaQueryResult() = default;
-
-			const FMetasoundFrontendGraphClass* FindGraphClass() const
-			{
-				using namespace Metasound;
-
-				if (UObject* MetaSound = FSoftObjectPath(ClassInfo.AssetPath).TryLoad())
-				{
-					if (const FMetasoundAssetBase* MetaSoundAsset = IMetasoundUObjectRegistry::Get().GetObjectAsAssetBase(MetaSound))
-					{
-						bIsLoaded = true;
-						if (!MetaSoundAsset->IsRegistered())
-						{
-							FGraphBuilder::RegisterGraphWithFrontend(*MetaSound);
-						}
-						return &MetaSoundAsset->GetConstDocumentChecked().RootGraph;
-					}
-				}
-
-				bIsLoaded = false;
-				return nullptr;
-			}
-
-			virtual const FMetasoundFrontendClass* FindClass() const override
-			{
-				return FindGraphClass();
-			}
-
-			virtual bool CanConnectInputOfTypeAndAccess(FName InputTypeName, EMetasoundFrontendVertexAccessType InputAccessType) const override
-			{
-				using namespace Frontend;
-
-				const bool bCanConnect = Algo::AnyOf(ClassInfo.InterfaceInfo.Outputs, [&InputTypeName, &InputAccessType](const FMetaSoundClassVertexInfo& Output)
-				{
-					return InputTypeName == Output.TypeName && FMetasoundFrontendClassVertex::CanConnectVertexAccessTypes(Output.AccessType, InputAccessType);
-				});
-				return bCanConnect;
-			}
-
-			virtual bool CanConnectOutputOfTypeAndAccess(FName OutputTypeName, EMetasoundFrontendVertexAccessType OutputAccessType) const override
-			{
-				using namespace Engine;
-				const bool bCanConnect = Algo::AnyOf(ClassInfo.InterfaceInfo.Inputs, [&OutputTypeName, &OutputAccessType](const FMetaSoundClassVertexInfo& Input)
-				{
-					return OutputTypeName == Input.TypeName && FMetasoundFrontendClassVertex::CanConnectVertexAccessTypes(OutputAccessType, Input.AccessType);
-				});
-				return bCanConnect;
-			}
-
-			virtual EMetasoundFrontendClassAccessFlags GetAccessFlags() const override
-			{
-				return ClassInfo.AccessFlags;
-			}
-
-			virtual FMetaSoundAssetKey GetAssetKey() const override
-			{
-				return FMetaSoundAssetKey(ClassInfo.ClassName, ClassInfo.Version);
-			}
-
-			virtual const TArray<FText>& GetCategoryHierarchy() const override
+			public:
+				FSchemaQueryResult(const Engine::FMetaSoundAssetManager& AssetManager, Frontend::FMetaSoundClassInfo InClassInfo)
+				: ClassInfo(MoveTemp(InClassInfo))
 			{
-				if (bIsLoaded)
+				if (FMetaSoundAssetKey::IsValidType(ClassInfo.ClassType))
 				{
-					if (const FMetasoundFrontendClass* Class = FindClass())
-					{
-						return Class->Metadata.GetCategoryHierarchy();
-					}
+					bIsNative = !AssetManager.ContainsKey(FMetaSoundAssetKey(ClassInfo.ToRegistryKey()));
 				}
-
-				return ClassInfo.InterfaceInfo.SearchInfo.Hierarchy;
-			}
-
-			virtual FText GetDisplayName() const override
-			{
-				using namespace Frontend;
-
-				FText DisplayName = FText::FromName(ClassInfo.AssetPath.GetAssetName());
-
-				if (bIsLoaded)
-				{
-					if (const FMetasoundFrontendGraphClass* GraphClass = FindGraphClass())
-					{
-						FText DisplayNameOverride = GraphClass->Metadata.GetDisplayName();
-						if (!DisplayNameOverride.IsEmptyOrWhitespace())
-						{
-							DisplayName = DisplayNameOverride;
-						}
-					}
-				}
-				else if (!ClassInfo.InterfaceInfo.SearchInfo.ClassDisplayName.IsEmptyOrWhitespace())
-				{
-					DisplayName = ClassInfo.InterfaceInfo.SearchInfo.ClassDisplayName;
-				}
-
-				if (ShowUnloadedAssetInBrowserCvar)
+				else
 				{
-					return FText::Format(LOCTEXT("FileNodeLoadedNameHintFormat", "{0}*"), DisplayName);
+					bIsNative = false;
 				}
-
-				return DisplayName;
 			}
 
-			virtual const TArray<FText>& GetKeywords() const override
-			{
-				if (bIsLoaded)
-				{
-					if (const FMetasoundFrontendClass* Class = FindClass())
-					{
-						return Class->Metadata.GetKeywords();
-					}
-				}
-
-				return ClassInfo.InterfaceInfo.SearchInfo.Keywords;
-			}
+			virtual ~FSchemaQueryResult() = default;
 
-			virtual EMetasoundFrontendClassType GetRegistryClassType() const override
+			virtual FMetasoundFrontendClassMetadata GenerateClassMetadata() const override
 			{
-				return EMetasoundFrontendClassType::External;
-			}
+				using namespace Frontend;
 
-			virtual FText GetTooltip() const override
-			{
-				if (bIsLoaded)
+				// Asset-defined node classes must be loaded (and registered) in order to generate metadata
+				if (!bIsNative)
 				{
-					if (const FMetasoundFrontendClass* Class = FindClass())
+					const FMetaSoundAssetKey AssetKey = ClassInfo.ToRegistryKey();
+					if (FMetasoundAssetBase* Asset = IMetaSoundAssetManager::GetChecked().TryLoadAssetFromKey(AssetKey))
 					{
-						const FText Author = FText::FromString(Class->Metadata.GetAuthor());
-						const FText& Description = Class->Metadata.GetDescription();
-						if (Description.IsEmptyOrWhitespace())
-						{
-							return Author;
-						}
-						return Author.IsEmptyOrWhitespace() ? Description : FText::Format(ClassDescriptionAndAuthorFormat, Description, Author);
+						FMetaSoundAssetRegistrationOptions RegOptions = FGraphBuilder::GetDefaultRegistrationOptions();
+						RegOptions.bForceReregister = false;
+						UObject* MetaSoundObject = Asset->GetOwningAsset();
+						check(MetaSoundObject);
+						FGraphBuilder::RegisterGraphWithFrontend(*MetaSoundObject);
 					}
 				}
 
-				return ClassInfo.InterfaceInfo.SearchInfo.ClassDescription;
-			}
-
-			virtual bool IsNative() const override
-			{
-				return false;
-			}
-
-		private:
-			Frontend::FMetaSoundAssetClassInfo ClassInfo;
-			mutable bool bIsLoaded = false;
-		};
-
-		class FRegistrySchemaQueryResult : public ISchemaQueryResult
-		{
-			public:
-			FRegistrySchemaQueryResult(const Engine::FMetaSoundAssetManager& AssetManager, FMetasoundFrontendClass InClass)
-				: Class(MoveTemp(InClass))
-			{
-				bIsNative = !AssetManager.IsAssetClass(Class.Metadata);
-			}
-
-			virtual ~FRegistrySchemaQueryResult() = default;
-
-			virtual const FMetasoundFrontendClass* FindClass() const override
-			{
-				return &Class;
+				FMetasoundFrontendClass RegisteredClass;
+				const bool bFoundClass = Frontend::INodeClassRegistry::GetChecked().FindFrontendClassFromRegistered(ClassInfo.ToRegistryKey(), RegisteredClass);
+				ensure(bFoundClass);
+				return RegisteredClass.Metadata;
 			}
 
 			virtual bool CanConnectInputOfTypeAndAccess(FName InputTypeName, EMetasoundFrontendVertexAccessType InputAccessType) const override
 			{
-				const bool bCanConnect = Algo::AnyOf(Class.GetDefaultInterface().Outputs, [&InputTypeName, &InputAccessType](const FMetasoundFrontendClassOutput& Output)
+				const bool bCanConnect = Algo::AnyOf(ClassInfo.InterfaceInfo.Outputs, [&InputTypeName, &InputAccessType](const FMetaSoundClassVertexInfo& Output)
 				{
 					return InputTypeName == Output.TypeName && FMetasoundFrontendClassVertex::CanConnectVertexAccessTypes(Output.AccessType, InputAccessType);
 				});
@@ -596,7 +464,7 @@ PRAGMA_ENABLE_DEPRECATION_WARNINGS
 
 			virtual bool CanConnectOutputOfTypeAndAccess(FName OutputTypeName, EMetasoundFrontendVertexAccessType OutputAccessType) const override
 			{
-				const bool bCanConnect = Algo::AnyOf(Class.GetDefaultInterface().Inputs, [&OutputTypeName, &OutputAccessType](const FMetasoundFrontendClassInput& Input)
+				const bool bCanConnect = Algo::AnyOf(ClassInfo.InterfaceInfo.Inputs, [&OutputTypeName, &OutputAccessType](const FMetaSoundClassVertexInfo& Input)
 				{
 					return OutputTypeName == Input.TypeName && FMetasoundFrontendClassVertex::CanConnectVertexAccessTypes(OutputAccessType, Input.AccessType);
 				});
@@ -605,17 +473,17 @@ PRAGMA_ENABLE_DEPRECATION_WARNINGS
 
 			virtual EMetasoundFrontendClassAccessFlags GetAccessFlags() const override
 			{
-				return Class.Metadata.GetAccessFlags();
+				return ClassInfo.AccessFlags;
 			}
 
 			virtual FMetaSoundAssetKey GetAssetKey() const override
 			{
-				return FMetaSoundAssetKey(Class.Metadata);
+				return FMetaSoundAssetKey(ClassInfo.ToRegistryKey());
 			}
 
 			virtual const TArray<FText>& GetCategoryHierarchy() const override
 			{
-				return Class.Metadata.GetCategoryHierarchy();
+				return ClassInfo.InterfaceInfo.SearchInfo.Hierarchy;
 			}
 
 			virtual FText GetDisplayName() const override
@@ -623,31 +491,34 @@ PRAGMA_ENABLE_DEPRECATION_WARNINGS
 				using namespace Frontend;
 
 				FText DisplayName;
-				auto GetAssetDisplayNameFromMetadata = [&DisplayName](const FMetasoundFrontendClassMetadata& Metadata)
+				auto GetAssetDisplayNameFromMetadata = [&DisplayName](const FMetaSoundClassInfo& ClassInfo)
 				{
-					DisplayName = Metadata.GetDisplayName();
+					DisplayName = ClassInfo.InterfaceInfo.SearchInfo.ClassDisplayName;
 					if (DisplayName.IsEmptyOrWhitespace())
 					{
-						const FTopLevelAssetPath Path = IMetaSoundAssetManager::GetChecked().FindAssetPath(FMetaSoundAssetKey(Metadata));
-						if (Path.IsValid())
+						if (FMetaSoundAssetKey::IsValidType(ClassInfo.ClassType))
 						{
-							DisplayName = FText::FromName(Path.GetAssetName());
+							const FTopLevelAssetPath Path = IMetaSoundAssetManager::GetChecked().FindAssetPath(FMetaSoundAssetKey(ClassInfo.ToRegistryKey()));
+							if (Path.IsValid())
+							{
+								DisplayName = FText::FromName(Path.GetAssetName());
+							}
 						}
 					}
 				};
 
 				// 1. Try to get display name from metadata or asset if one can be found from the asset manager
-				GetAssetDisplayNameFromMetadata(Class.Metadata);
+				GetAssetDisplayNameFromMetadata(ClassInfo);
 
 				// 2. If version is missing from the registry or from asset system, then this node
 				// will not provide a useful DisplayName.  In that case, attempt to find the next highest
 				// class & associated DisplayName.
 				if (DisplayName.IsEmptyOrWhitespace())
 				{
-					FMetasoundFrontendClass ClassWithHighestVersion;
-					if (ISearchEngine::Get().FindClassWithHighestVersion(Class.Metadata.GetClassName(), ClassWithHighestVersion))
+					FMetaSoundClassInfo HighestVersionInfo;
+					if (ISearchEngine::Get().FindClassWithHighestVersion(ClassInfo.ClassName, HighestVersionInfo))
 					{
-						GetAssetDisplayNameFromMetadata(ClassWithHighestVersion.Metadata);
+						GetAssetDisplayNameFromMetadata(HighestVersionInfo);
 					}
 				}
 
@@ -656,20 +527,26 @@ PRAGMA_ENABLE_DEPRECATION_WARNINGS
 
 			virtual const TArray<FText>& GetKeywords() const override
 			{
-				return Class.Metadata.GetKeywords();
+				return ClassInfo.InterfaceInfo.SearchInfo.Keywords;
 			}
 
 			virtual FText GetTooltip() const
 			{
-				const FString& Author = Class.Metadata.GetAuthor();
-				return Author.IsEmpty()
-					? FText::Format(ClassDescriptionAndAuthorFormat, Class.Metadata.GetDescription(), FText::FromString(Author))
-					: Class.Metadata.GetDescription();
+				const FText& ClassDescription = ClassInfo.InterfaceInfo.SearchInfo.ClassDescription;
+#if 0
+				// TODO: Add author to asset tag/class info data and include in tooltip
+				const FString& Author = ClassInfo.InterfaceInfo.SearchInfo.Author;
+				return !Author.IsEmpty()
+					? FText::Format(ClassDescriptionAndAuthorFormat, ClassDescription, FText::FromString(Author))
+					: ClassDescription;
+#else
+				return ClassDescription;
+#endif
 			}
 
 			virtual EMetasoundFrontendClassType GetRegistryClassType() const override
 			{
-				return Class.Metadata.GetType();
+				return ClassInfo.ClassType;
 			}
 
 			virtual bool IsNative() const
@@ -678,7 +555,7 @@ PRAGMA_ENABLE_DEPRECATION_WARNINGS
 			}
 
 			private:
-			FMetasoundFrontendClass Class;
+			Frontend::FMetaSoundClassInfo ClassInfo;
 			bool bIsNative = false;
 		};
 	} // namespace SchemaPrivate
@@ -1010,20 +887,13 @@ UEdGraphNode* FMetasoundGraphSchemaAction_NewNode::PerformAction(UEdGraph* Paren
 	using namespace Metasound::Editor;
 	using namespace Metasound::Frontend;
 
-	const FMetasoundFrontendClass* Class = QueryResult->FindClass();
-	if (!Class)
-	{
-		return nullptr;
-	}
-
 	const FScopedTransaction Transaction(LOCTEXT("AddNewNode", "Add New MetaSound Node"));
 	UMetasoundEditorGraph* MetaSoundGraph = CastChecked<UMetasoundEditorGraph>(ParentGraph);
 	UObject& ParentMetasound = MetaSoundGraph->GetMetasoundChecked();
 	ParentMetasound.Modify();
 	ParentGraph->Modify();
 
-	FMetasoundFrontendClassMetadata Metadata = Class->Metadata;
-	Metadata.SetType(EMetasoundFrontendClassType::External);
+	FMetasoundFrontendClassMetadata Metadata = QueryResult->GenerateClassMetadata();
 	if (UMetasoundEditorGraphExternalNode* NewGraphNode = FGraphBuilder::AddExternalNode(ParentMetasound, Metadata, bSelectNewNode))
 	{
 		NewGraphNode->Modify();
@@ -2586,7 +2456,6 @@ void UMetasoundEditorGraphSchema::GetConversionActions(FGraphActionMenuBuilder&
 	using namespace Metasound;
 	using namespace Metasound::Editor;
 
-	constexpr bool bScanAssetTags = false; // All conversion actions are natively defined, so no reason to list as conversion
 	QueryNodeClasses([&](TUniquePtr<ISchemaQueryResult>&& Result)
 	{
 		if (Result->GetRegistryClassType() != EMetasoundFrontendClassType::External)
@@ -2626,7 +2495,7 @@ void UMetasoundEditorGraphSchema::GetConversionActions(FGraphActionMenuBuilder&
 			NewNodeAction->QueryResult = TSharedPtr<ISchemaQueryResult>(Result.Release());
 			ActionMenuBuilder.AddAction(NewNodeAction);
 		}
-	}, bScanAssetTags);
+	});
 }
 
 void UMetasoundEditorGraphSchema::GetDataTypeInputNodeActions(FGraphContextMenuBuilder& ActionMenuBuilder, Metasound::Frontend::FConstGraphHandle InGraphHandle, Metasound::Editor::FInterfaceNodeFilterFunction InFilter, bool bShowSelectedActions) const
@@ -2898,49 +2767,26 @@ void UMetasoundEditorGraphSchema::SetNodePosition(UEdGraphNode* Node, const FVec
 	}
 }
 
-void UMetasoundEditorGraphSchema::QueryNodeClasses(TFunctionRef<void(TUniquePtr<Metasound::Editor::ISchemaQueryResult>&&)> OnClassFound, bool bScanAssetTags) const
+void UMetasoundEditorGraphSchema::QueryNodeClasses(TFunctionRef<void(TUniquePtr<Metasound::Editor::ISchemaQueryResult>&&)> OnClassFound) const
 {
 	using namespace Metasound::Editor;
 	using namespace Metasound::Engine;
 	using namespace Metasound::Frontend;
 
-	TArray<TUniquePtr<ISchemaQueryResult>> Results;
 	const FMetaSoundAssetManager& AssetManager = FMetaSoundAssetManager::GetChecked();
 
 	// Gather all loaded classes
 	{
-		TSet<FNodeRegistryKey> FoundKeys;
-		{
-			const bool bIncludeAllVersions = static_cast<bool>(SchemaPrivate::EnableAllVersionsMetaSoundNodeClassCreationCVar);
-			TArray<FMetasoundFrontendClass> FrontendClasses = ISearchEngine::Get().FindAllClasses(bIncludeAllVersions);
-			for (const FMetasoundFrontendClass& Class : FrontendClasses)
-			{
-				FoundKeys.Add(FNodeRegistryKey(Class.Metadata));
-				Results.Add(MakeUnique<SchemaPrivate::FRegistrySchemaQueryResult>(AssetManager, Class));
-			}
-		}
-
-		// Append all unloaded asset classes
-		if (bScanAssetTags)
+		static bool bIncludeUnloadedAssets = true;
+		const ISearchEngine::EResultVersion Options = SchemaPrivate::EnableAllVersionsMetaSoundNodeClassCreationCVar == 0
+			? ISearchEngine::EResultVersion::Highest
+			: ISearchEngine::EResultVersion::All;
+		TArray<FMetaSoundClassInfo> ClassInfos = ISearchEngine::Get().FindAllClasses(Options, bIncludeUnloadedAssets);
+		for (const FMetaSoundClassInfo& ClassInfo: ClassInfos)
 		{
-			AssetManager.IterateAssetTagData([&FoundKeys, &Results](FMetaSoundAssetClassInfo ClassInfo)
-			{
-				using namespace Metasound::Frontend;
-
-				const FMetaSoundAssetKey AssetKey(ClassInfo.ClassName, ClassInfo.Version);
-				const FNodeRegistryKey RegKey(AssetKey);
-				if (!FoundKeys.Contains(RegKey))
-				{
-					Results.Add(MakeUnique<SchemaPrivate::FAssetSchemaQueryResult>(MoveTemp(ClassInfo)));
-				}
-			});
+			OnClassFound(MakeUnique<SchemaPrivate::FSchemaQueryResult>(AssetManager, ClassInfo));
 		}
 	}
-
-	for (TUniquePtr<ISchemaQueryResult>& Result : Results)
-	{
-		OnClassFound(MoveTemp(Result));
-	}
 }
 #undef LOCTEXT_NAMESPACE
 
diff --git a/Engine/Plugins/Runtime/Metasound/Source/MetasoundEditor/Public/MetasoundEditorGraphSchema.h b/Engine/Plugins/Runtime/Metasound/Source/MetasoundEditor/Public/MetasoundEditorGraphSchema.h
--- a/Engine/Plugins/Runtime/Metasound/Source/MetasoundEditor/Public/MetasoundEditorGraphSchema.h
+++ b/Engine/Plugins/Runtime/Metasound/Source/MetasoundEditor/Public/MetasoundEditorGraphSchema.h
@@ -45,10 +45,10 @@ namespace Metasound::Editor
 	public:
 		virtual ~ISchemaQueryResult() = default;
 
-		virtual const FMetasoundFrontendClass* FindClass() const = 0;
 		virtual bool CanConnectInputOfTypeAndAccess(FName TypeName, EMetasoundFrontendVertexAccessType AccessType) const = 0;
 		virtual bool CanConnectOutputOfTypeAndAccess(FName TypeName, EMetasoundFrontendVertexAccessType AccessType) const = 0;
 
+		virtual FMetasoundFrontendClassMetadata GenerateClassMetadata() const = 0;
 		virtual EMetasoundFrontendClassAccessFlags GetAccessFlags() const = 0;
 		virtual FMetaSoundAssetKey GetAssetKey() const = 0;
 		virtual const TArray<FText>& GetCategoryHierarchy() const = 0;
@@ -567,7 +567,7 @@ public:
 	void BreakNodeLinks(UEdGraphNode& TargetNode, bool bShouldActuallyTransact) const;
 
 private:
-	void QueryNodeClasses(TFunctionRef<void(TUniquePtr<Metasound::Editor::ISchemaQueryResult>&&)> OnClassFound, bool bScanAssetTags = true) const;
+	void QueryNodeClasses(TFunctionRef<void(TUniquePtr<Metasound::Editor::ISchemaQueryResult>&&)> OnClassFound) const;
 
 	/** Adds actions for creating actions associated with graph DataTypes */
 	void GetConversionActions(FGraphActionMenuBuilder& ActionMenuBuilder, Metasound::Editor::FActionVertexFilters InFilters = Metasound::Editor::FActionVertexFilters(), bool bShowSelectedActions = true) const;
diff --git a/Engine/Plugins/Runtime/Metasound/Source/MetasoundEngine/Private/MetasoundAssetSubsystem.cpp b/Engine/Plugins/Runtime/Metasound/Source/MetasoundEngine/Private/MetasoundAssetSubsystem.cpp
--- a/Engine/Plugins/Runtime/Metasound/Source/MetasoundEngine/Private/MetasoundAssetSubsystem.cpp
+++ b/Engine/Plugins/Runtime/Metasound/Source/MetasoundEngine/Private/MetasoundAssetSubsystem.cpp
@@ -36,7 +36,7 @@ namespace Metasound::Engine
 
 		std::atomic<bool> InitialAssetScanComplete = false;
 
-		bool RemoveClassInfo(FCriticalSection& MapCritSec, FAssetClassInfoMap& Map, const FMetaSoundAssetKey& AssetKey, const FTopLevelAssetPath& InAssetPath)
+		bool RemoveClassInfo(Frontend::IMetaSoundAssetTransactor& Transactor, FCriticalSection& MapCritSec, FAssetClassInfoMap& Map, const FMetaSoundAssetKey& AssetKey, const FTopLevelAssetPath& InAssetPath)
 		{
 			using namespace Frontend;
 
@@ -67,6 +67,8 @@ namespace Metasound::Engine
 
 					if (AssetInfoArray->IsEmpty())
 					{
+						// Only unregister from Transactor if empty as transaction buffer only cares if *ANY* asset is registered with the given class
+						Transactor.UnregisterAsset(AssetKey);
 						Map.Remove(AssetKey);
 					}
 					return true;
@@ -76,7 +78,7 @@ namespace Metasound::Engine
 			return false;
 		}
 
-		FMetaSoundAssetKey AddClassInfo(FCriticalSection& MapCritSec, FAssetClassInfoMap& Map, Frontend::FMetaSoundAssetClassInfo&& ClassInfo)
+		FMetaSoundAssetKey AddClassInfo(Frontend::IMetaSoundAssetTransactor& Transactor, FCriticalSection& MapCritSec, FAssetClassInfoMap& Map, Frontend::FMetaSoundAssetClassInfo&& ClassInfo)
 		{
 			using namespace Frontend;
 
@@ -88,15 +90,24 @@ namespace Metasound::Engine
 
 			FScopeLock Lock(&MapCritSec);
 			TArray<FMetaSoundAssetClassInfo>& TagDatas = Map.FindOrAdd(AssetKey);
-			auto IsTagData = [&ClassInfo](const FMetaSoundAssetClassInfo& Iter) { return Iter.AssetPath == ClassInfo.AssetPath; };
-			TagDatas.RemoveAllSwap(IsTagData);
+			if (TagDatas.IsEmpty())
+			{
+				// Only register with Transactor if empty as transaction buffer only cares if *ANY* asset is registered with the given class
+				Transactor.RegisterAsset(AssetKey);
+			}
+			else
+			{
+				auto IsTagData = [&ClassInfo](const FMetaSoundAssetClassInfo& Iter) { return Iter.AssetPath == ClassInfo.AssetPath; };
+				TagDatas.RemoveAllSwap(IsTagData);
+			}
+
 			TagDatas.Add(MoveTemp(ClassInfo));
 
 #if !NO_LOGGING
 			if (TagDatas.Num() > 1)
 			{
 				TArray<FString> PathStrings;
-				Algo::Transform(TagDatas, PathStrings, [](const FMetaSoundAssetClassInfo& ClassInfo) { return ClassInfo.AssetPath.ToString(); });
+				Algo::Transform(TagDatas, PathStrings, [](const FMetaSoundAssetClassInfo& WarnInfo) { return WarnInfo.AssetPath.ToString(); });
 				UE_LOG(LogMetaSound, Warning,
 					TEXT("MetaSoundAssetManager has registered multiple assets with key '%s':\n%s\n"),
 					*AssetKey.ToString(),
@@ -106,6 +117,23 @@ namespace Metasound::Engine
 
 			return AssetKey;
 		}
+
+		TOptional<Frontend::FMetaSoundAssetClassInfo> TryGetClassInfoFromLoadedAsset(const FTopLevelAssetPath& InAssetPath)
+		{
+			using namespace Frontend;
+
+			const FSoftObjectPath ObjectPath(InAssetPath);
+			if (const UObject* Object = ObjectPath.ResolveObject())
+			{
+				TScriptInterface<const IMetaSoundDocumentInterface> DocInterface(Object);
+				if (DocInterface.GetObject())
+				{
+					return FMetaSoundAssetClassInfo(*DocInterface.GetInterface());
+				}
+			}
+
+			return { };
+		}
 	} // AssetSubsystemPrivate
 
 	FMetaSoundAssetManager::~FMetaSoundAssetManager()
@@ -263,7 +291,7 @@ namespace Metasound::Engine
 		TScriptInterface<const IMetaSoundDocumentInterface> ScriptInterface(&InObject);
 		check(ScriptInterface.GetObject() != nullptr);
 		const IMetaSoundDocumentInterface& DocInterface = *ScriptInterface;
-		return AssetSubsystemPrivate::AddClassInfo(TagDataMapCriticalSection, ClassInfoMap, FMetaSoundAssetClassInfo(DocInterface));
+		return AssetSubsystemPrivate::AddClassInfo(GetTransactor(), TagDataMapCriticalSection, ClassInfoMap, FMetaSoundAssetClassInfo(DocInterface));
 	}
 
 	void FMetaSoundAssetManager::AddOrLoadAndUpdateFromObjectAsync(const FAssetData& InAssetData, Frontend::IMetaSoundAssetManager::FOnUpdatedAssetLoaded&& OnUpdatedAssetLoaded)
@@ -333,7 +361,7 @@ namespace Metasound::Engine
 			FMetaSoundAssetClassInfo ClassInfo(InAssetData);
 			if (ClassInfo.bIsValid)
 			{
-				AssetSubsystemPrivate::AddClassInfo(TagDataMapCriticalSection, ClassInfoMap, MoveTemp(ClassInfo));
+				AssetSubsystemPrivate::AddClassInfo(GetTransactor(), TagDataMapCriticalSection, ClassInfoMap, MoveTemp(ClassInfo));
 			}
 			else
 			{
@@ -423,28 +451,24 @@ namespace Metasound::Engine
 	{
 		using namespace Frontend;
 
-		const FSoftObjectPath ObjectPath(InPath);
-		if (const UObject* Object = ObjectPath.ResolveObject())
+		TOptional<FMetaSoundAssetClassInfo> ClassInfo = AssetSubsystemPrivate::TryGetClassInfoFromLoadedAsset(InPath);
+		if (ClassInfo.IsSet())
 		{
-			TScriptInterface<const IMetaSoundDocumentInterface> DocInterface(Object);
-			if (DocInterface.GetObject())
-			{
-				return FMetaSoundAssetClassInfo(*DocInterface.GetInterface());
-			}
+			return *ClassInfo;
 		}
 
 		FAssetData AssetData;
 		const FAssetRegistryModule& AssetRegistryModule = FModuleManager::GetModuleChecked<FAssetRegistryModule>("AssetRegistry");
-		UE::AssetRegistry::EExists Exists = AssetRegistryModule.TryGetAssetByObjectPath(ObjectPath, AssetData);
+		UE::AssetRegistry::EExists Exists = AssetRegistryModule.TryGetAssetByObjectPath(FSoftObjectPath(InPath), AssetData);
 		if (Exists == UE::AssetRegistry::EExists::Exists)
 		{
 			FMetaSoundAssetKey AssetKey;
 			if (FMetaSoundAssetClassInfo::TryGetAssetKey(AssetData, AssetKey))
 			{
-				const TArray<Frontend::FMetaSoundAssetClassInfo> ClassInfo = FindAssetClassInfoInternal(AssetKey);
-				if (ClassInfo.Num() == 1)
+				const TArray<Frontend::FMetaSoundAssetClassInfo> CachedClassInfo = FindAssetClassInfoInternal(AssetKey);
+				if (CachedClassInfo.Num() == 1)
 				{
-					return ClassInfo.Last();
+					return CachedClassInfo.Last();
 				}
 			}
 
@@ -454,6 +478,25 @@ namespace Metasound::Engine
 		return { };
 	}
 
+	TArray<Frontend::FMetaSoundClassInfo> FMetaSoundAssetManager::FindRegisteredClassInfo(const FMetaSoundAssetKey& InKey) const
+	{
+		using namespace Frontend;
+
+		TArray<FMetaSoundAssetClassInfo> AssetClassInfo = FindAssetClassInfoInternal(InKey);
+		TArray<FMetaSoundClassInfo> ClassInfo;
+		Algo::Transform(AssetClassInfo, ClassInfo, [](const FMetaSoundAssetClassInfo& AssetInfo)
+		{
+			// If loaded, use that as its potentially "fresher" than the cached info in the AssetManager (provided at last load/save)
+			TOptional<FMetaSoundAssetClassInfo> LoadedInfo = AssetSubsystemPrivate::TryGetClassInfoFromLoadedAsset(AssetInfo.AssetPath);
+			if (LoadedInfo.IsSet())
+			{
+				return *LoadedInfo;
+			}
+			return AssetInfo;
+		});
+		return ClassInfo;
+	}
+
 	TArray<Frontend::FMetaSoundAssetClassInfo> FMetaSoundAssetManager::FindAssetClassInfoInternal(const FMetaSoundAssetKey& InKey) const
 	{
 		FScopeLock Lock(&TagDataMapCriticalSection);
@@ -687,11 +730,11 @@ namespace Metasound::Engine
 				if (!bIsRegistered)
 				{
 					// Don't report failure if a matching class with a matching major version and higher minor version exists (it will be autoupdated) 
-					FMetasoundFrontendClass FrontendClass;
-					const bool bDidFindClassWithName = ISearchEngine::Get().FindClassWithHighestVersion(AssetKey.ClassName.ToNodeClassName(), FrontendClass);
+					FMetaSoundClassInfo ClassInfo;
+					const bool bDidFindClassWithName = ISearchEngine::Get().FindClassWithHighestVersion(AssetKey.ClassName.ToNodeClassName(), ClassInfo);
 					if (!(bDidFindClassWithName &&
-						AssetKey.Version.Major == FrontendClass.Metadata.GetVersion().Major &&
-						AssetKey.Version.Minor < FrontendClass.Metadata.GetVersion().Minor))
+						AssetKey.Version.Major == ClassInfo.Version.Major &&
+						AssetKey.Version.Minor < ClassInfo.Version.Minor))
 					{
 						if (IsInitialAssetScanComplete())
 						{
@@ -778,13 +821,13 @@ namespace Metasound::Engine
 
 		{
 			const FMetaSoundAssetKey OldAssetKey(ClassMetadata.GetClassName(), ClassMetadata.GetVersion());
-			AssetSubsystemPrivate::RemoveClassInfo(TagDataMapCriticalSection, ClassInfoMap, OldAssetKey, Path);
+			AssetSubsystemPrivate::RemoveClassInfo(GetTransactor(), TagDataMapCriticalSection, ClassInfoMap, OldAssetKey, Path);
 		}
 
 		Builder.GenerateNewClassName();
 
 		{
-			AssetSubsystemPrivate::AddClassInfo(TagDataMapCriticalSection, ClassInfoMap, FMetaSoundAssetClassInfo(*DocInterface));
+			AssetSubsystemPrivate::AddClassInfo(GetTransactor(), TagDataMapCriticalSection, ClassInfoMap, FMetaSoundAssetClassInfo(*DocInterface));
 		}
 
 		AssetBase->UpdateAndRegisterForExecution(GetRegistrationOptions());
@@ -941,7 +984,7 @@ namespace Metasound::Engine
 
 		const FTopLevelAssetPath AssetPath(&InObject);
 		const FMetaSoundAssetKey AssetKey(Metadata.GetClassName(), Metadata.GetVersion());
-		AssetSubsystemPrivate::RemoveClassInfo(TagDataMapCriticalSection, ClassInfoMap, AssetKey, AssetPath);
+		AssetSubsystemPrivate::RemoveClassInfo(GetTransactor(), TagDataMapCriticalSection, ClassInfoMap, AssetKey, AssetPath);
 	}
 
 	void FMetaSoundAssetManager::RemoveAsset(const FAssetData& InAssetData)
@@ -958,7 +1001,7 @@ namespace Metasound::Engine
 			}
 
 			const FMetaSoundAssetKey AssetKey(ClassInfo.ClassName, ClassInfo.Version);
-			AssetSubsystemPrivate::RemoveClassInfo(TagDataMapCriticalSection, ClassInfoMap, AssetKey, ClassInfo.AssetPath);
+			AssetSubsystemPrivate::RemoveClassInfo(GetTransactor(), TagDataMapCriticalSection, ClassInfoMap, AssetKey, ClassInfo.AssetPath);
 		}
 	}
 
@@ -983,8 +1026,8 @@ namespace Metasound::Engine
 		{
 			const FMetaSoundAssetKey AssetKey(ClassInfo.ClassName, ClassInfo.Version);
 			const FTopLevelAssetPath OldPath(InOldObjectPath);
-			AssetSubsystemPrivate::RemoveClassInfo(TagDataMapCriticalSection, ClassInfoMap, AssetKey, OldPath);
-			AssetSubsystemPrivate::AddClassInfo(TagDataMapCriticalSection, ClassInfoMap, MoveTemp(ClassInfo));
+			AssetSubsystemPrivate::RemoveClassInfo(GetTransactor(), TagDataMapCriticalSection, ClassInfoMap, AssetKey, OldPath);
+			AssetSubsystemPrivate::AddClassInfo(GetTransactor(), TagDataMapCriticalSection, ClassInfoMap, MoveTemp(ClassInfo));
 		}
 	}
 
@@ -1001,8 +1044,8 @@ namespace Metasound::Engine
 			return bReferencesReplaced;
 		}
 
-		FMetasoundFrontendClass NewClass;
-		const bool bNewClassExists = ISearchEngine::Get().FindClassWithHighestVersion(NewClassKey.ClassName, NewClass);
+		FMetaSoundClassInfo ClassInfo;
+		const bool bNewClassExists = ISearchEngine::Get().FindClassWithHighestVersion(NewClassKey.ClassName, ClassInfo);
 		if (bNewClassExists)
 		{
 			TArray<FDirectoryPath> Directories;
@@ -1231,7 +1274,7 @@ namespace Metasound::Engine
 					{
 						FMetasoundFrontendRegistryContainer::Get()->UnregisterNode(RegistryKey);
 						const FTopLevelAssetPath AssetPath(AssetData.PackageName, AssetData.AssetName);
-						AssetSubsystemPrivate::RemoveClassInfo(TagDataMapCriticalSection, ClassInfoMap, AssetKey, AssetPath);
+						AssetSubsystemPrivate::RemoveClassInfo(GetTransactor(), TagDataMapCriticalSection, ClassInfoMap, AssetKey, AssetPath);
 					}
 				}
 			}
diff --git a/Engine/Plugins/Runtime/Metasound/Source/MetasoundEngine/Public/MetasoundAssetSubsystem.h b/Engine/Plugins/Runtime/Metasound/Source/MetasoundEngine/Public/MetasoundAssetSubsystem.h
--- a/Engine/Plugins/Runtime/Metasound/Source/MetasoundEngine/Public/MetasoundAssetSubsystem.h
+++ b/Engine/Plugins/Runtime/Metasound/Source/MetasoundEngine/Public/MetasoundAssetSubsystem.h
@@ -32,7 +32,7 @@ namespace Metasound::Engine
 
 
 	class FMetaSoundAssetManager :
-		public Frontend::IMetaSoundAssetManager,
+		public Frontend::FMetaSoundAssetManagerBase,
 		public FGCObject
 	{
 	public:
@@ -76,6 +76,12 @@ PRAGMA_ENABLE_DEPRECATION_WARNINGS
 		 **/
 		UE_API Frontend::FMetaSoundAssetClassInfo FindAssetClassInfo(const FTopLevelAssetPath& InPath) const;
 
+		// Returns all asset query info with the given asset key (may be more than one due to errors in duplication
+		// of assets, improper plugin migration, etc. but should usually just be a single entry) registered
+		// with the manager. Info may not be returned if the asset has yet to be registered with the MetaSoundAssetManager
+		// or may be out-of-date if edited after last registration attempt via AddOrUpdate... calls.
+		UE_API virtual TArray<Frontend::FMetaSoundClassInfo> FindRegisteredClassInfo(const FMetaSoundAssetKey& InKey) const override;
+
 		UE_API FMetaSoundAssetKey GetAssetKey(const FSoftObjectPath& InObjectPath) const;
 
 #if WITH_EDITOR
diff --git a/Engine/Plugins/Runtime/Metasound/Source/MetasoundFrontend/Private/MetasoundAssetManager.cpp b/Engine/Plugins/Runtime/Metasound/Source/MetasoundFrontend/Private/MetasoundAssetManager.cpp
--- a/Engine/Plugins/Runtime/Metasound/Source/MetasoundFrontend/Private/MetasoundAssetManager.cpp
+++ b/Engine/Plugins/Runtime/Metasound/Source/MetasoundFrontend/Private/MetasoundAssetManager.cpp
@@ -4,6 +4,8 @@
 
 #include "Algo/Copy.h"
 #include "AssetRegistry/AssetData.h"
+#include "JsonObjectConverter.h"
+#include "MetasoundAssetManagerTransactionPrivate.h"
 #include "MetasoundAssetTagCollections.h"
 #include "MetasoundFrontendDocumentVersioning.h"
 #include "MetasoundFrontendRegistryKey.h"
@@ -165,6 +167,7 @@ void FMetaSoundDocumentInfo::ExportToContext(FAssetRegistryTagsContext& OutConte
 namespace Metasound::Frontend
 {
 	FMetaSoundAssetClassInfo::FMetaSoundAssetClassInfo(const IMetaSoundDocumentInterface& InDocInterface)
+		: FMetaSoundClassInfo(EMetasoundFrontendClassType::External)
 	{
 		InitFromDocument(InDocInterface);
 	}
@@ -305,6 +308,17 @@ namespace Metasound::Frontend
 		check(!AssetManagerPrivate::Instance.IsValid());
 		AssetManagerPrivate::Instance = MoveTemp(InInterface);
 	}
+
+	FMetaSoundAssetManagerBase::FMetaSoundAssetManagerBase()
+		: Transactor(MakeUnique<FAssetManagerTransactor>())
+	{
+	}
+
+	IMetaSoundAssetTransactor& FMetaSoundAssetManagerBase::GetTransactor()
+	{
+		check(Transactor);
+		return *Transactor.Get();
+	}
 } // namespace Metasound::Frontend
 
 FMetaSoundAssetKey::FMetaSoundAssetKey(const FMetasoundFrontendClassName& InClassName, const FMetasoundFrontendVersionNumber& InVersion)
diff --git a/Engine/Plugins/Runtime/Metasound/Source/MetasoundFrontend/Private/MetasoundAssetManagerTransaction.cpp b/Engine/Plugins/Runtime/Metasound/Source/MetasoundFrontend/Private/MetasoundAssetManagerTransaction.cpp
new file mode 100644
--- /dev/null
+++ b/Engine/Plugins/Runtime/Metasound/Source/MetasoundFrontend/Private/MetasoundAssetManagerTransaction.cpp
@@ -0,0 +1,81 @@
+// Copyright Epic Games, Inc. All Rights Reserved.
+
+#include "MetasoundAssetManagerTransaction.h"
+#include "MetasoundAssetManagerTransactionPrivate.h"
+
+
+namespace Metasound::Frontend
+{
+	FString FAssetManagerTransaction::LexToString(FAssetManagerTransaction::ETransactionType InTransactionType)
+	{
+		using namespace Metasound;
+
+		switch (InTransactionType)
+		{
+		case ETransactionType::AssetRegistration:
+		{
+			return TEXT("Asset Registration");
+		}
+
+		case ETransactionType::AssetUnregistration:
+		{
+			return TEXT("Asset Unregistration");
+		}
+
+		case ETransactionType::Invalid:
+		{
+			return TEXT("Invalid");
+		}
+
+		default:
+		{
+			checkNoEntry();
+			return TEXT("");
+		}
+		}
+	}
+
+	FAssetManagerTransaction::FAssetManagerTransaction(FAssetManagerTransaction::ETransactionType InType, const FMetaSoundAssetKey& InKey, FAssetManagerTransaction::FTimeType InTimestamp)
+		: Type(InType)
+		, Key(InKey)
+		, Timestamp(InTimestamp)
+	{
+	}
+
+	FAssetManagerTransaction::ETransactionType FAssetManagerTransaction::GetTransactionType() const
+	{
+		return Type;
+	}
+
+	FNodeClassRegistryKey FAssetManagerTransaction::GetNodeRegistryKey() const
+	{
+		return Key;
+	}
+
+	FAssetManagerTransaction::FTimeType FAssetManagerTransaction::GetTimestamp() const
+	{
+		return Timestamp;
+	}
+
+	FAssetManagerTransactor::FAssetManagerTransactor()
+		: TransactionBuffer(MakeShared<FAssetManagerTransactionBuffer>())
+	{
+	}
+
+	TUniquePtr<FAssetManagerTransactionStream> FAssetManagerTransactor::CreateTransactionStream()
+	{
+		return MakeUnique<FAssetManagerTransactionStream>(TransactionBuffer);
+	}
+
+	void FAssetManagerTransactor::RegisterAsset(FMetaSoundAssetKey Key)
+	{
+		const FAssetManagerTransaction::FTimeType Timestamp = FPlatformTime::Cycles64();
+		TransactionBuffer->AddTransaction(FAssetManagerTransaction(FAssetManagerTransaction::ETransactionType::AssetRegistration, Key, Timestamp));
+	}
+
+	void FAssetManagerTransactor::UnregisterAsset(FMetaSoundAssetKey Key)
+	{
+		const FAssetManagerTransaction::FTimeType Timestamp = FPlatformTime::Cycles64();
+		TransactionBuffer->AddTransaction(FAssetManagerTransaction(FAssetManagerTransaction::ETransactionType::AssetUnregistration, Key, Timestamp));
+	}
+} // namespace Metasound::Frontend
diff --git a/Engine/Plugins/Runtime/Metasound/Source/MetasoundFrontend/Private/MetasoundAssetManagerTransactionPrivate.h b/Engine/Plugins/Runtime/Metasound/Source/MetasoundFrontend/Private/MetasoundAssetManagerTransactionPrivate.h
new file mode 100644
--- /dev/null
+++ b/Engine/Plugins/Runtime/Metasound/Source/MetasoundFrontend/Private/MetasoundAssetManagerTransactionPrivate.h
@@ -0,0 +1,26 @@
+// Copyright Epic Games, Inc. All Rights Reserved.
+#pragma once
+
+#include "MetasoundAssetKey.h"
+#include "MetasoundAssetManagerTransaction.h"
+#include "MetasoundFrontendRegistryTransaction.h"
+
+
+namespace Metasound::Frontend
+{
+	using FAssetManagerTransactionBuffer = TTransactionBuffer<FAssetManagerTransaction>;
+	using FAssetManagerTransactionStream = TTransactionStream<FAssetManagerTransaction>;
+
+	class FAssetManagerTransactor : public IMetaSoundAssetTransactor
+	{
+	public:
+		FAssetManagerTransactor();
+
+		TUniquePtr<FAssetManagerTransactionStream> CreateTransactionStream();
+		virtual void RegisterAsset(FMetaSoundAssetKey Key) override;
+		virtual void UnregisterAsset(FMetaSoundAssetKey Key) override;
+
+	private:
+		TSharedRef<FAssetManagerTransactionBuffer> TransactionBuffer;
+	};
+} // namespace Metasound::Frontend
\ No newline at end of file
diff --git a/Engine/Plugins/Runtime/Metasound/Source/MetasoundFrontend/Private/MetasoundFrontendQuery.cpp b/Engine/Plugins/Runtime/Metasound/Source/MetasoundFrontend/Private/MetasoundFrontendQuery.cpp
--- a/Engine/Plugins/Runtime/Metasound/Source/MetasoundFrontend/Private/MetasoundFrontendQuery.cpp
+++ b/Engine/Plugins/Runtime/Metasound/Source/MetasoundFrontend/Private/MetasoundFrontendQuery.cpp
@@ -22,7 +22,12 @@ namespace Metasound::Frontend
 {
 	namespace ClassTagsPrivate
 	{
-		static const FLazyName AccessFlags = "AssetFlags";
+		static const FLazyName AccessFlags = "AccessFlags";
+
+#if WITH_EDITORONLY_DATA
+		static const FLazyName AssetFlags_Deprecated = "AssetFlags";
+#endif // WITH_EDITORONLY_DATA
+
 		static const FLazyName ClassVersion = "ClassVersion";
 
 #if WITH_EDITORONLY_DATA
@@ -44,7 +49,7 @@ namespace Metasound::Frontend
 			return false;
 		};
 	}
-}
+} // namespace Metasound::Frontend
 
 
 FMetaSoundClassSearchInfo::FMetaSoundClassSearchInfo(const FMetasoundFrontendClassMetadata& InClassMetadata)
@@ -88,6 +93,33 @@ FMetaSoundClassInterfaceInfo::FMetaSoundClassInterfaceInfo(const IMetaSoundDocum
 #endif // WITH_EDITORONLY_DATA
 }
 
+FMetaSoundClassInterfaceInfo::FMetaSoundClassInterfaceInfo(const FMetasoundFrontendClass& InClass)
+{
+#if WITH_EDITORONLY_DATA
+	using namespace Metasound::Frontend;
+
+	if (!IsRunningCookCommandlet())
+	{
+		SearchInfo = FMetaSoundClassSearchInfo(InClass.Metadata);
+
+		{
+			auto CopyVertexInfo = [](const FMetasoundFrontendClassVertex& Vertex)
+			{
+				return FMetaSoundClassVertexInfo
+				{
+					.Name = Vertex.Name,
+					.TypeName = Vertex.TypeName,
+					.AccessType = Vertex.AccessType
+				};
+			};
+
+			Algo::Transform(InClass.GetDefaultInterface().Inputs, Inputs, CopyVertexInfo);
+			Algo::Transform(InClass.GetDefaultInterface().Outputs, Outputs, CopyVertexInfo);
+		}
+	}
+#endif // WITH_EDITORONLY_DATA
+}
+
 FMetaSoundClassInterfaceInfo::FMetaSoundClassInterfaceInfo(const FAssetData& InAssetData, bool& bOutIsValid)
 {
 	bOutIsValid = true;
@@ -156,13 +188,15 @@ namespace Metasound
 {
 	namespace Frontend
 	{
-		FMetaSoundClassInfo::FMetaSoundClassInfo()
-			: bIsValid(0)
+		FMetaSoundClassInfo::FMetaSoundClassInfo(EMetasoundFrontendClassType InClassType)
+			: ClassType(InClassType)
+			, bIsValid(0)
 		{
 		}
 
 		FMetaSoundClassInfo::FMetaSoundClassInfo(const FAssetData& InAssetData)
-			: bIsValid(0)
+			: ClassType(EMetasoundFrontendClassType::External)
+			, bIsValid(0)
 		{
 			bool bSuccess = true;
 
@@ -172,7 +206,15 @@ namespace Metasound
 
 			{
 				int32 AccessFlagsValue = 0;
-				const bool bFoundFlags = InAssetData.GetTagValue(ClassTagsPrivate::AccessFlags, AccessFlagsValue);
+				bool bFoundFlags = InAssetData.GetTagValue(ClassTagsPrivate::AccessFlags, AccessFlagsValue);
+
+#if WITH_EDITORONLY_DATA
+				if (!bFoundFlags)
+				{
+					bFoundFlags = InAssetData.GetTagValue(ClassTagsPrivate::AssetFlags_Deprecated, AccessFlagsValue);
+				}
+#endif // WITH_EDITORONLY_DATA
+
 				if (bFoundFlags)
 				{
 					AccessFlags = static_cast<EMetasoundFrontendClassAccessFlags>((uint16)(AccessFlagsValue));
@@ -192,6 +234,19 @@ namespace Metasound
 			bIsValid = (uint8)(bSuccess);
 		}
 
+		FMetaSoundClassInfo::FMetaSoundClassInfo(const FMetasoundFrontendClass& InClass)
+			: ClassName(InClass.Metadata.GetClassName())
+			, ClassType(InClass.Metadata.GetType())
+			, Version(InClass.Metadata.GetVersion())
+			, bIsValid(InClass.Metadata.GetType() != EMetasoundFrontendClassType::Invalid)
+		{
+#if WITH_EDITORONLY_DATA
+			InterfaceInfo = FMetaSoundClassInterfaceInfo(InClass);
+#endif // WITH_EDITORONLY_DATA
+
+			AccessFlags = InClass.Metadata.GetAccessFlags();
+		}
+
 		void FMetaSoundClassInfo::ExportToContext(FAssetRegistryTagsContext& OutContext) const
 		{
 			using FAssetRegistryTag = UObject::FAssetRegistryTag;
@@ -214,6 +269,7 @@ namespace Metasound
 
 			AccessFlags = Document.RootGraph.Metadata.GetAccessFlags();
 			ClassName = Document.RootGraph.Metadata.GetClassName();
+			ClassType = EMetasoundFrontendClassType::External;
 			Version = ClassMetadata.GetVersion();
 
 #if WITH_EDITORONLY_DATA
diff --git a/Engine/Plugins/Runtime/Metasound/Source/MetasoundFrontend/Private/MetasoundFrontendQuerySteps.cpp b/Engine/Plugins/Runtime/Metasound/Source/MetasoundFrontend/Private/MetasoundFrontendQuerySteps.cpp
--- a/Engine/Plugins/Runtime/Metasound/Source/MetasoundFrontend/Private/MetasoundFrontendQuerySteps.cpp
+++ b/Engine/Plugins/Runtime/Metasound/Source/MetasoundFrontend/Private/MetasoundFrontendQuerySteps.cpp
@@ -3,6 +3,8 @@
 #include "MetasoundFrontendQuerySteps.h"
 
 #include "Algo/MaxElement.h"
+#include "MetasoundAssetManager.h"
+#include "MetasoundAssetManagerTransactionPrivate.h"
 #include "MetasoundFrontend.h"
 #include "MetasoundFrontendDocument.h"
 #include "MetasoundFrontendNodeClassRegistry.h"
@@ -25,7 +27,7 @@ namespace Metasound
 		{
 			using namespace Frontend;
 
-			auto AddEntry = [&OutValues](const FNodeRegistryTransaction& InTransaction)
+			auto AddEntry = [&OutValues](const FNodeClassRegistryTransaction& InTransaction)
 			{
 				OutValues.Emplace(TInPlaceType<FNodeClassRegistryTransaction>(), InTransaction);
 			};
@@ -40,6 +42,37 @@ namespace Metasound
 		TUniquePtr<Frontend::FNodeClassRegistryTransactionStream> TransactionStream;
 	};
 
+	class FAssetRegistrationEventsPimpl : public IFrontendQueryStreamStep
+	{
+	public:
+		FAssetRegistrationEventsPimpl()
+		{
+			using namespace Frontend;
+
+			FMetaSoundAssetManagerBase& Manager = static_cast<FMetaSoundAssetManagerBase&>(IMetaSoundAssetManager::GetChecked());
+			FAssetManagerTransactor& Transactor = static_cast<FAssetManagerTransactor&>(Manager.GetTransactor());
+			TransactionStream = Transactor.CreateTransactionStream();
+		}
+
+		virtual void Stream(TArray<FFrontendQueryValue>& OutValues) override
+		{
+			using namespace Frontend;
+
+			auto AddEntry = [&OutValues](const FAssetManagerTransaction& InTransaction)
+			{
+				OutValues.Emplace(TInPlaceType<FAssetManagerTransaction>(), InTransaction);
+			};
+
+			if (TransactionStream.IsValid())
+			{
+				TransactionStream->Stream(AddEntry);
+			}
+		}
+
+	private:
+		TUniquePtr<Frontend::FAssetManagerTransactionStream> TransactionStream;
+	};
+
 	FNodeClassRegistrationEvents::FNodeClassRegistrationEvents()
 	: Pimpl(MakePimpl<FNodeClassRegistrationEventsPimpl>())
 	{
@@ -50,20 +83,140 @@ namespace Metasound
 		Pimpl->Stream(OutValues);
 	}
 
-	FFrontendQueryKey FMapRegistrationEventsToNodeRegistryKeys::Map(const FFrontendQueryEntry& InEntry) const 
+	FAssetRegistrationEvents::FAssetRegistrationEvents()
+		: Pimpl(MakePimpl<FAssetRegistrationEventsPimpl>())
+	{
+	}
+
+	void FAssetRegistrationEvents::Stream(TArray<FFrontendQueryValue>& OutValues)
+	{
+		Pimpl->Stream(OutValues);
+	}
+
+	FFrontendQueryKey FMapAssetRegistrationEventsToNodeRegistryKeys::Map(const FFrontendQueryEntry& InEntry) const 
 	{
 		using namespace Frontend;
 
 		FNodeRegistryKey RegistryKey;
 
-		if (ensure(InEntry.Value.IsType<FNodeRegistryTransaction>()))
+		if (ensure(InEntry.Value.IsType<FAssetManagerTransaction>()))
 		{
-			RegistryKey = InEntry.Value.Get<FNodeRegistryTransaction>().GetNodeRegistryKey();
+			RegistryKey = InEntry.Value.Get<FAssetManagerTransaction>().GetNodeRegistryKey();
 		}
 
 		return FFrontendQueryKey(RegistryKey.ToString());
 	}
 
+	FFrontendQueryKey FMapClassRegistrationEventsToNodeRegistryKeys::Map(const FFrontendQueryEntry& InEntry) const
+	{
+		using namespace Frontend;
+
+		FNodeRegistryKey RegistryKey;
+
+		if (ensure(InEntry.Value.IsType<FNodeClassRegistryTransaction>()))
+		{
+			RegistryKey = InEntry.Value.Get<FNodeClassRegistryTransaction>().GetNodeRegistryKey();
+		}
+
+		return FFrontendQueryKey(RegistryKey.ToString());
+	}
+
+	void FReduceAssetRegistrationEventsToCurrentStatus::Reduce(const FFrontendQueryKey& InKey, FFrontendQueryPartition& InOutEntries) const
+	{
+		using namespace Frontend;
+
+		// Track number of times each type of entry has been seen
+		struct FTransactionStats
+		{
+			int32 Num = 0;
+			const FFrontendQueryEntry* MostRecent = nullptr;
+
+			void Update(const FFrontendQueryEntry& InEntry)
+			{
+				Num++;
+				if (nullptr == MostRecent)
+				{
+					MostRecent = &InEntry;
+				}
+				else
+				{
+					if (FReduceAssetRegistrationEventsToCurrentStatus::GetTransactionTimestamp(InEntry) > FReduceAssetRegistrationEventsToCurrentStatus::GetTransactionTimestamp(*MostRecent))
+					{
+						MostRecent = &InEntry;
+					}
+				}
+			}
+		};
+
+		// Accumulate the counts of each transaction type
+		FTransactionStats NodeRegistrationStats;
+		FTransactionStats NodeUnregistrationStats;
+
+		for (const FFrontendQueryEntry& Entry : InOutEntries)
+		{
+			const FAssetManagerTransaction& Transaction = Entry.Value.Get<FAssetManagerTransaction>();
+			switch (Transaction.GetTransactionType())
+			{
+			case FAssetManagerTransaction::ETransactionType::AssetRegistration:
+				NodeRegistrationStats.Update(Entry);
+				break;
+
+			case FAssetManagerTransaction::ETransactionType::AssetUnregistration:
+				NodeUnregistrationStats.Update(Entry);
+				break;
+
+			default:
+			{
+				checkNoEntry();
+			}
+			}
+		}
+
+		// Use the final number of each transaction type to determine the final state
+		FFrontendQueryPartition RemainingEntries;
+		if (NodeRegistrationStats.Num > NodeUnregistrationStats.Num)
+		{
+			// Registration trumps migration, so we do not consider migration at all here.
+			RemainingEntries.Add(*NodeRegistrationStats.MostRecent);
+		}
+		else
+		{
+			// Make sure to propagate unregistration actions since reduce steps may get called
+			// on subpartitions. 
+			if (NodeUnregistrationStats.Num > NodeRegistrationStats.Num)
+			{
+				RemainingEntries.Add(*NodeUnregistrationStats.MostRecent);
+			}
+		}
+
+		InOutEntries = MoveTemp(RemainingEntries);
+	}
+
+	FReduceAssetRegistrationEventsToCurrentStatus::FTimeType FReduceAssetRegistrationEventsToCurrentStatus::GetTransactionTimestamp(const FFrontendQueryEntry& InEntry)
+	{
+		using namespace Frontend;
+
+		if (ensure(InEntry.Value.IsType<FAssetManagerTransaction>()))
+		{
+			return InEntry.Value.Get<FAssetManagerTransaction>().GetTimestamp();
+		}
+		return 0;
+	}
+
+	bool FReduceAssetRegistrationEventsToCurrentStatus::IsValidTransactionOfType(Frontend::FAssetManagerTransaction::ETransactionType InType, const FFrontendQueryEntry* InEntry)
+	{
+		using namespace Frontend;
+
+		if (nullptr != InEntry)
+		{
+			if (InEntry->Value.IsType<FAssetManagerTransaction>())
+			{
+				return InEntry->Value.Get<FAssetManagerTransaction>().GetTransactionType() == InType;
+			}
+		}
+		return false;
+	}
+
 	void FReduceRegistrationEventsToCurrentStatus::Reduce(const FFrontendQueryKey& InKey, FFrontendQueryPartition& InOutEntries) const
 	{
 		using namespace Frontend;
@@ -159,45 +312,67 @@ namespace Metasound
 	{
 		using namespace Frontend;
 
-		if (ensure(InEntry.Value.IsType<FNodeRegistryTransaction>()))
+		if (ensure(InEntry.Value.IsType<FNodeClassRegistryTransaction>()))
 		{
-			return InEntry.Value.Get<FNodeRegistryTransaction>().GetTimestamp();
+			return InEntry.Value.Get<FNodeClassRegistryTransaction>().GetTimestamp();
 		}
 		return 0;
 	}
 
-	bool FReduceRegistrationEventsToCurrentStatus::IsValidTransactionOfType(Frontend::FNodeRegistryTransaction::ETransactionType InType, const FFrontendQueryEntry* InEntry)
+	bool FReduceRegistrationEventsToCurrentStatus::IsValidTransactionOfType(Frontend::FNodeClassRegistryTransaction::ETransactionType InType, const FFrontendQueryEntry* InEntry)
 	{
 		using namespace Frontend;
 
 		if (nullptr != InEntry)
 		{
-			if (InEntry->Value.IsType<FNodeRegistryTransaction>())
+			if (InEntry->Value.IsType<FNodeClassRegistryTransaction>())
 			{
-				return InEntry->Value.Get<FNodeRegistryTransaction>().GetTransactionType() == InType;
+				return InEntry->Value.Get<FNodeClassRegistryTransaction>().GetTransactionType() == InType;
 			}
 		}
 		return false;
 	}
 
-	void FTransformRegistrationEventsToClasses::Transform(FFrontendQueryEntry::FValue& InValue) const
+	void FTransformAssetRegistrationEventsToClassInfo::Transform(FFrontendQueryEntry::FValue& InValue) const
 	{
 		using namespace Frontend;
 
-		FMetasoundFrontendClass FrontendClass;
+		FMetaSoundClassInfo ClassInfo(EMetasoundFrontendClassType::Invalid);
 
-		if (ensure(InValue.IsType<FNodeRegistryTransaction>()))
+		if (ensure(InValue.IsType<FAssetManagerTransaction>()))
 		{
-			const FNodeRegistryTransaction& Transaction = InValue.Get<FNodeRegistryTransaction>();
-			
-			if (Transaction.GetTransactionType() == Frontend::FNodeRegistryTransaction::ETransactionType::NodeRegistration)
+			const FAssetManagerTransaction& Transaction = InValue.Get<FAssetManagerTransaction>();
+
+			if (Transaction.GetTransactionType() == Frontend::FAssetManagerTransaction::ETransactionType::AssetRegistration)
+			{
+				// It's possible that the node is no longer registered (we're processing removals) 
+				// but that's okay because the returned default FrontendClass will be processed out later
+				ClassInfo = IMetaSoundAssetManager::GetChecked().FindRegisteredClassInfo(FMetaSoundAssetKey(Transaction.GetNodeRegistryKey())).Last();
+			}
+		}
+		InValue.Set<FMetaSoundClassInfo>(MoveTemp(ClassInfo));
+	}
+
+	void FTransformNodeClassRegistrationEventsToClassInfo::Transform(FFrontendQueryEntry::FValue& InValue) const
+	{
+		using namespace Frontend;
+
+		FMetaSoundClassInfo ClassInfo(EMetasoundFrontendClassType::Invalid);
+
+		if (ensure(InValue.IsType<FNodeClassRegistryTransaction>()))
+		{
+			const FNodeClassRegistryTransaction& Transaction = InValue.Get<FNodeClassRegistryTransaction>();
+
+			if (Transaction.GetTransactionType() == Frontend::FNodeClassRegistryTransaction::ETransactionType::NodeRegistration)
 			{
 				// It's possible that the node is no longer registered (we're processing removals) 
 				// but that's okay because the returned default FrontendClass will be processed out later
-				FMetasoundFrontendRegistryContainer::Get()->FindFrontendClassFromRegistered(Transaction.GetNodeRegistryKey(), FrontendClass);
+				FMetasoundFrontendClass Class;
+				FMetasoundFrontendRegistryContainer::Get()->FindFrontendClassFromRegistered(Transaction.GetNodeRegistryKey(), Class);
+				ClassInfo = FMetaSoundClassInfo(Class);
 			}
 		}
-		InValue.Set<FMetasoundFrontendClass>(MoveTemp(FrontendClass));
+		InValue.Set<FMetaSoundClassInfo>(MoveTemp(ClassInfo));
 	}
 
 	FFilterClassesByInputVertexDataType::FFilterClassesByInputVertexDataType(const FName& InTypeName)
@@ -207,14 +382,19 @@ namespace Metasound
 
 	bool FFilterClassesByInputVertexDataType::Filter(const FFrontendQueryEntry& InEntry) const
 	{
-		check(InEntry.Value.IsType<FMetasoundFrontendClass>());
+#if WITH_EDITORONLY_DATA
+		using namespace Frontend;
+		check(InEntry.Value.IsType<FMetaSoundClassInfo>());
 
-		return InEntry.Value.Get<FMetasoundFrontendClass>().GetDefaultInterface().Inputs.ContainsByPredicate(
-			[this](const FMetasoundFrontendClassInput& InDesc)
+		return InEntry.Value.Get<FMetaSoundClassInfo>().InterfaceInfo.Inputs.ContainsByPredicate(
+			[this](const FMetaSoundClassVertexInfo& VertexInfo)
 			{
-				return InDesc.TypeName == InputVertexTypeName;
+				return VertexInfo.TypeName == InputVertexTypeName;
 			}
 		);
+#else
+		return false;
+#endif // WITH_EDITORONLY_DATA
 	}
 
 	FFilterClassesByOutputVertexDataType::FFilterClassesByOutputVertexDataType(const FName& InTypeName)
@@ -224,17 +404,26 @@ namespace Metasound
 
 	bool FFilterClassesByOutputVertexDataType::Filter(const FFrontendQueryEntry& InEntry) const
 	{
-		return InEntry.Value.Get<FMetasoundFrontendClass>().GetDefaultInterface().Outputs.ContainsByPredicate(
-			[this](const FMetasoundFrontendClassOutput& InDesc)
+#if WITH_EDITORONLY_DATA
+		using namespace Frontend;
+		check(InEntry.Value.IsType<FMetaSoundClassInfo>());
+
+		return InEntry.Value.Get<FMetaSoundClassInfo>().InterfaceInfo.Outputs.ContainsByPredicate(
+			[this](const FMetaSoundClassVertexInfo& VertexInfo)
 			{
-				return InDesc.TypeName == OutputVertexTypeName;
+				return VertexInfo.TypeName == OutputVertexTypeName;
 			}
 		);
+#else
+		return false;
+#endif // WITH_EDITORONLY_DATA
 	}
 
 	FFrontendQueryKey FMapClassesToClassName::Map(const FFrontendQueryEntry& InEntry) const 
 	{
-		return FFrontendQueryKey(InEntry.Value.Get<FMetasoundFrontendClass>().Metadata.GetClassName().GetFullName());
+		using namespace Frontend;
+
+		return FFrontendQueryKey(InEntry.Value.Get<FMetaSoundClassInfo>().ClassName.GetFullName());
 	}
 
 	FFilterClassesByClassID::FFilterClassesByClassID(const FGuid InClassID)
@@ -244,23 +433,27 @@ namespace Metasound
 
 	bool FFilterClassesByClassID::Filter(const FFrontendQueryEntry& InEntry) const
 	{
-		return InEntry.Value.Get<FMetasoundFrontendClass>().ID == ClassID;
+		return false;
 	}
 
 	FFrontendQueryKey FMapToFullClassName::Map(const FFrontendQueryEntry& InEntry) const
 	{
-		const FMetasoundFrontendClass& FrontendClass = InEntry.Value.Get<FMetasoundFrontendClass>();
-		return FFrontendQueryKey(FrontendClass.Metadata.GetClassName().GetFullName());
+		using namespace Frontend;
+
+		const FMetaSoundClassInfo& ClassInfo = InEntry.Value.Get<FMetaSoundClassInfo>();
+		return FFrontendQueryKey(ClassInfo.ClassName.GetFullName());
 	}
 
 	void FReduceClassesToHighestVersion::Reduce(const FFrontendQueryKey& InKey, FFrontendQueryPartition& InOutEntries) const
 	{
+		using namespace Frontend;
+
 		FFrontendQueryEntry* HighestVersionEntry = nullptr;
 		FMetasoundFrontendVersionNumber HighestVersion;
 
 		for (FFrontendQueryEntry& Entry : InOutEntries)
 		{
-			const FMetasoundFrontendVersionNumber& Version = Entry.Value.Get<FMetasoundFrontendClass>().Metadata.GetVersion();
+			const FMetasoundFrontendVersionNumber& Version = Entry.Value.Get<FMetaSoundClassInfo>().Version;
 
 			if (!HighestVersionEntry || HighestVersion < Version)
 			{
@@ -279,8 +472,8 @@ namespace Metasound
 
 	bool FSortClassesByVersion::Sort(const FFrontendQueryEntry& InEntryLHS, const FFrontendQueryEntry& InEntryRHS) const
 	{
-		const FMetasoundFrontendVersionNumber& VersionLHS = InEntryLHS.Value.Get<FMetasoundFrontendClass>().Metadata.GetVersion();
-		const FMetasoundFrontendVersionNumber& VersionRHS = InEntryRHS.Value.Get<FMetasoundFrontendClass>().Metadata.GetVersion();
+		const FMetasoundFrontendVersionNumber& VersionLHS = InEntryLHS.Value.Get<Frontend::FMetaSoundClassInfo>().Version;
+		const FMetasoundFrontendVersionNumber& VersionRHS = InEntryRHS.Value.Get<Frontend::FMetaSoundClassInfo>().Version;
 		return VersionLHS > VersionRHS;
 	}
 }
diff --git a/Engine/Plugins/Runtime/Metasound/Source/MetasoundFrontend/Private/MetasoundFrontendSearchEngine.cpp b/Engine/Plugins/Runtime/Metasound/Source/MetasoundFrontend/Private/MetasoundFrontendSearchEngine.cpp
--- a/Engine/Plugins/Runtime/Metasound/Source/MetasoundFrontend/Private/MetasoundFrontendSearchEngine.cpp
+++ b/Engine/Plugins/Runtime/Metasound/Source/MetasoundFrontend/Private/MetasoundFrontendSearchEngine.cpp
@@ -6,25 +6,72 @@
 
 #include "Algo/MaxElement.h"
 #include "Interfaces/MetasoundFrontendInterfaceRegistry.h"
+#include "MetasoundFrontendNodeClassRegistry.h"
 #include "MetasoundFrontendQuery.h"
 #include "MetasoundFrontendQuerySteps.h"
 #include "MetasoundFrontendRegistryTransaction.h"
 #include "MetasoundTrace.h"
 #include "MetasoundAssetManager.h"
 
-namespace Metasound
+namespace Metasound::Frontend
 {
-	namespace Frontend
+	bool ISearchEngine::FindClassWithHighestMinorVersion(const FMetasoundFrontendClassName& InName, int32 InMajorVersion, FMetasoundFrontendClass& OutClass)
 	{
-		ISearchEngine& ISearchEngine::Get()
+		FMetaSoundClassInfo ClassInfo;
+		const bool bFound = FindRegisteredClass(InName, InMajorVersion, ClassInfo);
+		if (bFound)
 		{
+			const FNodeClassRegistryKey Key(ClassInfo.ClassType, ClassInfo.ClassName, ClassInfo.Version);
+			return INodeClassRegistry::Get()->FindFrontendClassFromRegistered(Key, OutClass);
+		}
+
+		return false;
+	}
+
 #if WITH_EDITORONLY_DATA
-			static FSearchEngineEditorOnly SearchEngine;
-#else
-			static FSearchEngineCore SearchEngine;
-#endif
-			return SearchEngine;
+	TArray<FMetasoundFrontendClass> ISearchEngine::FindAllClasses(bool bInIncludeAllVersions)
+	{
+		constexpr bool bIncludeUnloadedAssets = false;
+		const TArray<FMetaSoundClassInfo> AllClasses = FindAllClasses(bInIncludeAllVersions
+			? ISearchEngine::EResultVersion::All
+			: ISearchEngine::EResultVersion::Highest,
+			bIncludeUnloadedAssets);
+
+		TArray<FMetasoundFrontendClass> ToReturn;
+		for (const FMetaSoundClassInfo& ClassInfo : AllClasses)
+		{
+			FMetasoundFrontendClass Class;
+			const FNodeClassRegistryKey Key(ClassInfo.ClassType, ClassInfo.ClassName, ClassInfo.Version);
+			if (INodeClassRegistry::Get()->FindFrontendClassFromRegistered(Key, Class))
+			{
+				ToReturn.Add(MoveTemp(Class));
+			}
 		}
+
+		return ToReturn;
 	}
-}
 
+	bool ISearchEngine::FindClassWithHighestVersion(const FMetasoundFrontendClassName& InName, FMetasoundFrontendClass& OutClass)
+	{
+		FMetaSoundClassInfo ClassInfo;
+		const bool bFound = FindClassWithHighestVersion(InName, ClassInfo);
+		if (bFound)
+		{
+			const FNodeClassRegistryKey Key(ClassInfo.ClassType, ClassInfo.ClassName, ClassInfo.Version);
+			return INodeClassRegistry::Get()->FindFrontendClassFromRegistered(Key, OutClass);
+		}
+
+		return false;
+	}
+#endif // WITH_EDITORONLY_DATA
+
+	ISearchEngine& ISearchEngine::Get()
+	{
+#if WITH_EDITORONLY_DATA
+		static FSearchEngineEditorOnly SearchEngine;
+#else
+		static FSearchEngineCore SearchEngine;
+#endif
+		return SearchEngine;
+	}
+} // namespace Metasound::Frontend
diff --git a/Engine/Plugins/Runtime/Metasound/Source/MetasoundFrontend/Private/MetasoundFrontendSearchEngineCore.cpp b/Engine/Plugins/Runtime/Metasound/Source/MetasoundFrontend/Private/MetasoundFrontendSearchEngineCore.cpp
--- a/Engine/Plugins/Runtime/Metasound/Source/MetasoundFrontend/Private/MetasoundFrontendSearchEngineCore.cpp
+++ b/Engine/Plugins/Runtime/Metasound/Source/MetasoundFrontend/Private/MetasoundFrontendSearchEngineCore.cpp
@@ -18,7 +18,7 @@ namespace Metasound::Frontend
 	{
 		FFrontendQueryKey CreateKeyFromFullClassNameAndMajorVersion(const FMetasoundFrontendClassName& InClassName, int32 InMajorVersion)
 		{
-			return FFrontendQueryKey(FString::Format(TEXT("{0}_v{1}"), {InClassName.ToString(), InMajorVersion}));
+			return FFrontendQueryKey(FString::Format(TEXT("{0}_v{1}"), { InClassName.ToString(), InMajorVersion }));
 		}
 
 		FFrontendQueryKey CreateKeyFromUClassPathName(const FTopLevelAssetPath& InClassPath)
@@ -28,10 +28,10 @@ namespace Metasound::Frontend
 
 		FFrontendQueryKey FMapClassToFullClassNameAndMajorVersion::Map(const FFrontendQueryEntry& InEntry) const
 		{
-			if (ensure(InEntry.Value.IsType<FMetasoundFrontendClass>()))
+			if (ensure(InEntry.Value.IsType<FMetaSoundClassInfo>()))
 			{
-				const FMetasoundFrontendClassMetadata& Metadata = InEntry.Value.Get<FMetasoundFrontendClass>().Metadata;
-				return CreateKeyFromFullClassNameAndMajorVersion(Metadata.GetClassName(), Metadata.GetVersion().Major);
+				const FMetaSoundClassInfo& ClassInfo = InEntry.Value.Get<FMetaSoundClassInfo>();
+				return CreateKeyFromFullClassNameAndMajorVersion(ClassInfo.ClassName, ClassInfo.Version.Major);
 			}
 			return FFrontendQueryKey();
 		}
@@ -201,9 +201,9 @@ namespace Metasound::Frontend
 
 		bool FRemoveDeprecatedClasses::Filter(const FFrontendQueryEntry& InEntry) const
 		{
-			if (InEntry.Value.IsType<FMetasoundFrontendClass>())
+			if (InEntry.Value.IsType<FMetaSoundClassInfo>())
 			{
-				const EMetasoundFrontendClassAccessFlags ThisAccessFlags = InEntry.Value.Get<FMetasoundFrontendClass>().Metadata.GetAccessFlags();
+				const EMetasoundFrontendClassAccessFlags ThisAccessFlags = InEntry.Value.Get<FMetaSoundClassInfo>().AccessFlags;
 				if (EnumHasAnyFlags(ThisAccessFlags, EMetasoundFrontendClassAccessFlags::Deprecated))
 				{
 					return false;
@@ -219,8 +219,8 @@ namespace Metasound::Frontend
 			{
 				if (InEntry.Value.IsType<FNodeRegistryTransaction>())
 				{
-					const FNodeClassInfo& Info = InEntry.Value.Get<FNodeRegistryTransaction>().GetNodeClassInfo();
-					return CreateKeyFromFullClassNameAndMajorVersion(Info.ClassName, Info.Version.Major);
+					const FNodeClassRegistryKey& RegistryKey = InEntry.Value.Get<FNodeRegistryTransaction>().GetNodeRegistryKey();
+					return CreateKeyFromFullClassNameAndMajorVersion(RegistryKey.ClassName, RegistryKey.Version.Major);
 				}
 				else
 				{
@@ -249,23 +249,23 @@ namespace Metasound::Frontend
 			InValue.Set<FMetasoundFrontendVersion>(Version);
 		}
 
-		TArray<FMetasoundFrontendClass> BuildArrayOfClassesFromPartition(const FFrontendQueryPartition& InPartition)
+		TArray<FMetaSoundClassInfo> BuildArrayOfClassesFromPartition(const FFrontendQueryPartition& InPartition)
 		{
-			TArray<FMetasoundFrontendClass> Result;
+			TArray<FMetaSoundClassInfo> Result;
 			for (const FFrontendQueryEntry& Entry : InPartition)
 			{
-				check(Entry.Value.IsType<FMetasoundFrontendClass>());
-				Result.Add(Entry.Value.Get<FMetasoundFrontendClass>());
+				check(Entry.Value.IsType<FMetaSoundClassInfo>());
+				Result.Add(Entry.Value.Get<FMetaSoundClassInfo>());
 			}
 			return Result;
 		}
 
-		FMetasoundFrontendClass BuildSingleClassFromPartition(const FFrontendQueryPartition& InPartition)
+		FMetaSoundClassInfo BuildSingleClassFromPartition(const FFrontendQueryPartition& InPartition)
 		{
-			FMetasoundFrontendClass MetaSoundClass;
+			FMetaSoundClassInfo MetaSoundClass(EMetasoundFrontendClassType::Invalid);
 			if (ensureMsgf(InPartition.Num() > 0, TEXT("Cannot retrieve class from empty partition")))
 			{
-				MetaSoundClass = InPartition.CreateConstIterator()->Value.Get<FMetasoundFrontendClass>();
+				MetaSoundClass = InPartition.CreateConstIterator()->Value.Get<FMetaSoundClassInfo>();
 			}
 			return MetaSoundClass;
 		}
@@ -309,7 +309,7 @@ namespace Metasound::Frontend
 		using namespace SearchEngineQuerySteps;
 		FFrontendQuery Query;
 		Query.AddStep<FNodeClassRegistrationEvents>()
-			.AddStep<FMapRegistrationEventsToNodeRegistryKeys>()
+			.AddStep<FMapClassRegistrationEventsToNodeRegistryKeys>()
 			.AddStep<FReduceRegistrationEventsToCurrentStatus>()
 			.AddStep<FMapNodeRegistrationEventsToClassAndMajorVersion>();
 		return Query;
@@ -319,7 +319,7 @@ namespace Metasound::Frontend
 	{
 		using namespace SearchEngineQuerySteps;
 
-		FMetasoundFrontendVersionNumber HighestVersionNumber{-1, -1};
+		FMetasoundFrontendVersionNumber HighestVersionNumber { -1, -1 };
 		FNodeRegistryKey Key = FNodeRegistryKey::GetInvalid();
 		for (const FFrontendQueryEntry& InEntry : InPartition)
 		{
@@ -328,10 +328,10 @@ namespace Metasound::Frontend
 				const FNodeRegistryTransaction& Transaction = InEntry.Value.Get<FNodeRegistryTransaction>();
 				if (FNodeRegistryTransaction::ETransactionType::NodeRegistration == Transaction.GetTransactionType())
 				{
-					const FNodeClassInfo& NodeClassInfo = Transaction.GetNodeClassInfo();
-					if (NodeClassInfo.Version > HighestVersionNumber)
+					const FNodeClassRegistryKey& TransactionKey = Transaction.GetNodeRegistryKey();
+					if (TransactionKey.Version > HighestVersionNumber)
 					{
-						HighestVersionNumber = NodeClassInfo.Version;
+						HighestVersionNumber = TransactionKey.Version;
 						Key = Transaction.GetNodeRegistryKey();
 					}
 				}
@@ -446,6 +446,29 @@ namespace Metasound::Frontend
 		FindAllDefaultInterfaceVersionsForUClassQuery.Prime();
 	}
 
+	bool FSearchEngineCore::FindRegisteredClass(const FMetasoundFrontendClassName& InName, int32 InMajorVersion, FMetaSoundClassInfo& OutClass)
+	{
+		METASOUND_TRACE_CPUPROFILER_EVENT_SCOPE(metasound::FSearchEngineCore::FindRegisteredClass);
+		using namespace SearchEngineQuerySteps;
+
+		const FFrontendQueryKey QueryKey = CreateKeyFromFullClassNameAndMajorVersion(InName, InMajorVersion);
+		FNodeRegistryKey NodeRegistryKey; 
+		if (FindClassWithHighestMinorVersionQuery.UpdateAndFindResult(QueryKey, NodeRegistryKey))
+		{
+			if (FMetasoundFrontendRegistryContainer* NodeRegistry = FMetasoundFrontendRegistryContainer::Get())
+			{
+				FMetasoundFrontendClass Class;
+				const bool bFoundClass = NodeRegistry->FindFrontendClassFromRegistered(NodeRegistryKey, Class);
+				if (bFoundClass)
+				{
+					OutClass = FMetaSoundClassInfo(Class);
+					return true;
+				}
+			}
+		}
+		return false;
+	}
+
 	bool FSearchEngineCore::FindClassWithHighestMinorVersion(const FMetasoundFrontendClassName& InName, int32 InMajorVersion, FMetasoundFrontendClass& OutClass)
 	{
 		METASOUND_TRACE_CPUPROFILER_EVENT_SCOPE(metasound::FSearchEngineCore::FindClassWithHighestMinorVersion);
diff --git a/Engine/Plugins/Runtime/Metasound/Source/MetasoundFrontend/Private/MetasoundFrontendSearchEngineCore.h b/Engine/Plugins/Runtime/Metasound/Source/MetasoundFrontend/Private/MetasoundFrontendSearchEngineCore.h
--- a/Engine/Plugins/Runtime/Metasound/Source/MetasoundFrontend/Private/MetasoundFrontendSearchEngineCore.h
+++ b/Engine/Plugins/Runtime/Metasound/Source/MetasoundFrontend/Private/MetasoundFrontendSearchEngineCore.h
@@ -119,9 +119,9 @@ namespace Metasound
 			};
 
 			// Convert an array of partitions into an array of metasound frontend classes.
-			TArray<FMetasoundFrontendClass> BuildArrayOfClassesFromPartition(const FFrontendQueryPartition& InPartition);
+			TArray<FMetaSoundClassInfo> BuildArrayOfClassesFromPartition(const FFrontendQueryPartition& InPartition);
 
-			FMetasoundFrontendClass BuildSingleClassFromPartition(const FFrontendQueryPartition& InPartition);
+			FMetaSoundClassInfo BuildSingleClassFromPartition(const FFrontendQueryPartition& InPartition);
 
 			// Convert an array of partitions into an array of metasound frontend interfaces.
 			TArray<FMetasoundFrontendInterface> BuildArrayOfInterfacesFromPartition(const FFrontendQueryPartition& InPartition);
@@ -153,25 +153,25 @@ namespace Metasound
 		//
 		// struct FQueryPolicyExample
 		// {
-		//		using ResultType = FMetasoundFrontendClass;
+		//		using ResultType = FMetaSoundClassInfo;
 		//
 		//		static FFrontendQuery CreateQuery()
 		//		{
 		//			using namespace SearchEngineQuerySteps;
 		//			FFrontendQuery Query;
 		//			Query.AddStep<FNodeClassRegistrationEvents>()
-		//				.AddStep<FMapRegistrationEventsToNodeRegistryKeys>()
+		//				.AddStep<FMapClassRegistrationEventsToNodeRegistryKeys>()
 		//				.AddStep<FReduceRegistrationEventsToCurrentStatus>()
-		//				.AddStep<FTransformRegistrationEventsToClasses>();
+		//				.AddStep<FTransformNodeClassRegistrationEventsToClassInfo>();
 		//			return Query;
 		//		}
 		//
 		//		static ResultType BuildResult(const FFrontendQueryPartition& InPartition)
 		//		{
-		//			FMetasoundFrontendClass MetaSoundClass;
+		//			FMetaSoundClassInfo MetaSoundClass;
 		//			if (ensureMsgf(InPartition.Num() > 0, TEXT("Cannot retrieve class from empty partition")))
 		//			{
-		//				MetaSoundClass = InPartition.CreateConstIterator()->Value.Get<FMetasoundFrontendClass>();
+		//				MetaSoundClass = InPartition.CreateConstIterator()->Value.Get<FMetaSoundClassInfo>();
 		//			}
 		//			return MetaSoundClass;
 		//		}
@@ -327,6 +327,8 @@ namespace Metasound
 			virtual void Prime() override;
 			virtual bool FindClassWithHighestMinorVersion(const FMetasoundFrontendClassName& InName, int32 InMajorVersion, FMetasoundFrontendClass& OutClass) override;
 
+			virtual bool FindRegisteredClass(const FMetasoundFrontendClassName& InName, int32 InMajorVersion, FMetaSoundClassInfo& OutClass) override;
+
 			virtual TArray<FMetasoundFrontendInterface> FindUClassDefaultInterfaces(FName InUClassName) override;
 			virtual TArray<FMetasoundFrontendVersion> FindUClassDefaultInterfaceVersions(const FTopLevelAssetPath& InUClassPath) override;
 			virtual TArray<FMetasoundFrontendVersion> FindAllRegisteredInterfacesWithName(FName InInterfaceName) override;
diff --git a/Engine/Plugins/Runtime/Metasound/Source/MetasoundFrontend/Private/MetasoundFrontendSearchEngineEditorOnly.cpp b/Engine/Plugins/Runtime/Metasound/Source/MetasoundFrontend/Private/MetasoundFrontendSearchEngineEditorOnly.cpp
--- a/Engine/Plugins/Runtime/Metasound/Source/MetasoundFrontend/Private/MetasoundFrontendSearchEngineEditorOnly.cpp
+++ b/Engine/Plugins/Runtime/Metasound/Source/MetasoundFrontend/Private/MetasoundFrontendSearchEngineEditorOnly.cpp
@@ -6,6 +6,7 @@
 #include "Containers/Array.h"
 #include "Interfaces/MetasoundFrontendInterfaceRegistry.h"
 #include "MetasoundFrontendDocument.h"
+#include "MetasoundFrontendNodeClassRegistry.h"
 #include "MetasoundFrontendQuery.h"
 #include "MetasoundFrontendQuerySteps.h"
 #include "MetasoundFrontendRegistryTransaction.h"
@@ -14,174 +15,257 @@
 
 #if WITH_EDITORONLY_DATA
 
-namespace Metasound
+namespace Metasound::Frontend
 {
-	namespace Frontend
+	FFrontendQuery FFindAllAssetClassInfoQueryPolicy::CreateQuery()
 	{
-		FFrontendQuery FFindAllClassesQueryPolicy::CreateQuery()
-		{
-			using namespace SearchEngineQuerySteps;
-			FFrontendQuery Query;
-			Query.AddStep<FNodeClassRegistrationEvents>()
-				.AddStep<FMapRegistrationEventsToNodeRegistryKeys>()
-				.AddStep<FReduceRegistrationEventsToCurrentStatus>()
-				.AddStep<FTransformRegistrationEventsToClasses>()
-				.AddStep<FRemoveDeprecatedClasses>()
-				.AddStep<FMapToNull>();
-			return Query;
-		}
+		using namespace SearchEngineQuerySteps;
+		FFrontendQuery Query;
+		Query.AddStep<FAssetRegistrationEvents>()
+			.AddStep<FMapAssetRegistrationEventsToNodeRegistryKeys>()
+			.AddStep<FReduceAssetRegistrationEventsToCurrentStatus>()
+			.AddStep<FTransformAssetRegistrationEventsToClassInfo>()
+			.AddStep<FRemoveDeprecatedClasses>()
+			.AddStep<FMapToNull>();
+		return Query;
+	}
 
-		TArray<FMetasoundFrontendClass> FFindAllClassesQueryPolicy::BuildResult(const FFrontendQueryPartition& InPartition)
-		{
-			using namespace SearchEngineQuerySteps;
-			return BuildArrayOfClassesFromPartition(InPartition);
-		}
+	TArray<FMetaSoundClassInfo> FFindAllAssetClassInfoQueryPolicy::BuildResult(const FFrontendQueryPartition& InPartition)
+	{
+		using namespace SearchEngineQuerySteps;
+		return BuildArrayOfClassesFromPartition(InPartition);
+	}
 
-		FFrontendQuery FFindAllClassesIncludingAllVersionsQueryPolicy::CreateQuery()
-		{
-			using namespace SearchEngineQuerySteps;
-			FFrontendQuery Query;
-			Query.AddStep<FNodeClassRegistrationEvents>()
-				.AddStep<FMapRegistrationEventsToNodeRegistryKeys>()
-				.AddStep<FReduceRegistrationEventsToCurrentStatus>()
-				.AddStep<FTransformRegistrationEventsToClasses>()
-				.AddStep<FMapToNull>();
-			return Query;
-		}
+	FFrontendQuery FFindAllClassesQueryPolicy::CreateQuery()
+	{
+		using namespace SearchEngineQuerySteps;
+		FFrontendQuery Query;
+		Query.AddStep<FNodeClassRegistrationEvents>()
+			.AddStep<FMapClassRegistrationEventsToNodeRegistryKeys>()
+			.AddStep<FReduceRegistrationEventsToCurrentStatus>()
+			.AddStep<FTransformNodeClassRegistrationEventsToClassInfo>()
+			.AddStep<FRemoveDeprecatedClasses>()
+			.AddStep<FMapToNull>();
+		return Query;
+	}
 
-		TArray<FMetasoundFrontendClass> FFindAllClassesIncludingAllVersionsQueryPolicy::BuildResult(const FFrontendQueryPartition& InPartition)
-		{
-			using namespace SearchEngineQuerySteps;
-			return BuildArrayOfClassesFromPartition(InPartition);
-		}
+	TArray<FMetaSoundClassInfo> FFindAllClassesQueryPolicy::BuildResult(const FFrontendQueryPartition& InPartition)
+	{
+		using namespace SearchEngineQuerySteps;
+		return BuildArrayOfClassesFromPartition(InPartition);
+	}
 
-		FFrontendQuery FFindClassesWithNameSortedQueryPolicy::CreateQuery()
-		{
-			using namespace SearchEngineQuerySteps;
-			FFrontendQuery Query;
-			Query.AddStep<FNodeClassRegistrationEvents>()
-				.AddStep<FMapRegistrationEventsToNodeRegistryKeys>()
-				.AddStep<FReduceRegistrationEventsToCurrentStatus>()
-				.AddStep<FTransformRegistrationEventsToClasses>()
-				.AddStep<FMapClassesToClassName>()
-				.AddStep<FSortClassesByVersion>();
-			return Query;
-		}
+	FFrontendQuery FFindAllRegisteredClassesIncludingAllVersionsQueryPolicy::CreateQuery()
+	{
+		using namespace SearchEngineQuerySteps;
+		FFrontendQuery Query;
+		Query.AddStep<FNodeClassRegistrationEvents>()
+			.AddStep<FMapClassRegistrationEventsToNodeRegistryKeys>()
+			.AddStep<FReduceRegistrationEventsToCurrentStatus>()
+			.AddStep<FTransformNodeClassRegistrationEventsToClassInfo>()
+			.AddStep<FMapToFullClassName>();
+		return Query;
+	}
 
-		TArray<FMetasoundFrontendClass> FFindClassesWithNameSortedQueryPolicy::BuildResult(const FFrontendQueryPartition& InPartition)
-		{
-			using namespace SearchEngineQuerySteps;
-			return BuildArrayOfClassesFromPartition(InPartition);
-		}
+	TArray<FMetaSoundClassInfo> FFindAllRegisteredClassesIncludingAllVersionsQueryPolicy::BuildResult(const FFrontendQueryPartition& InPartition)
+	{
+		using namespace SearchEngineQuerySteps;
+		return BuildArrayOfClassesFromPartition(InPartition);
+	}
 
-		FFrontendQuery FFindClassesWithNameUnsortedQueryPolicy::CreateQuery()
-		{
-			using namespace SearchEngineQuerySteps;
-			FFrontendQuery Query;
-			Query.AddStep<FNodeClassRegistrationEvents>()
-				.AddStep<FMapRegistrationEventsToNodeRegistryKeys>()
-				.AddStep<FReduceRegistrationEventsToCurrentStatus>()
-				.AddStep<FTransformRegistrationEventsToClasses>()
-				.AddStep<FMapClassesToClassName>();
-			return Query;
-		}
+	FFrontendQuery FFindAllRegisteredAssetsQueryPolicy::CreateQuery()
+	{
+		using namespace SearchEngineQuerySteps;
+		FFrontendQuery Query;
+		Query.AddStep<FAssetRegistrationEvents>()
+			.AddStep<FMapClassRegistrationEventsToNodeRegistryKeys>()
+			.AddStep<FReduceRegistrationEventsToCurrentStatus>()
+			.AddStep<FTransformNodeClassRegistrationEventsToClassInfo>()
+			.AddStep<FMapToNull>();
+		return Query;
+	}
 
-		TArray<FMetasoundFrontendClass> FFindClassesWithNameUnsortedQueryPolicy::BuildResult(const FFrontendQueryPartition& InPartition)
-		{
-			using namespace SearchEngineQuerySteps;
-			return BuildArrayOfClassesFromPartition(InPartition);
-		}
+	TArray<FMetaSoundClassInfo> FFindAllRegisteredAssetsQueryPolicy::BuildResult(const FFrontendQueryPartition& InPartition)
+	{
+		using namespace SearchEngineQuerySteps;
+		return BuildArrayOfClassesFromPartition(InPartition);
+	}
+
+	FFrontendQuery FFindClassesWithNameSortedQueryPolicy::CreateQuery()
+	{
+		using namespace SearchEngineQuerySteps;
+		FFrontendQuery Query;
+		Query.AddStep<FNodeClassRegistrationEvents>()
+			.AddStep<FMapClassRegistrationEventsToNodeRegistryKeys>()
+			.AddStep<FReduceRegistrationEventsToCurrentStatus>()
+			.AddStep<FTransformNodeClassRegistrationEventsToClassInfo>()
+			.AddStep<FMapClassesToClassName>()
+			.AddStep<FSortClassesByVersion>();
+		return Query;
+	}
+
+	TArray<FMetaSoundClassInfo> FFindClassesWithNameSortedQueryPolicy::BuildResult(const FFrontendQueryPartition& InPartition)
+	{
+		using namespace SearchEngineQuerySteps;
+		return BuildArrayOfClassesFromPartition(InPartition);
+	}
+
+	FFrontendQuery FFindClassesWithNameUnsortedQueryPolicy::CreateQuery()
+	{
+		using namespace SearchEngineQuerySteps;
+		FFrontendQuery Query;
+		Query.AddStep<FNodeClassRegistrationEvents>()
+			.AddStep<FMapClassRegistrationEventsToNodeRegistryKeys>()
+			.AddStep<FReduceRegistrationEventsToCurrentStatus>()
+			.AddStep<FTransformNodeClassRegistrationEventsToClassInfo>()
+			.AddStep<FMapClassesToClassName>();
+		return Query;
+	}
+
+	TArray<FMetaSoundClassInfo> FFindClassesWithNameUnsortedQueryPolicy::BuildResult(const FFrontendQueryPartition& InPartition)
+	{
+		using namespace SearchEngineQuerySteps;
+		return BuildArrayOfClassesFromPartition(InPartition);
+	}
+
+	FFrontendQuery FFindClassWithHighestVersionQueryPolicy::CreateQuery()
+	{
+		using namespace SearchEngineQuerySteps;
+		FFrontendQuery Query;
+		Query.AddStep<FNodeClassRegistrationEvents>()
+			.AddStep<FMapClassRegistrationEventsToNodeRegistryKeys>()
+			.AddStep<FReduceRegistrationEventsToCurrentStatus>()
+			.AddStep<FTransformNodeClassRegistrationEventsToClassInfo>()
+			.AddStep<FRemoveDeprecatedClasses>()
+			.AddStep<FMapToFullClassName>()
+			.AddStep<FReduceClassesToHighestVersion>();
+		return Query;
+	}
+
+	FMetaSoundClassInfo FFindClassWithHighestVersionQueryPolicy::BuildResult(const FFrontendQueryPartition& InPartition)
+	{
+		using namespace SearchEngineQuerySteps;
+		return BuildSingleClassFromPartition(InPartition);
+	}
 
-		FFrontendQuery FFindClassWithHighestVersionQueryPolicy::CreateQuery()
+	void FSearchEngineEditorOnly::Prime()
+	{
+		Super::Prime();
 		{
-			using namespace SearchEngineQuerySteps;
-			FFrontendQuery Query;
-			Query.AddStep<FNodeClassRegistrationEvents>()
-				.AddStep<FMapRegistrationEventsToNodeRegistryKeys>()
-				.AddStep<FReduceRegistrationEventsToCurrentStatus>()
-				.AddStep<FTransformRegistrationEventsToClasses>()
-				.AddStep<FRemoveDeprecatedClasses>()
-				.AddStep<FMapToFullClassName>()
-				.AddStep<FReduceClassesToHighestVersion>();
-			return Query;
+			METASOUND_LLM_SCOPE;
+			METASOUND_TRACE_CPUPROFILER_EVENT_SCOPE(Metasound::Frontend::FSearchEngineEditorOnly::Prime);
+
+			FindAllClassesQuery.Prime();
+			FindAllClassesIncludingAllVersionsQuery.Prime();
+			FindClassesWithNameUnsortedQuery.Prime();
+			FindClassesWithNameSortedQuery.Prime();
+			FindClassWithHighestVersionQuery.Prime();
+			FindAllInterfacesIncludingAllVersionsQuery.Prime();
+			FindAllInterfacesQuery.Prime();
 		}
+	}
+
+	TArray<FMetasoundFrontendClass> FSearchEngineEditorOnly::FindClassesWithName(const FMetasoundFrontendClassName& InName, bool bInSortByVersion)
+	{
+		METASOUND_TRACE_CPUPROFILER_EVENT_SCOPE(metasound::FSearchEngineEditorOnly::FindClassesWithName);
 
-		FMetasoundFrontendClass FFindClassWithHighestVersionQueryPolicy::BuildResult(const FFrontendQueryPartition& InPartition)
+		TArray<FMetaSoundClassInfo> ClassInfos = FindClassesWithName(InName, bInSortByVersion ? ESortByVersion::No : ESortByVersion::Yes);
+		TArray<FMetasoundFrontendClass> Classes;
+		Algo::Transform(ClassInfos, Classes, [](const FMetaSoundClassInfo& ClassInfo)
 		{
-			using namespace SearchEngineQuerySteps;
-			return BuildSingleClassFromPartition(InPartition);
-		}
+			FMetasoundFrontendClass Class;
+			INodeClassRegistry::Get()->FindFrontendClassFromRegistered(ClassInfo.ToRegistryKey(), Class);
+			return Class;
+		});
+
+		return Classes;
+	}
+
+	TArray<FMetaSoundClassInfo> FSearchEngineEditorOnly::FindAllClasses(ISearchEngine::EResultVersion VersionOptions, bool bIncludeUnloadedAssets)
+	{
+		METASOUND_TRACE_CPUPROFILER_EVENT_SCOPE(Metasound::Frontend::FSearchEngineEditorOnly::FindAllClasses);
 
-		void FSearchEngineEditorOnly::Prime()
+		TArray<FMetaSoundClassInfo> ClassInfos;
+
+		const FFrontendQueryKey NullKey;
+
+		switch (VersionOptions)
 		{
-			Super::Prime();
+			case ISearchEngine::EResultVersion::All:
 			{
-				METASOUND_LLM_SCOPE;
-				METASOUND_TRACE_CPUPROFILER_EVENT_SCOPE(metasound::FSearchEngineEditorOnly::Prime);
-
-				FindAllClassesQuery.Prime();
-				FindAllClassesIncludingAllVersionsQuery.Prime();
-				FindClassesWithNameUnsortedQuery.Prime();
-				FindClassesWithNameSortedQuery.Prime();
-				FindClassWithHighestVersionQuery.Prime();
-				FindAllInterfacesIncludingAllVersionsQuery.Prime();
-				FindAllInterfacesQuery.Prime();
+				ClassInfos = FindAllClassesIncludingAllVersionsQuery.UpdateAndFindResult(NullKey);
 			}
-		}
+			break;
 
-		TArray<FMetasoundFrontendClass> FSearchEngineEditorOnly::FindAllClasses(bool bInIncludeAllVersions)
-		{
-			METASOUND_TRACE_CPUPROFILER_EVENT_SCOPE(metasound::FSearchEngineEditorOnly::FindAllClasses);
-
-			FFrontendQueryKey NullKey;
-			if (bInIncludeAllVersions)
+			case ISearchEngine::EResultVersion::Highest:
 			{
-				return FindAllClassesIncludingAllVersionsQuery.UpdateAndFindResult(NullKey);
+				ClassInfos = FindAllClassesQuery.UpdateAndFindResult(NullKey);
 			}
-			else
+			break;
+
+			default:
+				checkNoEntry();
+		}
+		
+		if (bIncludeUnloadedAssets)
+		{
+			TMap<FNodeClassRegistryKey, FMetaSoundClassInfo> KeysToInfo;
+			Algo::Transform(ClassInfos, KeysToInfo, [](const FMetaSoundClassInfo& Info) { return TPair<FNodeClassRegistryKey, FMetaSoundClassInfo>(Info.ToRegistryKey(), Info); });
+
+			TArray<FMetaSoundClassInfo> AssetInfos = FindAllAssetClassInfoQuery.UpdateAndFindResult(NullKey);			
+			for (FMetaSoundClassInfo& Info : AssetInfos)
 			{
-				return FindAllClassesQuery.UpdateAndFindResult(NullKey);
+				if (!KeysToInfo.Contains(Info.ToRegistryKey()))
+				{
+					ClassInfos.Add(MoveTemp(Info));
+				}
 			}
 		}
 
-		TArray<FMetasoundFrontendClass> FSearchEngineEditorOnly::FindClassesWithName(const FMetasoundFrontendClassName& InName, bool bInSortByVersion)
-		{
-			METASOUND_TRACE_CPUPROFILER_EVENT_SCOPE(metasound::FSearchEngineEditorOnly::FindClassesWithName);
+		return ClassInfos;
+	}
 
-			const FFrontendQueryKey Key(InName.GetFullName());
-			if (bInSortByVersion)
-			{
+	TArray<FMetaSoundClassInfo> FSearchEngineEditorOnly::FindClassesWithName(const FMetasoundFrontendClassName& InName, ISearchEngine::ESortByVersion SortByVersion)
+	{
+		METASOUND_TRACE_CPUPROFILER_EVENT_SCOPE(Metasound::Frontend::FSearchEngineEditorOnly::FindClassesWithName);
+
+		const FFrontendQueryKey Key(InName.GetFullName());
+		switch (SortByVersion)
+		{
+			case ISearchEngine::ESortByVersion::Yes:
 				return FindClassesWithNameSortedQuery.UpdateAndFindResult(Key);
-			}
-			else 
-			{
+			case ISearchEngine::ESortByVersion::No:
 				return FindClassesWithNameUnsortedQuery.UpdateAndFindResult(Key);
-			}
+			default:
+				checkNoEntry();
 		}
 
-		bool FSearchEngineEditorOnly::FindClassWithHighestVersion(const FMetasoundFrontendClassName& InName, FMetasoundFrontendClass& OutClass)
-		{
-			METASOUND_TRACE_CPUPROFILER_EVENT_SCOPE(metasound::FSearchEngineEditorOnly::FindClassWithHighestVersion);
+		return { };
+	}
 
-			const FFrontendQueryKey Key{InName.GetFullName()};
-			return FindClassWithHighestVersionQuery.UpdateAndFindResult(Key, OutClass);
-		}
+	bool FSearchEngineEditorOnly::FindClassWithHighestVersion(const FMetasoundFrontendClassName& InName, FMetaSoundClassInfo& OutClassInfo)
+	{
+		METASOUND_TRACE_CPUPROFILER_EVENT_SCOPE(Metasound::FSearchEngineEditorOnly::FindClassWithHighestVersion);
 
-		TArray<FMetasoundFrontendInterface> FSearchEngineEditorOnly::FindAllInterfaces(bool bInIncludeAllVersions)
-		{
-			METASOUND_TRACE_CPUPROFILER_EVENT_SCOPE(metasound::FSearchEngineEditorOnly::FindAllInterfaces);
+		const FFrontendQueryKey Key{InName.GetFullName()};
+		return FindClassWithHighestVersionQuery.UpdateAndFindResult(Key, OutClassInfo);
 
-			const FFrontendQueryKey NullKey;
-			if (bInIncludeAllVersions)
-			{
-				return FindAllInterfacesIncludingAllVersionsQuery.UpdateAndFindResult(NullKey);
-			}
-			else
-			{
-				return FindAllInterfacesQuery.UpdateAndFindResult(NullKey);
-			}
-		}
+		return false;
 	}
-}
 
-#endif
+	TArray<FMetasoundFrontendInterface> FSearchEngineEditorOnly::FindAllInterfaces(bool bInIncludeAllVersions)
+	{
+		METASOUND_TRACE_CPUPROFILER_EVENT_SCOPE(Metasound::FSearchEngineEditorOnly::FindAllInterfaces);
+
+		const FFrontendQueryKey NullKey;
+		if (bInIncludeAllVersions)
+		{
+			return FindAllInterfacesIncludingAllVersionsQuery.UpdateAndFindResult(NullKey);
+		}
+		else
+		{
+			return FindAllInterfacesQuery.UpdateAndFindResult(NullKey);
+		}
+	}
+} // namespace Metasound::Frontend
+#endif // WITH_EDITORONLY_DATA
diff --git a/Engine/Plugins/Runtime/Metasound/Source/MetasoundFrontend/Private/MetasoundFrontendSearchEngineEditorOnly.h b/Engine/Plugins/Runtime/Metasound/Source/MetasoundFrontend/Private/MetasoundFrontendSearchEngineEditorOnly.h
--- a/Engine/Plugins/Runtime/Metasound/Source/MetasoundFrontend/Private/MetasoundFrontendSearchEngineEditorOnly.h
+++ b/Engine/Plugins/Runtime/Metasound/Source/MetasoundFrontend/Private/MetasoundFrontendSearchEngineEditorOnly.h
@@ -9,89 +9,104 @@
 #include "MetasoundFrontendSearchEngineCore.h"
 
 #if WITH_EDITORONLY_DATA
-
-namespace Metasound
+namespace Metasound::Frontend
 {
-	namespace Frontend
+	// Policy for finding all registered asset classes that are not deprecated and referenceable.
+	struct FFindAllAssetClassInfoQueryPolicy
+	{
+		using ResultType = TArray<FMetaSoundClassInfo>;
+
+		static FFrontendQuery CreateQuery();
+		static ResultType BuildResult(const FFrontendQueryPartition& InPartition);
+	};
+
+	// Policy for finding all registered metasound classes that are not deprecated and referenceable.
+	struct FFindAllClassesQueryPolicy
+	{
+		using ResultType = TArray<FMetaSoundClassInfo>;
+
+		static FFrontendQuery CreateQuery();
+		static ResultType BuildResult(const FFrontendQueryPartition& InPartition);
+	};
+
+	// Policy for finding all registered metasound classes, including deprecated classes. 
+	struct FFindAllRegisteredClassesIncludingAllVersionsQueryPolicy
+	{
+		using ResultType = TArray<FMetaSoundClassInfo>;
+
+		static FFrontendQuery CreateQuery();
+		static ResultType BuildResult(const FFrontendQueryPartition& InPartition);
+	};
+
+	// Policy for finding all registered metasound assets. 
+	struct FFindAllRegisteredAssetsQueryPolicy
+	{
+		using ResultType = TArray<FMetaSoundClassInfo>;
+
+		static FFrontendQuery CreateQuery();
+		static ResultType BuildResult(const FFrontendQueryPartition& InPartition);
+	};
+
+	// Policy for finding all registered metasound classes sorted by version 
+	// and indexed by name.
+	struct FFindClassesWithNameSortedQueryPolicy
+	{
+		using ResultType = TArray<FMetaSoundClassInfo>;
+
+		static FFrontendQuery CreateQuery();
+		static ResultType BuildResult(const FFrontendQueryPartition& InPartition);
+	};
+
+	// Policy for finding all registered metasound classes indexed by name.
+	struct FFindClassesWithNameUnsortedQueryPolicy
+	{
+		using ResultType = TArray<FMetaSoundClassInfo>;
+
+		static FFrontendQuery CreateQuery();
+		static ResultType BuildResult(const FFrontendQueryPartition& InPartition);
+	};
+
+	// Policy for finding highest versioned metasound class by name.
+	struct FFindClassWithHighestVersionQueryPolicy
+	{
+		using ResultType = FMetaSoundClassInfo;
+
+		static FFrontendQuery CreateQuery();
+		static ResultType BuildResult(const FFrontendQueryPartition& InPartition);
+	};
+
+	// To minimize runtime memory costs during gampeplay, some search engine
+	// queries are only exposed when the editor only data is available. This
+	// class supports the editor only functionality of the search engine.
+	class FSearchEngineEditorOnly : public FSearchEngineCore
 	{
-		// Policy for finding all registered metasound classes that are not deprecated and referenceable. 
-		struct FFindAllClassesQueryPolicy
-		{
-			using ResultType = TArray<FMetasoundFrontendClass>;
-
-			static FFrontendQuery CreateQuery();
-			static ResultType BuildResult(const FFrontendQueryPartition& InPartition);
-		};
-
-		// Policy for finding all registered metasound classes, including deprecated classes. 
-		struct FFindAllClassesIncludingAllVersionsQueryPolicy
-		{
-			using ResultType = TArray<FMetasoundFrontendClass>;
-
-			static FFrontendQuery CreateQuery();
-			static ResultType BuildResult(const FFrontendQueryPartition& InPartition);
-		};
-
-		// Policy for finding all registered metasound classes sorted by version 
-		// and indexed by name.
-		struct FFindClassesWithNameSortedQueryPolicy
-		{
-			using ResultType = TArray<FMetasoundFrontendClass>;
-
-			static FFrontendQuery CreateQuery();
-			static ResultType BuildResult(const FFrontendQueryPartition& InPartition);
-		};
-
-		// Policy for finding all registered metasound classes indexed by name.
-		struct FFindClassesWithNameUnsortedQueryPolicy
-		{
-			using ResultType = TArray<FMetasoundFrontendClass>;
-
-			static FFrontendQuery CreateQuery();
-			static ResultType BuildResult(const FFrontendQueryPartition& InPartition);
-		};
-
-		// Policy for finding highest versioned metasound class by name.
-		struct FFindClassWithHighestVersionQueryPolicy
-		{
-			using ResultType = FMetasoundFrontendClass;
-
-			static FFrontendQuery CreateQuery();
-			static ResultType BuildResult(const FFrontendQueryPartition& InPartition);
-		};
-
-		// To minimize runtime memory costs during gampeplay, some search engine
-		// queries are only exposed when the editor only data is available. This
-		// class supports the editor only functionality of the search engine.
-		class FSearchEngineEditorOnly : public FSearchEngineCore
-		{
-			FSearchEngineEditorOnly(const FSearchEngineEditorOnly&) = delete;
-			FSearchEngineEditorOnly(FSearchEngineEditorOnly&&) = delete;
-			FSearchEngineEditorOnly& operator=(const FSearchEngineEditorOnly&) = delete;
-			FSearchEngineEditorOnly& operator=(FSearchEngineEditorOnly&&) = delete;
-
-			using Super = FSearchEngineCore;
-
-		public:
-			FSearchEngineEditorOnly() = default;
-			virtual ~FSearchEngineEditorOnly() = default;
-
-			virtual void Prime() override;
-			virtual TArray<FMetasoundFrontendClass> FindAllClasses(bool bInIncludeAllVersions) override;
-			virtual TArray<FMetasoundFrontendClass> FindClassesWithName(const FMetasoundFrontendClassName& InName, bool bInSortByVersion) override;
-			virtual bool FindClassWithHighestVersion(const FMetasoundFrontendClassName& InName, FMetasoundFrontendClass& OutClass) override;
-			virtual TArray<FMetasoundFrontendInterface> FindAllInterfaces(bool bInIncludeAllVersions) override;
-
-		private:
-			TSearchEngineQuery<FFindAllClassesQueryPolicy> FindAllClassesQuery;
-			TSearchEngineQuery<FFindAllClassesIncludingAllVersionsQueryPolicy> FindAllClassesIncludingAllVersionsQuery;
-			TSearchEngineQuery<FFindClassesWithNameUnsortedQueryPolicy> FindClassesWithNameUnsortedQuery;
-			TSearchEngineQuery<FFindClassesWithNameSortedQueryPolicy> FindClassesWithNameSortedQuery;
-			TSearchEngineQuery<FFindClassWithHighestVersionQueryPolicy> FindClassWithHighestVersionQuery;
-			TSearchEngineQuery<FFindAllInterfacesQueryPolicy> FindAllInterfacesQuery;
-			TSearchEngineQuery<FFindAllInterfacesIncludingAllVersionsQueryPolicy> FindAllInterfacesIncludingAllVersionsQuery;
-		};
-	}
-}
-
-#endif
+		FSearchEngineEditorOnly(const FSearchEngineEditorOnly&) = delete;
+		FSearchEngineEditorOnly(FSearchEngineEditorOnly&&) = delete;
+		FSearchEngineEditorOnly& operator=(const FSearchEngineEditorOnly&) = delete;
+		FSearchEngineEditorOnly& operator=(FSearchEngineEditorOnly&&) = delete;
+
+		using Super = FSearchEngineCore;
+
+	public:
+		FSearchEngineEditorOnly() = default;
+		virtual ~FSearchEngineEditorOnly() = default;
+
+		virtual void Prime() override;
+		virtual TArray<FMetaSoundClassInfo> FindAllClasses(ISearchEngine::EResultVersion VersionOptions, bool bIncludeUnloadedAssets) override;
+		virtual TArray<FMetaSoundClassInfo> FindClassesWithName(const FMetasoundFrontendClassName& InName, ESortByVersion SortByVersion = ESortByVersion::Yes) override;
+		virtual TArray<FMetasoundFrontendClass> FindClassesWithName(const FMetasoundFrontendClassName& InName, bool bInSortByVersion) override;
+		virtual bool FindClassWithHighestVersion(const FMetasoundFrontendClassName& InName, FMetaSoundClassInfo& OutClassInfo) override;
+		virtual TArray<FMetasoundFrontendInterface> FindAllInterfaces(bool bInIncludeAllVersions) override;
+
+	private:
+		TSearchEngineQuery<FFindAllAssetClassInfoQueryPolicy> FindAllAssetClassInfoQuery;
+		TSearchEngineQuery<FFindAllClassesQueryPolicy> FindAllClassesQuery;
+		TSearchEngineQuery<FFindAllRegisteredClassesIncludingAllVersionsQueryPolicy> FindAllClassesIncludingAllVersionsQuery;
+		TSearchEngineQuery<FFindClassesWithNameUnsortedQueryPolicy> FindClassesWithNameUnsortedQuery;
+		TSearchEngineQuery<FFindClassesWithNameSortedQueryPolicy> FindClassesWithNameSortedQuery;
+		TSearchEngineQuery<FFindClassWithHighestVersionQueryPolicy> FindClassWithHighestVersionQuery;
+		TSearchEngineQuery<FFindAllInterfacesQueryPolicy> FindAllInterfacesQuery;
+		TSearchEngineQuery<FFindAllInterfacesIncludingAllVersionsQueryPolicy> FindAllInterfacesIncludingAllVersionsQuery;
+	};
+} // namespace Metasound::Frontend
+#endif // WITH_EDITORONLY_DATA
diff --git a/Engine/Plugins/Runtime/Metasound/Source/MetasoundFrontend/Public/MetasoundAssetManager.h b/Engine/Plugins/Runtime/Metasound/Source/MetasoundFrontend/Public/MetasoundAssetManager.h
--- a/Engine/Plugins/Runtime/Metasound/Source/MetasoundFrontend/Public/MetasoundAssetManager.h
+++ b/Engine/Plugins/Runtime/Metasound/Source/MetasoundFrontend/Public/MetasoundAssetManager.h
@@ -1,8 +1,8 @@
 // Copyright Epic Games, Inc. All Rights Reserved.
 #pragma once
 
-#include "JsonObjectConverter.h"
 #include "MetasoundAssetKey.h"
+#include "MetasoundAssetManagerTransaction.h"
 #include "MetasoundDocumentInterface.h"
 #include "MetasoundFrontendDocument.h"
 #include "MetasoundFrontendQuery.h"
@@ -285,6 +285,12 @@ namespace Metasound::Frontend
 		// If multiple assets are associated with the given key, the last one is returned. 
 		virtual FMetasoundAssetBase* FindAsset(const FMetaSoundAssetKey& InAssetKey) const = 0;
 
+		// Returns all class info associated with the given asset key (may be more than one due to errors in duplication
+		// of assets, improper plugin migration, etc. but should usually just be a single entry) registered
+		// with the manager. Info may not be returned if the asset has yet to be registered with the MetaSoundAssetManager
+		// or may be out-of-date if edited after last registration attempt via AddOrUpdate... calls.
+		virtual TArray<Frontend::FMetaSoundClassInfo> FindRegisteredClassInfo(const FMetaSoundAssetKey& InKey) const = 0;
+
 		// Returns object (if loaded) associated with the given key as a Document Interface (null if key not registered with the AssetManager)
 		virtual TScriptInterface<IMetaSoundDocumentInterface> FindAssetAsDocumentInterface(const FMetaSoundAssetKey& InKey) const = 0;
 
@@ -373,6 +379,19 @@ PRAGMA_ENABLE_DEPRECATION_WARNINGS
 		// Waits until all async load requests related to this asset are complete.
  		virtual void WaitUntilAsyncLoadReferencedAssetsComplete(FMetasoundAssetBase& InAssetBase) = 0;
 	};
+
+	class FMetaSoundAssetManagerBase : public IMetaSoundAssetManager
+	{
+	public:
+		UE_API FMetaSoundAssetManagerBase();
+		virtual ~FMetaSoundAssetManagerBase() = default;
+
+		UE_API IMetaSoundAssetTransactor& GetTransactor();
+
+	private:
+		TUniquePtr<IMetaSoundAssetTransactor> Transactor;
+	};
+
 } // namespace Metasound::Frontend
 
 #undef UE_API
diff --git a/Engine/Plugins/Runtime/Metasound/Source/MetasoundFrontend/Public/MetasoundAssetManagerTransaction.h b/Engine/Plugins/Runtime/Metasound/Source/MetasoundFrontend/Public/MetasoundAssetManagerTransaction.h
new file mode 100644
--- /dev/null
+++ b/Engine/Plugins/Runtime/Metasound/Source/MetasoundFrontend/Public/MetasoundAssetManagerTransaction.h
@@ -0,0 +1,50 @@
+// Copyright Epic Games, Inc. All Rights Reserved.
+#pragma once
+
+#include "Containers/UnrealString.h"
+#include "HAL/Platform.h"
+#include "MetasoundAssetKey.h"
+#include "MetasoundFrontendRegistryKey.h"
+
+#define UE_API METASOUNDFRONTEND_API
+
+
+namespace Metasound::Frontend
+{
+	class FAssetManagerTransaction
+	{
+	public:
+		using FTimeType = uint64;
+
+		/** Describes the type of transaction. */
+		enum class ETransactionType : uint8
+		{
+			AssetRegistration,     //< Something was added to the registry.
+			AssetUnregistration,  //< Something was removed from the registry.
+			Invalid
+		};
+
+		static FString LexToString(ETransactionType InTransactionType);
+
+		FAssetManagerTransaction(ETransactionType InType, const FMetaSoundAssetKey& InKey, FTimeType InTimestamp);
+
+		ETransactionType GetTransactionType() const;
+		FNodeClassRegistryKey GetNodeRegistryKey() const;
+		FTimeType GetTimestamp() const;
+
+	private:
+		ETransactionType Type;
+		FNodeClassRegistryKey Key;
+		FTimeType Timestamp;
+	};
+
+	class IMetaSoundAssetTransactor
+	{
+	public:
+		virtual ~IMetaSoundAssetTransactor() = default;
+
+		virtual void RegisterAsset(FMetaSoundAssetKey Key) = 0;
+		virtual void UnregisterAsset(FMetaSoundAssetKey Key) = 0;
+	};
+} // namespace Metasound::Frontend
+#undef UE_API
\ No newline at end of file
diff --git a/Engine/Plugins/Runtime/Metasound/Source/MetasoundFrontend/Public/MetasoundFrontendQuery.h b/Engine/Plugins/Runtime/Metasound/Source/MetasoundFrontend/Public/MetasoundFrontendQuery.h
--- a/Engine/Plugins/Runtime/Metasound/Source/MetasoundFrontend/Public/MetasoundFrontendQuery.h
+++ b/Engine/Plugins/Runtime/Metasound/Source/MetasoundFrontend/Public/MetasoundFrontendQuery.h
@@ -7,6 +7,7 @@
 #include "Interfaces/MetasoundFrontendInterfaceRegistry.h"
 #include "Math/NumericLimits.h"
 #include "MetasoundAssetKey.h"
+#include "MetasoundAssetManagerTransaction.h"
 #include "MetasoundFrontendDocument.h"
 #include "MetasoundFrontendRegistries.h"
 #include "Misc/TVariant.h"
@@ -124,6 +125,7 @@ struct FMetaSoundClassInterfaceInfo
 
 	FMetaSoundClassInterfaceInfo() = default;
 	UE_API FMetaSoundClassInterfaceInfo(const IMetaSoundDocumentInterface& InDocInterface);
+	UE_API FMetaSoundClassInterfaceInfo(const FMetasoundFrontendClass& InClass);
 	UE_API FMetaSoundClassInterfaceInfo(const FAssetData& InAssetData, bool& bOutIsValid);
 
 	UE_API void ExportToContext(FAssetRegistryTagsContext& OutContext) const;
@@ -161,7 +163,8 @@ namespace Metasound
 		// without a local asset required to be loaded or registered with the node class registry).
 		struct FMetaSoundClassInfo
 		{
-			UE_API FMetaSoundClassInfo();
+			UE_API FMetaSoundClassInfo(EMetasoundFrontendClassType InClassType = EMetasoundFrontendClassType::Invalid);
+			UE_API FMetaSoundClassInfo(const FMetasoundFrontendClass& InClass);
 			UE_API FMetaSoundClassInfo(const FAssetData& InAssetData);
 
 			virtual ~FMetaSoundClassInfo() = default;
@@ -181,6 +184,8 @@ namespace Metasound
 			// ClassName of class
 			FMetasoundFrontendClassName ClassName;
 
+			EMetasoundFrontendClassType ClassType = EMetasoundFrontendClassType::Invalid;
+
 			// Version of class
 			FMetasoundFrontendVersionNumber Version;
 
@@ -190,6 +195,11 @@ namespace Metasound
 			UE_API bool InheritsInterface(FName InterfaceName, const  FMetasoundFrontendVersionNumber& VersionNumber = { 0, 0 }) const;
 #endif // WITH_EDITORONLY_DATA
 
+			FNodeClassRegistryKey ToRegistryKey() const
+			{
+				return FNodeClassRegistryKey(ClassType, ClassName, Version);
+			}
+
 			EMetasoundFrontendClassAccessFlags AccessFlags;
 
 			// If true, class info is valid and accurately reflects
@@ -234,7 +244,15 @@ namespace Metasound
 	};
 
 	/** A FFrontendQueryValue contains data of interest. */
-	using FFrontendQueryValue = TVariant<FMetasoundFrontendVersion, Frontend::FNodeRegistryTransaction, FMetasoundFrontendClass, Frontend::FInterfaceRegistryTransaction, FMetasoundFrontendInterface>;
+	using FFrontendQueryValue = TVariant
+	<
+		FMetasoundFrontendVersion,
+		Frontend::FAssetManagerTransaction,
+		Frontend::FInterfaceRegistryTransaction,
+		Frontend::FMetaSoundClassInfo,
+		Frontend::FNodeClassRegistryTransaction,
+		FMetasoundFrontendInterface
+	>;
 
 	/** FFrontendQueryEntry represents one value in the query. It contains an ID,
 	 * value and score.  */
diff --git a/Engine/Plugins/Runtime/Metasound/Source/MetasoundFrontend/Public/MetasoundFrontendQuerySteps.h b/Engine/Plugins/Runtime/Metasound/Source/MetasoundFrontend/Public/MetasoundFrontendQuerySteps.h
--- a/Engine/Plugins/Runtime/Metasound/Source/MetasoundFrontend/Public/MetasoundFrontendQuerySteps.h
+++ b/Engine/Plugins/Runtime/Metasound/Source/MetasoundFrontend/Public/MetasoundFrontendQuerySteps.h
@@ -14,6 +14,7 @@
 namespace Metasound
 {
 	// Forward declare private implementation
+	class FAssetRegistrationEventsPimpl;
 	class FNodeClassRegistrationEventsPimpl;
 
 	class FMapToNull : public IFrontendQueryMapStep
@@ -37,13 +38,45 @@ namespace Metasound
 		TPimplPtr<FNodeClassRegistrationEventsPimpl> Pimpl;
 	};
 
-	/** Partitions node registration events by their node registration keys. */
-	class FMapRegistrationEventsToNodeRegistryKeys : public IFrontendQueryMapStep
+	/** Streams assets that have been newly registered or unregistered since last call to Stream()
+	 */
+	class FAssetRegistrationEvents : public IFrontendQueryStreamStep
+	{
+	public:
+		UE_API FAssetRegistrationEvents();
+		UE_API virtual void Stream(TArray<FFrontendQueryValue>& OutValues) override;
+
+	private:
+		TPimplPtr<FAssetRegistrationEventsPimpl> Pimpl;
+	};
+
+	/** Partitions asset registration events by their node registration keys. */
+	class FMapAssetRegistrationEventsToNodeRegistryKeys : public IFrontendQueryMapStep
 	{
 	public:
 		UE_API virtual FFrontendQueryKey Map(const FFrontendQueryEntry& InEntry) const override;
 	};
 
+	/** Partitions class registration events by their node registration keys. */
+	class FMapClassRegistrationEventsToNodeRegistryKeys : public IFrontendQueryMapStep
+	{
+	public:
+		UE_API virtual FFrontendQueryKey Map(const FFrontendQueryEntry& InEntry) const override;
+	};
+
+	/** Reduces asset registration events mapped to the same key by inspecting their add/remove state in
+	 * order to determine their final state. If an item has been added more than it has been removed,
+	 * then it is added to the output. Otherwise, it is omitted. */
+	class FReduceAssetRegistrationEventsToCurrentStatus : public IFrontendQueryReduceStep
+	{
+	public:
+		UE_API virtual void Reduce(const FFrontendQueryKey& InKey, FFrontendQueryPartition& InOutEntries) const override;
+	private:
+		using FTimeType = Frontend::FAssetManagerTransaction::FTimeType;
+		static FTimeType GetTransactionTimestamp(const FFrontendQueryEntry& InEntry);
+		static bool IsValidTransactionOfType(Frontend::FAssetManagerTransaction::ETransactionType InType, const FFrontendQueryEntry* InEntry);
+	};
+
 	/** Reduces registration events mapped to the same key by inspecting their add/remove state in
 	 * order to determine their final state. If an item has been added more than it has been removed,
 	 * then it is added to the output. Otherwise, it is omitted. */
@@ -57,8 +90,15 @@ namespace Metasound
 		static bool IsValidTransactionOfType(Frontend::FNodeRegistryTransaction::ETransactionType InType, const FFrontendQueryEntry* InEntry);
 	};
 
-	/** Transforms a registration event into a FMetasoundFrontendClass. */
-	class FTransformRegistrationEventsToClasses : public IFrontendQueryTransformStep
+	/** Transforms an asset registration event into a FMetaSoundClassInfo. */
+	class FTransformAssetRegistrationEventsToClassInfo : public IFrontendQueryTransformStep
+	{
+	public:
+		UE_API virtual void Transform(FFrontendQueryEntry::FValue& InValue) const override;
+	};
+
+	/** Transforms a node class registration event into a FMetaSoundClassInfo. */
+	class FTransformNodeClassRegistrationEventsToClassInfo : public IFrontendQueryTransformStep
 	{
 	public:
 		UE_API virtual void Transform(FFrontendQueryEntry::FValue& InValue) const override;
diff --git a/Engine/Plugins/Runtime/Metasound/Source/MetasoundFrontend/Public/MetasoundFrontendRegistryKey.h b/Engine/Plugins/Runtime/Metasound/Source/MetasoundFrontend/Public/MetasoundFrontendRegistryKey.h
--- a/Engine/Plugins/Runtime/Metasound/Source/MetasoundFrontend/Public/MetasoundFrontendRegistryKey.h
+++ b/Engine/Plugins/Runtime/Metasound/Source/MetasoundFrontend/Public/MetasoundFrontendRegistryKey.h
@@ -61,6 +61,13 @@ namespace Metasound::Frontend
 
 		UE_API FNodeClassInfo& operator=(const FNodeClassInfo& Other);
 		UE_API FNodeClassInfo& operator=(FNodeClassInfo&& Other);
+
+		// Added to aid in deprecation of this class in favor of FNodeClassRegistryKey
+		UE_API static const FNodeClassInfo& GetInvalid()
+		{
+			static const FNodeClassInfo Invalid;
+			return Invalid;
+		}
 	};
 
 	struct FNodeClassRegistryKey
diff --git a/Engine/Plugins/Runtime/Metasound/Source/MetasoundFrontend/Public/MetasoundFrontendSearchEngine.h b/Engine/Plugins/Runtime/Metasound/Source/MetasoundFrontend/Public/MetasoundFrontendSearchEngine.h
--- a/Engine/Plugins/Runtime/Metasound/Source/MetasoundFrontend/Public/MetasoundFrontendSearchEngine.h
+++ b/Engine/Plugins/Runtime/Metasound/Source/MetasoundFrontend/Public/MetasoundFrontendSearchEngine.h
@@ -11,6 +11,9 @@
 
 namespace Metasound::Frontend
 {
+	// Forward Declarations
+	struct FMetaSoundClassInfo;
+
 	/** Interface for frontend search engine. A frontend search engine provides
 		* a simple interface for common frontend queries. It also serves as an
 		* opportunity to cache queries to reduce CPU load. 
@@ -18,6 +21,18 @@ namespace Metasound::Frontend
 	class ISearchEngine
 	{
 		public:
+			enum class EResultVersion : uint8
+			{
+				All,
+				Highest
+			};
+
+			enum class ESortByVersion : uint8
+			{
+				Yes,
+				No
+			};
+
 			/** Return an instance of a search engine. */
 			static UE_API ISearchEngine& Get();
 
@@ -26,7 +41,9 @@ namespace Metasound::Frontend
 			/** Updates internal state to speed up queries. */
 			virtual void Prime() = 0;
 
-			/** Find the class with the given ClassName & Major Version. Returns false if not found, true if found. */
+			/** Find the class info with the given ClassName & Major Version. Returns false if not found, true if found. */
+			virtual bool FindRegisteredClass(const FMetasoundFrontendClassName& InName, int32 InMajorVersion, FMetaSoundClassInfo& OutClass) = 0;
+
 			virtual bool FindClassWithHighestMinorVersion(const FMetasoundFrontendClassName& InName, int32 InMajorVersion, FMetasoundFrontendClass& OutClass) = 0;
 
 			/** Finds all registered interfaces with the given name. */
@@ -42,18 +59,33 @@ namespace Metasound::Frontend
 			virtual TArray<FMetasoundFrontendVersion> FindUClassDefaultInterfaceVersions(const FTopLevelAssetPath& InUClassPath) = 0;
 
 #if WITH_EDITORONLY_DATA
-			/** Find all FMetasoundFrontendClasses.
+			/** Find all class infos including those not registered but available via the asset manager for selection.
+				* Options - If set to 'All', include all versions (i.e. deprecated classes and versions of classes that are not the highest major version).
+				* bIncludeUnloadedAssets - If true, include class info for assets not loaded/registered.
+				*/
+			virtual TArray<FMetaSoundClassInfo> FindAllClasses(EResultVersion Options, bool bIncludeUnloadedAssets) = 0;
+
+			/** (To be deprecated) Find all FMetasoundFrontendClasses.
 				* (Optional) Include all versions (i.e. deprecated classes and versions of classes that are not the highest major version).
 				*/
-			virtual TArray<FMetasoundFrontendClass> FindAllClasses(bool bInIncludeAllVersions) = 0;
+			virtual TArray<FMetasoundFrontendClass> FindAllClasses(bool bInIncludeAllVersions);
+
+			/** Find all class info with the given ClassName.
+				* (Optional) Sort matches based on version.
+				*/
+			virtual TArray<FMetaSoundClassInfo> FindClassesWithName(const FMetasoundFrontendClassName& InName, ESortByVersion SortByVersion = ESortByVersion::Yes) = 0;
 
 			/** Find all classes with the given ClassName.
 				* (Optional) Sort matches based on version.
 				*/
 			virtual TArray<FMetasoundFrontendClass> FindClassesWithName(const FMetasoundFrontendClassName& InName, bool bInSortByVersion) = 0;
 
-			/** Find the highest version of a class with the given ClassName. Returns false if not found, true if found. */
-			virtual bool FindClassWithHighestVersion(const FMetasoundFrontendClassName& InName, FMetasoundFrontendClass& OutClass) = 0;
+
+			/** Find the highest version of class info with the given ClassName. Returns false if not found, true if found. */
+			virtual bool FindClassWithHighestVersion(const FMetasoundFrontendClassName& InName, FMetaSoundClassInfo& OutInfo) = 0;
+
+			/** (To be deprecated) Find the highest version of a class with the given ClassName. Returns false if not found, true if found. */
+			virtual bool FindClassWithHighestVersion(const FMetasoundFrontendClassName& InName, FMetasoundFrontendClass& OutClass);
 
 			/** Returns array with all registered interfaces. Optionally, include interface versions that are not the highest version. */
 			virtual TArray<FMetasoundFrontendInterface> FindAllInterfaces(bool bInIncludeAllVersions = false) = 0;
