BugHawk fixes
#rnx

diff --git a/Engine/Plugins/Runtime/Metasound/Source/MetasoundFrontend/Private/MetasoundFrontendQuerySteps.cpp b/Engine/Plugins/Runtime/Metasound/Source/MetasoundFrontend/Private/MetasoundFrontendQuerySteps.cpp
--- a/Engine/Plugins/Runtime/Metasound/Source/MetasoundFrontend/Private/MetasoundFrontendQuerySteps.cpp
+++ b/Engine/Plugins/Runtime/Metasound/Source/MetasoundFrontend/Private/MetasoundFrontendQuerySteps.cpp
@@ -252,7 +252,7 @@ namespace Metasound
 
 		for (const FFrontendQueryEntry& Entry : InOutEntries)
 		{
-			const FNodeClassRegistryTransaction& Transaction = Entry.Value.Get<FNodeRegistryTransaction>();
+			const FNodeClassRegistryTransaction& Transaction = Entry.Value.Get<FNodeClassRegistryTransaction>();
 			switch (Transaction.GetTransactionType())
 			{
 				case FNodeClassRegistryTransaction::ETransactionType::NodeRegistration:
@@ -347,7 +347,10 @@ namespace Metasound
 			{
 				// It's possible that the node is no longer registered (we're processing removals) 
 				// but that's okay because the returned default FrontendClass will be processed out later
-				ClassInfo = IMetaSoundAssetManager::GetChecked().FindRegisteredClassInfo(FMetaSoundAssetKey(Transaction.GetNodeRegistryKey())).Last();
+				TArray<FMetaSoundClassInfo>ClassInfos = IMetaSoundAssetManager::GetChecked()
+					.FindRegisteredClassInfo(FMetaSoundAssetKey(Transaction.GetNodeRegistryKey()));
+				check(!ClassInfos.IsEmpty());
+				ClassInfo = ClassInfos.Last();
 			}
 		}
 		InValue.Set<FMetaSoundClassInfo>(MoveTemp(ClassInfo));
diff --git a/Engine/Plugins/Runtime/Metasound/Source/MetasoundFrontend/Private/MetasoundFrontendSearchEngineCore.cpp b/Engine/Plugins/Runtime/Metasound/Source/MetasoundFrontend/Private/MetasoundFrontendSearchEngineCore.cpp
--- a/Engine/Plugins/Runtime/Metasound/Source/MetasoundFrontend/Private/MetasoundFrontendSearchEngineCore.cpp
+++ b/Engine/Plugins/Runtime/Metasound/Source/MetasoundFrontend/Private/MetasoundFrontendSearchEngineCore.cpp
@@ -217,9 +217,9 @@ namespace Metasound::Frontend
 		public:
 			virtual FFrontendQueryKey Map(const FFrontendQueryEntry& InEntry) const override
 			{
-				if (InEntry.Value.IsType<FNodeRegistryTransaction>())
+				if (InEntry.Value.IsType<FNodeClassRegistryTransaction>())
 				{
-					const FNodeClassRegistryKey& RegistryKey = InEntry.Value.Get<FNodeRegistryTransaction>().GetNodeRegistryKey();
+					const FNodeClassRegistryKey& RegistryKey = InEntry.Value.Get<FNodeClassRegistryTransaction>().GetNodeRegistryKey();
 					return CreateKeyFromFullClassNameAndMajorVersion(RegistryKey.ClassName, RegistryKey.Version.Major);
 				}
 				else
@@ -323,10 +323,10 @@ namespace Metasound::Frontend
 		FNodeRegistryKey Key = FNodeRegistryKey::GetInvalid();
 		for (const FFrontendQueryEntry& InEntry : InPartition)
 		{
-			if (InEntry.Value.IsType<FNodeRegistryTransaction>())
+			if (InEntry.Value.IsType<FNodeClassRegistryTransaction>())
 			{
-				const FNodeRegistryTransaction& Transaction = InEntry.Value.Get<FNodeRegistryTransaction>();
-				if (FNodeRegistryTransaction::ETransactionType::NodeRegistration == Transaction.GetTransactionType())
+				const FNodeClassRegistryTransaction& Transaction = InEntry.Value.Get<FNodeClassRegistryTransaction>();
+				if (FNodeClassRegistryTransaction::ETransactionType::NodeRegistration == Transaction.GetTransactionType())
 				{
 					const FNodeClassRegistryKey& TransactionKey = Transaction.GetNodeRegistryKey();
 					if (TransactionKey.Version > HighestVersionNumber)
diff --git a/Engine/Plugins/Runtime/Metasound/Source/MetasoundFrontend/Public/MetasoundFrontendNodeClassRegistry.h b/Engine/Plugins/Runtime/Metasound/Source/MetasoundFrontend/Public/MetasoundFrontendNodeClassRegistry.h
--- a/Engine/Plugins/Runtime/Metasound/Source/MetasoundFrontend/Public/MetasoundFrontendNodeClassRegistry.h
+++ b/Engine/Plugins/Runtime/Metasound/Source/MetasoundFrontend/Public/MetasoundFrontendNodeClassRegistry.h
@@ -238,7 +238,8 @@ namespace Metasound::Frontend
 		FNodeClassInfo NodeClassInfo;
 		FTimeType Timestamp;
 	};
-	using FNodeRegistryTransaction = FNodeClassRegistryTransaction;
+
+	using FNodeRegistryTransaction UE_DEPRECATED(5.8, "FNodeRegistryTransaction is now FNodeClassRegistryTransaction") = FNodeClassRegistryTransaction;
 
 	namespace NodeClassRegistryKey
 	{
diff --git a/Engine/Plugins/Runtime/Metasound/Source/MetasoundFrontend/Public/MetasoundFrontendQuerySteps.h b/Engine/Plugins/Runtime/Metasound/Source/MetasoundFrontend/Public/MetasoundFrontendQuerySteps.h
--- a/Engine/Plugins/Runtime/Metasound/Source/MetasoundFrontend/Public/MetasoundFrontendQuerySteps.h
+++ b/Engine/Plugins/Runtime/Metasound/Source/MetasoundFrontend/Public/MetasoundFrontendQuerySteps.h
@@ -85,9 +85,9 @@ namespace Metasound
 	public:
 		UE_API virtual void Reduce(const FFrontendQueryKey& InKey, FFrontendQueryPartition& InOutEntries) const override;
 	private:
-		using FTimeType = Frontend::FNodeRegistryTransaction::FTimeType;
+		using FTimeType = Frontend::FNodeClassRegistryTransaction::FTimeType;
 		static FTimeType GetTransactionTimestamp(const FFrontendQueryEntry& InEntry);
-		static bool IsValidTransactionOfType(Frontend::FNodeRegistryTransaction::ETransactionType InType, const FFrontendQueryEntry* InEntry);
+		static bool IsValidTransactionOfType(Frontend::FNodeClassRegistryTransaction::ETransactionType InType, const FFrontendQueryEntry* InEntry);
 	};
 
 	/** Transforms an asset registration event into a FMetaSoundClassInfo. */
