Fix for tests crashing due to unregistration being only reduced transaction.
#rnx

diff --git a/Engine/Plugins/Runtime/Metasound/Source/MetasoundFrontend/Private/MetasoundFrontendQuerySteps.cpp b/Engine/Plugins/Runtime/Metasound/Source/MetasoundFrontend/Private/MetasoundFrontendQuerySteps.cpp
--- a/Engine/Plugins/Runtime/Metasound/Source/MetasoundFrontend/Private/MetasoundFrontendQuerySteps.cpp
+++ b/Engine/Plugins/Runtime/Metasound/Source/MetasoundFrontend/Private/MetasoundFrontendQuerySteps.cpp
@@ -349,8 +349,10 @@ namespace Metasound
 				// but that's okay because the returned default FrontendClass will be processed out later
 				TArray<FMetaSoundClassInfo>ClassInfos = IMetaSoundAssetManager::GetChecked()
 					.FindRegisteredClassInfo(FMetaSoundAssetKey(Transaction.GetNodeRegistryKey()));
-				check(!ClassInfos.IsEmpty());
-				ClassInfo = ClassInfos.Last();
+				if(!ClassInfos.IsEmpty())
+				{
+					ClassInfo = ClassInfos.Last();
+				}
 			}
 		}
 		InValue.Set<FMetaSoundClassInfo>(MoveTemp(ClassInfo));
