Fix tests & associated borked query from SearchEngine optimization updates

#rnx

diff --git a/Engine/Plugins/Runtime/Metasound/Source/MetasoundEngineTest/Private/EngineTestMetaSoundAutomatedNodeTest.cpp b/Engine/Plugins/Runtime/Metasound/Source/MetasoundEngineTest/Private/EngineTestMetaSoundAutomatedNodeTest.cpp
--- a/Engine/Plugins/Runtime/Metasound/Source/MetasoundEngineTest/Private/EngineTestMetaSoundAutomatedNodeTest.cpp
+++ b/Engine/Plugins/Runtime/Metasound/Source/MetasoundEngineTest/Private/EngineTestMetaSoundAutomatedNodeTest.cpp
@@ -30,8 +30,8 @@
 
 #if WITH_EDITORONLY_DATA
 
-namespace Metasound::EngineTest{
-
+namespace Metasound::EngineTest
+{
 	// Return audio mixer device if one is available
 	Audio::FMixerDevice* GetAudioMixerDevice()
 	{
@@ -751,67 +751,67 @@ PRAGMA_ENABLE_DEPRECATION_WARNINGS
 
 	void GetAllRegisteredNodes(TArray<FString>& OutBeautifiedNames, TArray <FString>& OutTestCommands, bool bInIncludeDeprecatedNodes)
 	{
-		using namespace Metasound;
+		using namespace Frontend;
 
 		// Get all the classes that have been registered
-		Frontend::ISearchEngine& NodeSearchEngine = Frontend::ISearchEngine::Get();
-		TArray<FMetasoundFrontendClass> AllClasses = NodeSearchEngine.FindAllClasses(true /* IncludeAllVersions */);
+		constexpr bool bIncludeUnloadedAssets = false;
+		ISearchEngine& NodeSearchEngine = ISearchEngine::Get();
+		TArray<FMetaSoundClassInfo> AllClasses = NodeSearchEngine.FindAllClasses(ISearchEngine::EResultVersion::All, bIncludeUnloadedAssets);
 
-		for (const FMetasoundFrontendClass& NodeClass : AllClasses)
+		for (const FMetaSoundClassInfo& NodeClass : AllClasses)
 		{
 			// Exclude template classes because they cannot be created directly from the node registry
-			if (NodeClass.Metadata.GetType() == EMetasoundFrontendClassType::Template)
+			if (NodeClass.ClassType == EMetasoundFrontendClassType::Template)
 			{
 				continue;
 			}
-			else if (!bInIncludeDeprecatedNodes && EnumHasAnyFlags(NodeClass.Metadata.GetAccessFlags(), EMetasoundFrontendClassAccessFlags::Deprecated))
+			else if (!bInIncludeDeprecatedNodes && EnumHasAnyFlags(NodeClass.AccessFlags, EMetasoundFrontendClassAccessFlags::Deprecated))
 			{
 				// Optionally skip deprecated classes. 
 				continue;
 			}
 
-			OutBeautifiedNames.Add(FString::Printf(TEXT("%s %s"), *NodeClass.Metadata.GetClassName().ToString(), *NodeClass.Metadata.GetVersion().ToString()));
+			OutBeautifiedNames.Add(FString::Printf(TEXT("%s %s"), *NodeClass.ClassName.ToString(), *NodeClass.Version.ToString()));
 
 			// Test commands are node registry keys
-			Frontend::FNodeRegistryKey NodeRegistryKey(NodeClass.Metadata);
-			OutTestCommands.Add(NodeRegistryKey.ToString());
+			OutTestCommands.Add(NodeClass.ToRegistryKey().ToString());
 		}
 	}
 
 	void GetAllRegisteredNativeNodes(TArray<FString>& OutBeautifiedNames, TArray <FString>& OutTestCommands, bool bInIncludeDeprecatedNodes)
 	{
-		using namespace Metasound;
+		using namespace Frontend;
 
 		// Get all the classes that have been registered
-		Frontend::ISearchEngine& NodeSearchEngine = Frontend::ISearchEngine::Get();
-		TArray<FMetasoundFrontendClass> AllClasses = NodeSearchEngine.FindAllClasses(true /* IncludeAllVersions */);
+		constexpr bool bIncludeUnloadedAssets = false;
+		ISearchEngine& NodeSearchEngine = Frontend::ISearchEngine::Get();
+		TArray<FMetaSoundClassInfo> AllClasses = NodeSearchEngine.FindAllClasses(Frontend::ISearchEngine::EResultVersion::All, bIncludeUnloadedAssets);
 
 		FMetasoundFrontendRegistryContainer* NodeRegistry = FMetasoundFrontendRegistryContainer::Get();
 		check(nullptr != NodeRegistry);
 
-		Frontend::IMetaSoundAssetManager* AssetManager = Frontend::IMetaSoundAssetManager::Get();
-		for (const FMetasoundFrontendClass& NodeClass : AllClasses)
+		IMetaSoundAssetManager* AssetManager = IMetaSoundAssetManager::Get();
+		for (const FMetaSoundClassInfo& NodeClass : AllClasses)
 		{
 			// Exclude template classes because they cannot be created directly from the node registry
-			if (NodeClass.Metadata.GetType() == EMetasoundFrontendClassType::Template)
+			if (NodeClass.ClassType == EMetasoundFrontendClassType::Template)
 			{
 				continue;
 			}
-			else if (!bInIncludeDeprecatedNodes && EnumHasAnyFlags(NodeClass.Metadata.GetAccessFlags(), EMetasoundFrontendClassAccessFlags::Deprecated))
+			else if (!bInIncludeDeprecatedNodes && EnumHasAnyFlags(NodeClass.AccessFlags, EMetasoundFrontendClassAccessFlags::Deprecated))
 			{
 				// Optionally skip deprecated classes. 
 				continue;
 			}
 
-			Frontend::FNodeRegistryKey NodeRegistryKey(NodeClass.Metadata);
-
 			// Exclude asset-defined node classes
+			const FNodeRegistryKey NodeRegistryKey = NodeClass.ToRegistryKey();
 			if (AssetManager && AssetManager->ContainsKey(NodeRegistryKey))
 			{
 				continue;
 			}
 
-			OutBeautifiedNames.Add(FString::Printf(TEXT("%s %s"), *NodeClass.Metadata.GetClassName().ToString(), *NodeClass.Metadata.GetVersion().ToString()));
+			OutBeautifiedNames.Add(FString::Printf(TEXT("%s %s"), *NodeClass.ClassName.ToString(), *NodeClass.Version.ToString()));
 
 			// Test commands are node registry keys
 			OutTestCommands.Add(NodeRegistryKey.ToString());
diff --git a/Engine/Plugins/Runtime/Metasound/Source/MetasoundFrontend/Private/MetasoundFrontendSearchEngineEditorOnly.cpp b/Engine/Plugins/Runtime/Metasound/Source/MetasoundFrontend/Private/MetasoundFrontendSearchEngineEditorOnly.cpp
--- a/Engine/Plugins/Runtime/Metasound/Source/MetasoundFrontend/Private/MetasoundFrontendSearchEngineEditorOnly.cpp
+++ b/Engine/Plugins/Runtime/Metasound/Source/MetasoundFrontend/Private/MetasoundFrontendSearchEngineEditorOnly.cpp
@@ -63,7 +63,7 @@ namespace Metasound::Frontend
 			.AddStep<FMapClassRegistrationEventsToNodeRegistryKeys>()
 			.AddStep<FReduceRegistrationEventsToCurrentStatus>()
 			.AddStep<FTransformNodeClassRegistrationEventsToClassInfo>()
-			.AddStep<FMapToFullClassName>();
+			.AddStep<FMapToNull>();
 		return Query;
 	}
 
