Remove all instances of FNodeClassInfo from the MetaSound NodeClassRegistry as its redundant data (and memory bloat)

#rnx

diff --git a/Engine/Plugins/Runtime/Metasound/Source/MetasoundEditor/Private/MetasoundEditorDocumentClipboardUtils.cpp b/Engine/Plugins/Runtime/Metasound/Source/MetasoundEditor/Private/MetasoundEditorDocumentClipboardUtils.cpp
--- a/Engine/Plugins/Runtime/Metasound/Source/MetasoundEditor/Private/MetasoundEditorDocumentClipboardUtils.cpp
+++ b/Engine/Plugins/Runtime/Metasound/Source/MetasoundEditor/Private/MetasoundEditorDocumentClipboardUtils.cpp
@@ -371,10 +371,10 @@ namespace Metasound::Editor
 				}
 				else
 				{
-					FMetasoundFrontendClass RegisteredClass;
-					if (ISearchEngine::Get().FindClassWithHighestMinorVersion(Breadcrumb.ClassName, Breadcrumb.Version.Major, RegisteredClass))
+					FMetaSoundClassInfo ClassInfo;
+					if (ISearchEngine::Get().FindRegisteredClass(Breadcrumb.ClassName, Breadcrumb.Version.Major, ClassInfo))
 					{
-						if (EnumHasAnyFlags(RegisteredClass.Metadata.GetAccessFlags(), EMetasoundFrontendClassAccessFlags::Referenceable))
+						if (EnumHasAnyFlags(ClassInfo.AccessFlags, EMetasoundFrontendClassAccessFlags::Referenceable))
 						{
 							const FMetasoundFrontendNode* NewNode = Builder.AddNodeByClassName(Breadcrumb.ClassName, Breadcrumb.Version.Major);
 							if (NewNode)
@@ -424,10 +424,10 @@ namespace Metasound::Editor
 				}
 				else
 				{
-					FMetasoundFrontendClass ExternalClass;
-					if (ISearchEngine::Get().FindClassWithHighestMinorVersion(Breadcrumb.ClassName, Breadcrumb.Version.Major, ExternalClass))
+					FMetaSoundClassInfo ClassInfo;
+					if (ISearchEngine::Get().FindRegisteredClass(Breadcrumb.ClassName, Breadcrumb.Version.Major, ClassInfo))
 					{
-						if (EnumHasAnyFlags(ExternalClass.Metadata.GetAccessFlags(), EMetasoundFrontendClassAccessFlags::Referenceable))
+						if (EnumHasAnyFlags(ClassInfo.AccessFlags, EMetasoundFrontendClassAccessFlags::Referenceable))
 						{
 							const FMetasoundFrontendNode* NewNode = Builder.AddNodeByClassName(Breadcrumb.ClassName, Breadcrumb.Version.Major);
 							if (NewNode)
diff --git a/Engine/Plugins/Runtime/Metasound/Source/MetasoundEngine/Private/MetasoundBuilderBase.cpp b/Engine/Plugins/Runtime/Metasound/Source/MetasoundEngine/Private/MetasoundBuilderBase.cpp
--- a/Engine/Plugins/Runtime/Metasound/Source/MetasoundEngine/Private/MetasoundBuilderBase.cpp
+++ b/Engine/Plugins/Runtime/Metasound/Source/MetasoundEngine/Private/MetasoundBuilderBase.cpp
@@ -300,15 +300,15 @@ FMetaSoundNodeHandle UMetaSoundBuilderBase::AddNodeByClassName(const FMetasoundF
 	using namespace Metasound::Frontend;
 
 	FMetaSoundNodeHandle NewHandle;
-	FMetasoundFrontendClass RegisteredClass;
-	if (!ISearchEngine::Get().FindClassWithHighestMinorVersion(InClassName, MajorVersion, RegisteredClass))
+	FMetaSoundClassInfo ClassInfo;
+	if (!ISearchEngine::Get().FindRegisteredClass(InClassName, MajorVersion, ClassInfo))
 	{
 		UE_LOG(LogMetaSound, Error, TEXT("Failed to add new node by class name '%s' and major version '%d': Class not found"), *InClassName.ToString(), MajorVersion);
 		OutResult = EMetaSoundBuilderResult::Failed;
 		return NewHandle;
 	}
 
-	if (!EnumHasAnyFlags(RegisteredClass.Metadata.GetAccessFlags(), EMetasoundFrontendClassAccessFlags::Referenceable))
+	if (!EnumHasAnyFlags(ClassInfo.AccessFlags, EMetasoundFrontendClassAccessFlags::Referenceable))
 	{
 		UE_LOG(LogMetaSound, Warning, TEXT("Failed to add new node by name of class '%s': Class access flag '%s' not set."),
 			*InClassName.ToString(),
diff --git a/Engine/Plugins/Runtime/Metasound/Source/MetasoundEngineTest/Private/EngineTestMetasoundSource.cpp b/Engine/Plugins/Runtime/Metasound/Source/MetasoundEngineTest/Private/EngineTestMetasoundSource.cpp
--- a/Engine/Plugins/Runtime/Metasound/Source/MetasoundEngineTest/Private/EngineTestMetasoundSource.cpp
+++ b/Engine/Plugins/Runtime/Metasound/Source/MetasoundEngineTest/Private/EngineTestMetasoundSource.cpp
@@ -73,11 +73,14 @@ namespace EngineTestMetaSoundSourcePrivate
 	{
 		Metasound::Frontend::FNodeHandle Node = Metasound::Frontend::INodeController::GetInvalidHandle();
 		FMetasoundFrontendClass NodeClass;
+
+PRAGMA_DISABLE_DEPRECATION_WARNINGS
 		if (ensure(Metasound::Frontend::ISearchEngine::Get().FindClassWithHighestMinorVersion(InClassName, InMajorVersion, NodeClass)))
 		{
 			Node = InGraph.AddNode(NodeClass.Metadata);
 			check(Node->IsValid());
 		}
+PRAGMA_ENABLE_DEPRECATION_WARNINGS
 		return Node;
 	}
 
diff --git a/Engine/Plugins/Runtime/Metasound/Source/MetasoundEngineTest/Private/NodeTestGraphBuilder.cpp b/Engine/Plugins/Runtime/Metasound/Source/MetasoundEngineTest/Private/NodeTestGraphBuilder.cpp
--- a/Engine/Plugins/Runtime/Metasound/Source/MetasoundEngineTest/Private/NodeTestGraphBuilder.cpp
+++ b/Engine/Plugins/Runtime/Metasound/Source/MetasoundEngineTest/Private/NodeTestGraphBuilder.cpp
@@ -28,10 +28,13 @@ namespace Metasound::Test
 
 		FNodeHandle Node = INodeController::GetInvalidHandle();
 		FMetasoundFrontendClass NodeClass;
+
+PRAGMA_DISABLE_DEPRECATION_WARNINGS
 		if (ISearchEngine::Get().FindClassWithHighestMinorVersion(ClassName, MajorVersion, NodeClass))
 		{
 			Node = RootGraph->AddNode(NodeClass.Metadata);
 		}
+PRAGMA_ENABLE_DEPRECATION_WARNINGS
 
 		return Node;
 	}
diff --git a/Engine/Plugins/Runtime/Metasound/Source/MetasoundFrontend/Private/MetasoundAssetBase.cpp b/Engine/Plugins/Runtime/Metasound/Source/MetasoundFrontend/Private/MetasoundAssetBase.cpp
--- a/Engine/Plugins/Runtime/Metasound/Source/MetasoundFrontend/Private/MetasoundAssetBase.cpp
+++ b/Engine/Plugins/Runtime/Metasound/Source/MetasoundFrontend/Private/MetasoundAssetBase.cpp
@@ -83,7 +83,6 @@ namespace Metasound
 				FDocumentNodeRegistryEntryForSerialization(const FMetasoundFrontendDocument& InDocument, const FTopLevelAssetPath& InAssetPath)
 					: Interfaces(InDocument.Interfaces)
 					, FrontendClass(InDocument.RootGraph)
-					, ClassInfo(InDocument.RootGraph)
 					, AssetPath(InAssetPath)
 				{
 					// Copy FrontendClass to preserve original document.
@@ -94,11 +93,6 @@ namespace Metasound
 
 				virtual ~FDocumentNodeRegistryEntryForSerialization() = default;
 
-				virtual const FNodeClassInfo& GetClassInfo() const override
-				{
-					return ClassInfo;
-				}
-
 				virtual TUniquePtr<INode> CreateNode(const FNodeInitData&) const override { return nullptr; }
 				virtual TUniquePtr<INode> CreateNode(FNodeData) const override { return nullptr; }
 
@@ -133,7 +127,6 @@ namespace Metasound
 			private:
 				TSet<FMetasoundFrontendVersion> Interfaces;
 				FMetasoundFrontendClass FrontendClass;
-				FNodeClassInfo ClassInfo;
 				FTopLevelAssetPath AssetPath;
 			};
 
diff --git a/Engine/Plugins/Runtime/Metasound/Source/MetasoundFrontend/Private/MetasoundFrontendDataTypeRegistry.cpp b/Engine/Plugins/Runtime/Metasound/Source/MetasoundFrontend/Private/MetasoundFrontendDataTypeRegistry.cpp
--- a/Engine/Plugins/Runtime/Metasound/Source/MetasoundFrontend/Private/MetasoundFrontendDataTypeRegistry.cpp
+++ b/Engine/Plugins/Runtime/Metasound/Source/MetasoundFrontend/Private/MetasoundFrontendDataTypeRegistry.cpp
@@ -119,11 +119,6 @@ namespace Metasound
 
 				virtual ~FDataTypeNodeRegistryEntry() = default;
 
-				virtual const FNodeClassInfo& GetClassInfo() const override
-				{
-					return ClassInfo;
-				}
-
 				virtual TUniquePtr<INode> CreateNode(const FNodeInitData&) const override
 				{
 					return nullptr;
@@ -150,16 +145,6 @@ namespace Metasound
 					// No node configuration supported for this node type, so only compatible if setting to invalid (null) configuration
 					return !InNodeConfiguration.IsValid();
 				}
-
-			protected:
-				void UpdateNodeClassInfo(const FMetasoundFrontendClass& InClass)
-				{
-					ClassInfo = FNodeClassInfo(InClass.Metadata);
-				}
-
-			private:
-				
-				FNodeClassInfo ClassInfo;
 			};
 
 			// Node registry entry for input nodes created from a data type registry entry.
@@ -171,7 +156,6 @@ namespace Metasound
 				FInputNodeRegistryEntry(const TSharedRef<IDataTypeRegistryEntry, ESPMode::ThreadSafe>& InDataTypeEntry)
 				: DataTypeEntry(InDataTypeEntry)
 				{
-					UpdateNodeClassInfo(DataTypeEntry->GetFrontendInputClass());
 				}
 
 				virtual ~FInputNodeRegistryEntry() = default;
@@ -226,7 +210,6 @@ PRAGMA_ENABLE_DEPRECATION_WARNINGS
 				FConstructorInputNodeRegistryEntry(const TSharedRef<IDataTypeRegistryEntry, ESPMode::ThreadSafe>& InDataTypeEntry)
 				: DataTypeEntry(InDataTypeEntry)
 				{
-					UpdateNodeClassInfo(DataTypeEntry->GetFrontendConstructorInputClass());
 				}
 
 				virtual ~FConstructorInputNodeRegistryEntry() = default;
@@ -281,7 +264,6 @@ PRAGMA_ENABLE_DEPRECATION_WARNINGS
 				FOutputNodeRegistryEntry(const TSharedRef<IDataTypeRegistryEntry, ESPMode::ThreadSafe>& InDataTypeEntry)
 				: DataTypeEntry(InDataTypeEntry)
 				{
-					UpdateNodeClassInfo(DataTypeEntry->GetFrontendOutputClass());
 				}
 
 				virtual ~FOutputNodeRegistryEntry() = default;
@@ -322,7 +304,6 @@ PRAGMA_ENABLE_DEPRECATION_WARNINGS
 				FConstructorOutputNodeRegistryEntry(const TSharedRef<IDataTypeRegistryEntry, ESPMode::ThreadSafe>& InDataTypeEntry)
 				: DataTypeEntry(InDataTypeEntry)
 				{
-					UpdateNodeClassInfo(DataTypeEntry->GetFrontendConstructorOutputClass());
 				}
 
 				virtual ~FConstructorOutputNodeRegistryEntry() = default;
@@ -364,7 +345,6 @@ PRAGMA_ENABLE_DEPRECATION_WARNINGS
 				FLiteralNodeRegistryEntry(const TSharedRef<IDataTypeRegistryEntry, ESPMode::ThreadSafe>& InDataTypeEntry)
 				: DataTypeEntry(InDataTypeEntry)
 				{
-					UpdateNodeClassInfo(DataTypeEntry->GetFrontendLiteralClass());
 				}
 
 				virtual ~FLiteralNodeRegistryEntry() = default;
@@ -407,7 +387,6 @@ PRAGMA_ENABLE_DEPRECATION_WARNINGS
 				FVariableNodeRegistryEntry(const TSharedRef<IDataTypeRegistryEntry, ESPMode::ThreadSafe>& InDataTypeEntry)
 				: DataTypeEntry(InDataTypeEntry)
 				{
-					UpdateNodeClassInfo(DataTypeEntry->GetFrontendVariableClass());
 				}
 
 				virtual ~FVariableNodeRegistryEntry() = default;
@@ -458,7 +437,6 @@ PRAGMA_ENABLE_DEPRECATION_WARNINGS
 				FVariableMutatorNodeRegistryEntry(const TSharedRef<IDataTypeRegistryEntry, ESPMode::ThreadSafe>& InDataTypeEntry)
 				: DataTypeEntry(InDataTypeEntry)
 				{
-					UpdateNodeClassInfo(DataTypeEntry->GetFrontendVariableMutatorClass());
 				}
 
 				virtual ~FVariableMutatorNodeRegistryEntry() = default;
@@ -504,7 +482,6 @@ PRAGMA_ENABLE_DEPRECATION_WARNINGS
 				FVariableAccessorNodeRegistryEntry(const TSharedRef<IDataTypeRegistryEntry, ESPMode::ThreadSafe>& InDataTypeEntry)
 				: DataTypeEntry(InDataTypeEntry)
 				{
-					UpdateNodeClassInfo(DataTypeEntry->GetFrontendVariableAccessorClass());
 				}
 
 				virtual ~FVariableAccessorNodeRegistryEntry() = default;
@@ -550,7 +527,6 @@ PRAGMA_ENABLE_DEPRECATION_WARNINGS
 				FVariableDeferredAccessorNodeRegistryEntry(const TSharedRef<IDataTypeRegistryEntry, ESPMode::ThreadSafe>& InDataTypeEntry)
 				: DataTypeEntry(InDataTypeEntry)
 				{
-					UpdateNodeClassInfo(DataTypeEntry->GetFrontendVariableDeferredAccessorClass());
 				}
 
 				virtual ~FVariableDeferredAccessorNodeRegistryEntry() = default;
diff --git a/Engine/Plugins/Runtime/Metasound/Source/MetasoundFrontend/Private/MetasoundFrontendDocumentAccessPtr.cpp b/Engine/Plugins/Runtime/Metasound/Source/MetasoundFrontend/Private/MetasoundFrontendDocumentAccessPtr.cpp
--- a/Engine/Plugins/Runtime/Metasound/Source/MetasoundFrontend/Private/MetasoundFrontendDocumentAccessPtr.cpp
+++ b/Engine/Plugins/Runtime/Metasound/Source/MetasoundFrontend/Private/MetasoundFrontendDocumentAccessPtr.cpp
@@ -412,23 +412,6 @@ namespace Metasound
 				FMetasoundFrontendClassMetadata Metadata;
 			};
 
-			/** A callable object for testing whether a class contains a matching node info. */
-			class FIsClassWithNodeInfo
-			{
-			public:
-				FIsClassWithNodeInfo(const FNodeClassInfo& InNodeInfo)
-				: NodeInfo(InNodeInfo)
-				{
-				}
-
-				bool operator()(const FMetasoundFrontendClass& InClass) const
-				{
-					return NodeRegistryKey::IsEqual(NodeInfo, InClass.Metadata);
-				};
-			private:
-				FNodeClassInfo NodeInfo;
-			};
-
 			/** A callable object for testing whether a class contains a matching registry key. */
 			class FIsClassWithNodeRegistryKey
 			{
@@ -898,12 +881,7 @@ namespace Metasound
 
 		FClassAccessPtr FDocumentAccessPtr::GetClassWithInfo(const FNodeClassInfo& InInfo)
 		{
-			using namespace MetasoundFrontendDocumentAccessPtrPrivate;
-
-			using FFinder = TFindClassInDocument<FMetasoundFrontendClass, FIsClassWithNodeInfo>;
-			FFinder Finder(InInfo);
-
-			return GetMemberAccessPtr<FClassAccessPtr>(Finder);
+			return { };
 		}
 
 		FClassAccessPtr FDocumentAccessPtr::GetClassWithRegistryKey(const FNodeRegistryKey& InKey)
@@ -964,12 +942,7 @@ namespace Metasound
 
 		FConstClassAccessPtr FDocumentAccessPtr::GetClassWithInfo(const FNodeClassInfo& InInfo) const
 		{
-			using namespace MetasoundFrontendDocumentAccessPtrPrivate;
-
-			using FFinder = TFindClassInDocument<const FMetasoundFrontendClass, FIsClassWithNodeInfo>;
-			FFinder Finder(InInfo);
-
-			return GetMemberAccessPtr<FConstClassAccessPtr>(Finder);
+			return { };
 		}
 
 		FConstClassAccessPtr FDocumentAccessPtr::GetClassWithRegistryKey(const FNodeRegistryKey& InKey) const
@@ -1030,15 +1003,9 @@ namespace Metasound
 
 		FConstClassAccessPtr FConstDocumentAccessPtr::GetClassWithInfo(const FNodeClassInfo& InInfo) const
 		{
-			using namespace MetasoundFrontendDocumentAccessPtrPrivate;
-
-			using FFinder = TFindClassInDocument<const FMetasoundFrontendClass, FIsClassWithNodeInfo>;
-			FFinder Finder(InInfo);
-
-			return GetMemberAccessPtr<FConstClassAccessPtr>(Finder);
+			return { };
 		}
 
-
 		FConstClassAccessPtr FConstDocumentAccessPtr::GetClassWithRegistryKey(const FNodeRegistryKey& InKey) const
 		{
 			using namespace MetasoundFrontendDocumentAccessPtrPrivate;
diff --git a/Engine/Plugins/Runtime/Metasound/Source/MetasoundFrontend/Private/MetasoundFrontendDocumentBuilder.cpp b/Engine/Plugins/Runtime/Metasound/Source/MetasoundFrontend/Private/MetasoundFrontendDocumentBuilder.cpp
--- a/Engine/Plugins/Runtime/Metasound/Source/MetasoundFrontend/Private/MetasoundFrontendDocumentBuilder.cpp
+++ b/Engine/Plugins/Runtime/Metasound/Source/MetasoundFrontend/Private/MetasoundFrontendDocumentBuilder.cpp
@@ -1815,14 +1815,14 @@ const FMetasoundFrontendNode* FMetaSoundFrontendDocumentBuilder::AddNodeByClassN
 	const FMetasoundFrontendClass* Dependency = nullptr;
 	FNodeRegistryKey ClassKey;
 	{
-		FMetasoundFrontendClass RegisteredClass;
-		if (!ISearchEngine::Get().FindClassWithHighestMinorVersion(InClassName, InMajorVersion, RegisteredClass))
+		FMetaSoundClassInfo ClassInfo;
+		if (!ISearchEngine::Get().FindRegisteredClass(InClassName, InMajorVersion, ClassInfo))
 		{
 			UE_LOG(LogMetaSound, Error, TEXT("Failed to add new node by class name '%s' and major version '%d': Class not found"), *InClassName.ToString(), InMajorVersion);
 			return nullptr;
 		}
 
-		const EMetasoundFrontendClassType ClassType = RegisteredClass.Metadata.GetType();
+		const EMetasoundFrontendClassType ClassType = ClassInfo.ClassType;
 		if (ClassType != EMetasoundFrontendClassType::External && ClassType != EMetasoundFrontendClassType::Graph)
 		{
 			UE_LOG(LogMetaSound, Warning, TEXT("Failed to add new node by class name '%s': Class is restricted type '%s' that cannot be added via this function."),
@@ -1832,12 +1832,13 @@ const FMetasoundFrontendNode* FMetaSoundFrontendDocumentBuilder::AddNodeByClassN
 		}
 
 		// Dependency is considered "External" when looked up or added as a dependency to a graph
-		RegisteredClass.Metadata.SetType(EMetasoundFrontendClassType::External);
-		ClassKey = FNodeRegistryKey(RegisteredClass.Metadata);
-		Dependency = FindDependency(RegisteredClass.Metadata);
+		ClassKey = ClassInfo.ToRegistryKey();
+		Dependency = FindDependency(ClassKey);
 		if (!Dependency)
 		{
-			Dependency = AddDependency(MoveTemp(RegisteredClass));
+			FMetasoundFrontendClass NewDependency;
+			INodeClassRegistry::GetChecked().GetFrontendClassFromRegistered(ClassKey, NewDependency);
+			Dependency = AddDependency(MoveTemp(NewDependency));
 		}
 	}
 
@@ -2173,21 +2174,21 @@ bool FMetaSoundFrontendDocumentBuilder::CanAutoUpdate(const FGuid& InNodeID, Met
 	const FMetasoundFrontendClassMetadata& NodeClassMetadata = Class->Metadata;
 	const IMetaSoundAssetManager& AssetManager = IMetaSoundAssetManager::GetChecked();
 	
-	FMetasoundFrontendClass RegistryClass;
-	if (!ISearchEngine::Get().FindClassWithHighestMinorVersion(NodeClassMetadata.GetClassName().ToNodeClassName(), NodeClassMetadata.GetVersion().Major, RegistryClass))
+	FMetaSoundClassInfo RegistryClassInfo;
+	if (!ISearchEngine::Get().FindRegisteredClass(NodeClassMetadata.GetClassName().ToNodeClassName(), NodeClassMetadata.GetVersion().Major, RegistryClassInfo))
 	{
 		return false;
 	}
 
 	// 1. Document's class version is somehow higher than registries, so can't update.
-	if (RegistryClass.Metadata.GetVersion() < NodeClassMetadata.GetVersion())
+	if (RegistryClassInfo.Version < NodeClassMetadata.GetVersion())
 	{
 		return false;
 	}
 
 	// 2. Document's class version is equal, so have to dif and check change IDs
 	// to ensure a change wasn't created that didn't contain a version change.
-	if (RegistryClass.Metadata.GetVersion() == NodeClassMetadata.GetVersion())
+	if (RegistryClassInfo.Version == NodeClassMetadata.GetVersion())
 	{
 		if (ChangeIDComparisonEnabledInAutoUpdate())
 		{
@@ -2196,9 +2197,12 @@ bool FMetaSoundFrontendDocumentBuilder::CanAutoUpdate(const FGuid& InNodeID, Met
 			// Its less performant to ignore ChangeIDs, but auto-update is no longer
 			// called at runtime and the editor now avoids running on mass numbers of
 			// assets at once due to the utilization of AssetTag data for querying class info.
-			const bool bIsAssetClass = AssetManager.IsAssetClass(RegistryClass.Metadata);
-			if (bIsAssetClass)
+			const FNodeClassRegistryKey ClassKey(RegistryClassInfo.ToRegistryKey());
+			bool bIsAssetClassType = FMetaSoundAssetKey::IsValidType(RegistryClassInfo.ClassType);
+			if (bIsAssetClassType && AssetManager.ContainsKey(FMetaSoundAssetKey(ClassKey)))
 			{
+				FMetasoundFrontendClass RegistryClass;
+				INodeClassRegistry::GetChecked().GetFrontendClassFromRegistered(ClassKey, RegistryClass);
 				if (RegistryClass.Metadata.GetChangeID() == NodeClassMetadata.GetChangeID())
 				{
 					const FMetasoundFrontendClassInterface& ClassInterface = Class->GetInterfaceForNode(*Node);
@@ -2495,10 +2499,10 @@ void FMetaSoundFrontendDocumentBuilder::SetPresetFlags(bool bIsPreset)
 
 #if WITH_EDITORONLY_DATA
 bool FMetaSoundFrontendDocumentBuilder::DiffAgainstRegistryInterface(const FGuid& InNodeID, const FGuid& InPageID, bool bInUseHighestMinorVersion, Metasound::Frontend::FClassInterfaceUpdates& OutInterfaceUpdates, bool bInForceRegenerateClassInterfaceOverride) const
-
 {
 	METASOUND_TRACE_CPUPROFILER_EVENT_SCOPE(FMetaSoundFrontendDocumentBuilder::DiffAgainstRegistryInterface);
 
+	using namespace Metasound;
 	using namespace Metasound::Frontend;
 
 	OutInterfaceUpdates = FClassInterfaceUpdates();
@@ -2518,30 +2522,40 @@ bool FMetaSoundFrontendDocumentBuilder::DiffAgainstRegistryInterface(const FGuid
 	const FMetasoundFrontendClassMetadata& NodeClassMetadata = Class->Metadata;
 	const FMetasoundFrontendClassInterface& NodeClassInterface = Class->GetInterfaceForNode(*Node);
 
-	Metasound::FNodeClassName NodeClassName = NodeClassMetadata.GetClassName().ToNodeClassName();
+	FNodeClassName NodeClassName = NodeClassMetadata.GetClassName().ToNodeClassName();
+
+	auto ClearInterface = [&]()
+	{
+		Algo::Transform(NodeClassInterface.Inputs, OutInterfaceUpdates.RemovedInputs, [&](const FMetasoundFrontendClassInput& Input) { return &Input; });
+		Algo::Transform(NodeClassInterface.Outputs, OutInterfaceUpdates.RemovedOutputs, [&](const FMetasoundFrontendClassOutput& Output) { return &Output; });
+	};
 
 	// Get class 
 	if (bInUseHighestMinorVersion)
 	{
-		if (!ISearchEngine::Get().FindClassWithHighestMinorVersion(NodeClassName, NodeClassMetadata.GetVersion().Major, OutInterfaceUpdates.RegistryClass))
+		FMetaSoundClassInfo ClassInfo;
+		if (!ISearchEngine::Get().FindRegisteredClass(NodeClassName, NodeClassMetadata.GetVersion().Major, ClassInfo))
+		{
+			ClearInterface();
+			return false;
+		}
+
+		bool bFoundRegisteredClass = INodeClassRegistry::GetChecked().FindFrontendClassFromRegistered(ClassInfo.ToRegistryKey(), OutInterfaceUpdates.RegistryClass);
+		if (!ensureMsgf(bFoundRegisteredClass, TEXT("SearchEngine returned registered class info for class that is not registered")))
 		{
-			Algo::Transform(NodeClassInterface.Inputs, OutInterfaceUpdates.RemovedInputs, [&](const FMetasoundFrontendClassInput& Input) { return &Input; });
-			Algo::Transform(NodeClassInterface.Outputs, OutInterfaceUpdates.RemovedOutputs, [&](const FMetasoundFrontendClassOutput& Output) { return &Output; });
+			ClearInterface();
 			return false;
 		}
 	}
 	else
 	{
 		// Find class with same metadata in the node registry.
-		INodeClassRegistry* Registry = INodeClassRegistry::Get();
-		checkf(nullptr != Registry, TEXT("The metasound node registry should always be available if the metasound plugin is loaded"));
-		bool bFoundRegisteredClass = Registry->FindFrontendClassFromRegistered(FNodeRegistryKey(NodeClassMetadata), OutInterfaceUpdates.RegistryClass);
+		bool bFoundRegisteredClass = INodeClassRegistry::GetChecked().FindFrontendClassFromRegistered(FNodeRegistryKey(NodeClassMetadata), OutInterfaceUpdates.RegistryClass);
 		if (!bFoundRegisteredClass)
 		{
 			// If the class was not found, mark all inputs and outputs as removed.
 			UE_LOG(LogMetaSound, Warning, TEXT("Could not find registered version of interface. %s %s is not registered."), *NodeClassMetadata.GetClassName().ToString(), *NodeClassMetadata.GetVersion().ToString());
-			Algo::Transform(NodeClassInterface.Inputs, OutInterfaceUpdates.RemovedInputs, [&](const FMetasoundFrontendClassInput& Input) { return &Input; });
-			Algo::Transform(NodeClassInterface.Outputs, OutInterfaceUpdates.RemovedOutputs, [&](const FMetasoundFrontendClassOutput& Output) { return &Output; });
+			ClearInterface();
 			return false;
 		}
 	}
diff --git a/Engine/Plugins/Runtime/Metasound/Source/MetasoundFrontend/Private/MetasoundFrontendGraphController.cpp b/Engine/Plugins/Runtime/Metasound/Source/MetasoundFrontend/Private/MetasoundFrontendGraphController.cpp
--- a/Engine/Plugins/Runtime/Metasound/Source/MetasoundFrontend/Private/MetasoundFrontendGraphController.cpp
+++ b/Engine/Plugins/Runtime/Metasound/Source/MetasoundFrontend/Private/MetasoundFrontendGraphController.cpp
@@ -1549,7 +1549,6 @@ PRAGMA_ENABLE_DEPRECATION_WARNINGS
 
 		FNodeHandle FGraphController::AddNode(const FNodeRegistryKey& InKey, FGuid InNodeGuid)
 		{
-			// Construct a FNodeClassInfo from this lookup key.
 			FConstClassAccessPtr Class = OwningDocument->FindOrAddClass(InKey);
 			const bool bIsValidClass = (nullptr != Class.Get());
 
@@ -1573,7 +1572,6 @@ PRAGMA_ENABLE_DEPRECATION_WARNINGS
 
 		FNodeHandle FGraphController::AddTemplateNode(const INodeTemplate& InTemplate, FNodeTemplateGenerateInterfaceParams Params, FGuid InNodeGuid)
 		{
-			// Construct a FNodeClassInfo from this lookup key.
 			const FNodeRegistryKey Key(InTemplate.GetFrontendClass().Metadata);
 			FConstClassAccessPtr Class = OwningDocument->FindOrAddClass(Key);
 			const bool bIsValidClass = (nullptr != Class.Get());
diff --git a/Engine/Plugins/Runtime/Metasound/Source/MetasoundFrontend/Private/MetasoundFrontendNodeClassRegistry.cpp b/Engine/Plugins/Runtime/Metasound/Source/MetasoundFrontend/Private/MetasoundFrontendNodeClassRegistry.cpp
--- a/Engine/Plugins/Runtime/Metasound/Source/MetasoundFrontend/Private/MetasoundFrontendNodeClassRegistry.cpp
+++ b/Engine/Plugins/Runtime/Metasound/Source/MetasoundFrontend/Private/MetasoundFrontendNodeClassRegistry.cpp
@@ -108,7 +108,6 @@ namespace Metasound::Frontend
 #endif // WITH_EDITOR
 		}
 
-
 		// FDocumentNodeClassRegistryEntry encapsulates a node registry entry for a FGraph
 		class FDocumentNodeClassRegistryEntry : public INodeClassRegistryEntry
 		{
@@ -116,12 +115,10 @@ namespace Metasound::Frontend
 			FDocumentNodeClassRegistryEntry(
 				const FMetasoundFrontendGraphClass& InGraphClass,
 				const TSet<FMetasoundFrontendVersion>& InInterfaces,
-				FNodeClassInfo&& InNodeClassInfo,
 				TSharedPtr<const IGraph> InGraph,
 				FTopLevelAssetPath InAssetPath)
 			: FrontendClass(InGraphClass)
 			, Interfaces(InInterfaces)
-			, ClassInfo(MoveTemp(InNodeClassInfo))
 			, Graph(InGraph)
 			, AssetPath(InAssetPath)
 			{
@@ -132,11 +129,6 @@ namespace Metasound::Frontend
 
 			virtual ~FDocumentNodeClassRegistryEntry() = default;
 
-			virtual const FNodeClassInfo& GetClassInfo() const override
-			{
-				return ClassInfo;
-			}
-
 			virtual TUniquePtr<INode> CreateNode(const FNodeInitData& InNodeInitData) const override
 			{
 				if (Graph.IsValid())
@@ -207,7 +199,12 @@ namespace Metasound::Frontend
 			FTopLevelAssetPath AssetPath;
 		};
 	} // namespace RegistryPrivate
-	  
+	
+	const FNodeClassInfo& INodeTemplateRegistryEntry::GetClassInfo() const
+	{
+		return FNodeClassInfo::GetInvalid();
+	}
+
 	TUniquePtr<INode> INodeClassRegistryEntry::CreateNode(FNodeData InNodeData) const
 	{
 #if !UE_METASOUND_DISABLE_5_6_NODE_REGISTRATION_DEPRECATION_WARNINGS
@@ -216,9 +213,10 @@ namespace Metasound::Frontend
 		if (!bDidLogError)
 		{
 			bDidLogError = true;
-
-			FNodeClassInfo ClassInfo = GetClassInfo();
-			UE_LOG(LogMetaSound, Warning, TEXT("Use of deprecated code path for node registration. First occurrence on node %s. Please implement INodeClassRegistryEntry::CreateNode(FNodeData) const"), *ClassInfo.ClassName.ToString());
+			UE_LOG(LogMetaSound, Warning,
+				TEXT("Use of deprecated code path for node registration. First occurrence on node %s. "
+					"Please implement INodeClassRegistryEntry::CreateNode(FNodeData) const"),
+				*GetFrontendClass().Metadata.GetClassName().ToString());
 		}
 #endif
 PRAGMA_DISABLE_DEPRECATION_WARNINGS
@@ -233,7 +231,11 @@ PRAGMA_ENABLE_DEPRECATION_WARNINGS
 		static bool bDidWarn = false;
 		if (!bDidWarn)
 		{
-			UE_LOG(LogMetaSound, Warning, TEXT("Please implement INodeClassRegistryEntry::CreateFrontendNodeConfiguration for the registry entry class representing node %s. This method will become pure virtual in future releases. Define UE_METASOUND_PURE_VIRTUAL_CREATE_FRONTEND_NODE_EXTENSION in order to build with this method as a pure virtual on the interface."), *GetClassInfo().ClassName.ToString());
+			UE_LOG(LogMetaSound, Warning,
+				TEXT("Please implement INodeClassRegistryEntry::CreateFrontendNodeConfiguration for the registry entry class representing node %s. "
+					"This method will become pure virtual in future releases. Define UE_METASOUND_PURE_VIRTUAL_CREATE_FRONTEND_NODE_EXTENSION in order "
+					"to build with this method as a pure virtual on the interface."),
+				*GetFrontendClass().Metadata.GetClassName().ToString());
 		}
 		return TInstancedStruct<FMetaSoundFrontendNodeConfiguration>();
 	}
@@ -254,8 +256,15 @@ PRAGMA_ENABLE_DEPRECATION_WARNINGS
 #endif
 
 	FNodeClassRegistryTransaction::FNodeClassRegistryTransaction(ETransactionType InType, const FNodeClassInfo& InNodeClassInfo, FNodeClassRegistryTransaction::FTimeType InTimestamp)
+		: Type(InType)
+		, RegistryKey()
+		, Timestamp(InTimestamp)
+	{
+	}
+
+	FNodeClassRegistryTransaction::FNodeClassRegistryTransaction(ETransactionType InType, FNodeClassRegistryKey InRegistryKey, FNodeClassRegistryTransaction::FTimeType InTimestamp)
 	: Type(InType)
-	, NodeClassInfo(InNodeClassInfo)
+	, RegistryKey(MoveTemp(InRegistryKey))
 	, Timestamp(InTimestamp)
 	{
 	}
@@ -267,12 +276,12 @@ PRAGMA_ENABLE_DEPRECATION_WARNINGS
 
 	const FNodeClassInfo& FNodeClassRegistryTransaction::GetNodeClassInfo() const
 	{
-		return NodeClassInfo;
+		return FNodeClassInfo::GetInvalid();
 	}
 
 	FNodeClassRegistryKey FNodeClassRegistryTransaction::GetNodeRegistryKey() const
 	{
-		return FNodeClassRegistryKey(NodeClassInfo);
+		return RegistryKey;
 	}
 
 	FNodeClassRegistryTransaction::FTimeType FNodeClassRegistryTransaction::GetTimestamp() const
@@ -355,15 +364,10 @@ PRAGMA_ENABLE_DEPRECATION_WARNINGS
 		{
 			return FNodeClassRegistryKey(InGraphClass);
 		}
-
-		FNodeClassRegistryKey CreateKey(const FNodeClassInfo& InClassInfo)
-		{
-			return FNodeClassRegistryKey(InClassInfo);
-		}
 	} // namespace NodeClassRegistryKey
 	  
 
-	void FNodeClassRegistry::BuildAndRegisterGraphFromDocument(const FMetasoundFrontendDocument& InDocument, const FProxyDataCache& InProxyDataCache, FNodeClassInfo&& InNodeClassInfo, const FTopLevelAssetPath& AssetPath, TArrayView<const FGuid> InPageOrder)
+	void FNodeClassRegistry::BuildAndRegisterGraphFromDocument(const FMetasoundFrontendDocument& InDocument, const FProxyDataCache& InProxyDataCache, const FMetasoundFrontendClassName& InClassName, const FTopLevelAssetPath& AssetPath, TArrayView<const FGuid> InPageOrder)
 	{
 		using namespace RegistryPrivate;
 
@@ -373,7 +377,7 @@ PRAGMA_ENABLE_DEPRECATION_WARNINGS
 		FGuid AssetClassID;
 		if (IMetaSoundAssetManager* AssetManager = IMetaSoundAssetManager::Get())
 		{
-			ensureAlways(AssetManager->TryGetAssetIDFromClassName(InNodeClassInfo.ClassName, AssetClassID));
+			ensureAlways(AssetManager->TryGetAssetIDFromClassName(InClassName, AssetClassID));
 		}
 		else
 		{
@@ -392,7 +396,6 @@ PRAGMA_ENABLE_DEPRECATION_WARNINGS
 		TUniquePtr<INodeClassRegistryEntry> RegistryEntry = MakeUnique<FDocumentNodeClassRegistryEntry>(
 			InDocument.RootGraph,
 			InDocument.Interfaces,
-			MoveTemp(InNodeClassInfo),
 			GraphToRegister,
 			AssetPath);
 
@@ -612,10 +615,9 @@ PRAGMA_ENABLE_DEPRECATION_WARNINGS
 
 		// Store update to newly registered node in history so nodes
 		// can be queried by transaction ID
-		FNodeClassInfo NodeClassInfo(Document.RootGraph);
 		{
 			FNodeClassRegistryTransaction::FTimeType Timestamp = FPlatformTime::Cycles64();
-			TransactionBuffer->AddTransaction(FNodeClassRegistryTransaction(FNodeClassRegistryTransaction::ETransactionType::NodeRegistration, NodeClassInfo, Timestamp));
+			TransactionBuffer->AddTransaction(FNodeClassRegistryTransaction(FNodeClassRegistryTransaction::ETransactionType::NodeRegistration, RegistryKey.NodeKey, Timestamp));
 		}
 
 		if (bAsync)
@@ -623,12 +625,12 @@ PRAGMA_ENABLE_DEPRECATION_WARNINGS
 			TArray<FGuid> PageOrder(InPageOrder.GetData(), InPageOrder.Num());
 			Tasks::FTask BuildAndRegisterTask = AsyncRegistrationPipe.Launch(
 				UE_SOURCE_LOCATION,
-				[RegistryKey, ClassInfo = MoveTemp(NodeClassInfo), AssetPath, RegistryDocInterface, ProxyDataCache = MoveTemp(ProxyDataCache), PageOrder = MoveTemp(PageOrder)]() mutable
+				[RegistryKey, AssetPath, RegistryDocInterface, ProxyDataCache = MoveTemp(ProxyDataCache), PageOrder = MoveTemp(PageOrder)]() mutable
 				{
 					FNodeClassRegistry& Registry = FNodeClassRegistry::Get();
 					// Unregister the graph before re-registering
 					Registry.UnregisterGraphInternal(RegistryKey);
-					Registry.BuildAndRegisterGraphFromDocument(RegistryDocInterface->GetConstDocument(), ProxyDataCache, MoveTemp(ClassInfo), AssetPath, PageOrder);
+					Registry.BuildAndRegisterGraphFromDocument(RegistryDocInterface->GetConstDocument(), ProxyDataCache, RegistryKey.NodeKey.ClassName, AssetPath, PageOrder);
 					Registry.RemoveRegistrationTask(RegistryKey, FNodeClassRegistryTransaction::ETransactionType::NodeRegistration);
 					Registry.RemoveDocumentReference(RegistryDocInterface);
 				}
@@ -647,7 +649,7 @@ PRAGMA_ENABLE_DEPRECATION_WARNINGS
 			UnregisterGraphInternal(RegistryKey);
 
 			// Build and register graph synchronously
-			BuildAndRegisterGraphFromDocument(RegistryDocInterface->GetConstDocument(), ProxyDataCache, MoveTemp(NodeClassInfo), AssetPath, InPageOrder);
+			BuildAndRegisterGraphFromDocument(RegistryDocInterface->GetConstDocument(), ProxyDataCache, RegistryKey.NodeKey.ClassName, AssetPath, InPageOrder);
 		}
 
 		return RegistryKey;
@@ -785,7 +787,10 @@ PRAGMA_ENABLE_DEPRECATION_WARNINGS
 		using namespace RegistryPrivate;
 
 		METASOUND_TRACE_CPUPROFILER_EVENT_SCOPE(Metasound::FNodeClassRegistry::UnregisterGraph);
+
+#if (UE_BUILD_DEBUG || UE_BUILD_DEVELOPMENT)
 		METASOUND_TRACE_CPUPROFILER_EVENT_SCOPE_TEXT(*InRegistryKey.ToString(TEXT("FNodeClassRegistry::UnregisterGraph")));
+#endif // (UE_BUILD_DEBUG || UE_BUILD_DEVELOPMENT)
 
 		check(IsInGameThread());
 		check(InDocumentInterface);
@@ -797,19 +802,10 @@ PRAGMA_ENABLE_DEPRECATION_WARNINGS
 			return false;
 		}
 
-		const FMetasoundFrontendDocument& Document = InDocumentInterface->GetConstDocument();
-		FNodeClassInfo NodeClassInfo(Document.RootGraph.Metadata);
-
-		// This is a hack to avoid requiring the asset path to be passed while unregistering.
-		// The asset path may be invalid by this point if the object being unregistered is being GC'ed.
-		// FNodeClassInfo needs to be deprecated in favor of more precise types as a key, editor data, etc.
-		// Its currently kind of a dumping ground as it stands.
-		NodeClassInfo.Type = EMetasoundFrontendClassType::External;
-
 		// Store update to unregistered node in history so nodes can be queried by transaction ID
 		{
 			FNodeClassRegistryTransaction::FTimeType Timestamp = FPlatformTime::Cycles64();
-			TransactionBuffer->AddTransaction(FNodeClassRegistryTransaction(FNodeClassRegistryTransaction::ETransactionType::NodeUnregistration, NodeClassInfo, Timestamp));
+			TransactionBuffer->AddTransaction(FNodeClassRegistryTransaction(FNodeClassRegistryTransaction::ETransactionType::NodeUnregistration, InRegistryKey.NodeKey, Timestamp));
 		}
 
 		// Async registration is only available if:
@@ -877,9 +873,13 @@ PRAGMA_ENABLE_DEPRECATION_WARNINGS
 			return { };
 		}
 
-		const FNodeClassRegistryKey Key = FNodeClassRegistryKey(InEntry->GetClassInfo());
-		METASOUND_TRACE_CPUPROFILER_EVENT_SCOPE_TEXT(*Key.ToString(TEXT("FNodeClassRegistry::RegisterNodeInternal")))
+		const FNodeClassRegistryKey Key = FNodeClassRegistryKey(InEntry->GetFrontendClass().Metadata);
 
+#if UE_BUILD_DEBUG || UE_BUILD_DEVELOPMENT
+		METASOUND_TRACE_CPUPROFILER_EVENT_SCOPE_TEXT(*Key.ToString(TEXT("FNodeClassRegistry::RegisterNodeInternal")))
+#else // !(UE_BUILD_DEBUG || UE_BUILD_DEVELOPMENT)
+		METASOUND_TRACE_CPUPROFILER_EVENT_SCOPE("FNodeClassRegistry::RegisterNodeInternal")
+#endif // !(UE_BUILD_DEBUG || UE_BUILD_DEVELOPMENT)
 
 #if !NO_LOGGING
 		TArray<TSharedRef<INodeClassRegistryEntry>> Entries;
@@ -920,7 +920,6 @@ PRAGMA_ENABLE_DEPRECATION_WARNINGS
 
 	FNodeClassRegistryKey FNodeClassRegistry::RegisterNode(TUniquePtr<INodeClassRegistryEntry>&& InEntry)
 	{
-		const FNodeClassInfo ClassInfo = InEntry->GetClassInfo();
 		const FNodeClassRegistryKey Key = RegisterNodeInternal(MoveTemp(InEntry));
 
 		if (Key.IsValid())
@@ -928,7 +927,7 @@ PRAGMA_ENABLE_DEPRECATION_WARNINGS
 			// Store update to newly registered node in history so nodes
 			// can be queried by transaction ID
 			const FNodeClassRegistryTransaction::FTimeType Timestamp = FPlatformTime::Cycles64();
-			TransactionBuffer->AddTransaction(FNodeClassRegistryTransaction(FNodeClassRegistryTransaction::ETransactionType::NodeRegistration, ClassInfo, Timestamp));
+			TransactionBuffer->AddTransaction(FNodeClassRegistryTransaction(FNodeClassRegistryTransaction::ETransactionType::NodeRegistration, Key, Timestamp));
 		}
 
 		return Key;
@@ -937,12 +936,15 @@ PRAGMA_ENABLE_DEPRECATION_WARNINGS
 #if WITH_EDITORONLY_DATA
 	bool FNodeClassRegistry::RegisterNodeMigration(const FNodeMigrationInfo& InMigrationInfo)
 	{
-		FNodeClassInfo ClassInfo;
-		ClassInfo.ClassName = InMigrationInfo.ClassName;
-		ClassInfo.Version = FMetasoundFrontendVersionNumber{ .Major=InMigrationInfo.MajorVersion, .Minor=InMigrationInfo.MinorVersion };
-		ClassInfo.Type = EMetasoundFrontendClassType::External;
-
-		const FNodeClassRegistryKey Key{ClassInfo};
+		const FNodeClassRegistryKey Key(
+			EMetasoundFrontendClassType::External,
+			InMigrationInfo.ClassName,
+			FMetasoundFrontendVersionNumber
+			{
+				.Major = InMigrationInfo.MajorVersion,
+				.Minor = InMigrationInfo.MinorVersion
+			}
+		);
 
 		// Attempts to register the exact same migration info twice will result in bSuccess=false
 		bool bSuccess = false;
@@ -961,7 +963,7 @@ PRAGMA_ENABLE_DEPRECATION_WARNINGS
 		{
 			// Add to transactions so that the search engine knows it exists. 
 			const FNodeClassRegistryTransaction::FTimeType Timestamp = FPlatformTime::Cycles64();
-			TransactionBuffer->AddTransaction(FNodeClassRegistryTransaction(FNodeClassRegistryTransaction::ETransactionType::NodeMigrationRegistration, ClassInfo, Timestamp));
+			TransactionBuffer->AddTransaction(FNodeClassRegistryTransaction(FNodeClassRegistryTransaction::ETransactionType::NodeMigrationRegistration, Key, Timestamp));
 		}
 
 		UE_CLOG(!bSuccess, LogMetaSound, Error, TEXT("Invalid attempt to register node migration more than once: %s. Please ensure node migration info is unique."), *InMigrationInfo.ToString());
@@ -972,12 +974,15 @@ PRAGMA_ENABLE_DEPRECATION_WARNINGS
 
 	bool FNodeClassRegistry::UnregisterNodeMigration(const FNodeMigrationInfo& InMigrationInfo)
 	{
-		FNodeClassInfo ClassInfo;
-		ClassInfo.ClassName = InMigrationInfo.ClassName;
-		ClassInfo.Version = FMetasoundFrontendVersionNumber{ .Major=InMigrationInfo.MajorVersion, .Minor=InMigrationInfo.MinorVersion };
-		ClassInfo.Type = EMetasoundFrontendClassType::External;
-
-		const FNodeClassRegistryKey Key{ClassInfo};
+		const FNodeClassRegistryKey Key(
+			EMetasoundFrontendClassType::External,
+			InMigrationInfo.ClassName,
+			FMetasoundFrontendVersionNumber
+			{
+				.Major = InMigrationInfo.MajorVersion,
+				.Minor = InMigrationInfo.MinorVersion
+			}
+		);
 
 		// Attempts to unregister a migration that does not exist will result in bSuccess=false
 		bool bSuccess = false;
@@ -996,7 +1001,7 @@ PRAGMA_ENABLE_DEPRECATION_WARNINGS
 		{
 			// Add to transactions so that the search engine knows it exists. 
 			const FNodeClassRegistryTransaction::FTimeType Timestamp = FPlatformTime::Cycles64();
-			TransactionBuffer->AddTransaction(FNodeClassRegistryTransaction(FNodeClassRegistryTransaction::ETransactionType::NodeMigrationUnregistration, ClassInfo, Timestamp));
+			TransactionBuffer->AddTransaction(FNodeClassRegistryTransaction(FNodeClassRegistryTransaction::ETransactionType::NodeMigrationUnregistration, Key, Timestamp));
 		}
 
 		UE_CLOG(!bSuccess, LogMetaSound, Error, TEXT("Invalid attempt to unregister node migration more than once: %s. Please ensure node migration info is unique."), *InMigrationInfo.ToString());
@@ -1017,7 +1022,7 @@ PRAGMA_ENABLE_DEPRECATION_WARNINGS
 
 			FNodeClassRegistryTransaction::FTimeType Timestamp = FPlatformTime::Cycles64();
 
-			Key = FNodeClassRegistryKey(Entry->GetClassInfo());
+			Key = FNodeClassRegistryKey(Entry->GetFrontendClass().Metadata);
 
 			{
 				UE::TScopeLock Lock(RegistryMapsCriticalSection);
@@ -1034,28 +1039,28 @@ PRAGMA_ENABLE_DEPRECATION_WARNINGS
 			// Store update to newly registered node in history so nodes
 			// can be queried by transaction ID
 
-			TransactionBuffer->AddTransaction(FNodeClassRegistryTransaction(FNodeClassRegistryTransaction::ETransactionType::NodeRegistration, Entry->GetClassInfo(), Timestamp));
+			TransactionBuffer->AddTransaction(FNodeClassRegistryTransaction(FNodeClassRegistryTransaction::ETransactionType::NodeRegistration, Key, Timestamp));
 		}
 
 		return Key;
 	}
 
-	bool FNodeClassRegistry::UnregisterNodeInternal(const FNodeClassRegistryKey& InKey, FNodeClassInfo* OutClassInfo)
+	bool FNodeClassRegistry::UnregisterNodeInternal(const FNodeClassRegistryKey& InKey)
 	{
 		METASOUND_LLM_SCOPE;
 
 		if (InKey.IsValid())
 		{
+#if (UE_BUILD_DEBUG || UE_BUILD_DEVELOPMENT)
 			METASOUND_TRACE_CPUPROFILER_EVENT_SCOPE_TEXT(*FString::Printf(TEXT("FNodeClassRegistry::UnregisterNodeInternal key %s"), *InKey.ToString()))
+#else // !(UE_BUILD_DEBUG || UE_BUILD_DEVELOPMENT)
+			METASOUND_TRACE_CPUPROFILER_EVENT_SCOPE_TEXT(TEXT("FNodeClassRegistry::UnregisterNodeInternal"))
+#endif // !(UE_BUILD_DEBUG || UE_BUILD_DEVELOPMENT)
 
 			UE::TScopeLock Lock(RegistryMapsCriticalSection);
 			if (const TSharedRef<INodeClassRegistryEntry, ESPMode::ThreadSafe>* EntryPtr = RegisteredNodes.Find(InKey))
 			{
 				const TSharedRef<INodeClassRegistryEntry>& Entry = *EntryPtr;
-				if (OutClassInfo)
-				{
-					*OutClassInfo = Entry->GetClassInfo();
-				}
 				const uint32 NumRemoved = RegisteredNodes.RemoveSingle(InKey, Entry);
 				if (ensure(NumRemoved == 1))
 				{
@@ -1064,20 +1069,15 @@ PRAGMA_ENABLE_DEPRECATION_WARNINGS
 			}
 		}
 
-		if (OutClassInfo)
-		{
-			*OutClassInfo = { };
-		}
 		return false;
 	}
 
 	bool FNodeClassRegistry::UnregisterNode(const FNodeClassRegistryKey& InKey)
 	{
-		FNodeClassInfo ClassInfo;
-		if (UnregisterNodeInternal(InKey, &ClassInfo))
+		if (UnregisterNodeInternal(InKey))
 		{
 			const FNodeClassRegistryTransaction::FTimeType Timestamp = FPlatformTime::Cycles64();
-			TransactionBuffer->AddTransaction(FNodeClassRegistryTransaction(FNodeClassRegistryTransaction::ETransactionType::NodeUnregistration, ClassInfo, Timestamp));
+			TransactionBuffer->AddTransaction(FNodeClassRegistryTransaction(FNodeClassRegistryTransaction::ETransactionType::NodeUnregistration, InKey, Timestamp));
 
 			return true;
 		}
@@ -1095,7 +1095,7 @@ PRAGMA_ENABLE_DEPRECATION_WARNINGS
 			{
 				FNodeClassRegistryTransaction::FTimeType Timestamp = FPlatformTime::Cycles64();
 
-				TransactionBuffer->AddTransaction(FNodeClassRegistryTransaction(FNodeClassRegistryTransaction::ETransactionType::NodeUnregistration, Entry->GetClassInfo(), Timestamp));
+				TransactionBuffer->AddTransaction(FNodeClassRegistryTransaction(FNodeClassRegistryTransaction::ETransactionType::NodeUnregistration, InKey, Timestamp));
 
 				{
 					UE::TScopeLock Lock(RegistryMapsCriticalSection);
@@ -1333,7 +1333,7 @@ PRAGMA_ENABLE_DEPRECATION_WARNINGS
 			// Only call function on classes of certain type.
 			auto IsMatchingClassType = [&](const TPair<FNodeClassRegistryKey, TSharedPtr<INodeClassRegistryEntry, ESPMode::ThreadSafe>>& Pair)
 			{
-				return Pair.Value->GetClassInfo().Type == InClassType;
+				return Pair.Value->GetFrontendClass().Metadata.GetType() == InClassType;
 			};
 			Algo::ForEachIf(RegisteredNodes, IsMatchingClassType, WrappedFunc);
 		}
diff --git a/Engine/Plugins/Runtime/Metasound/Source/MetasoundFrontend/Private/MetasoundFrontendNodeClassRegistryPrivate.h b/Engine/Plugins/Runtime/Metasound/Source/MetasoundFrontend/Private/MetasoundFrontendNodeClassRegistryPrivate.h
--- a/Engine/Plugins/Runtime/Metasound/Source/MetasoundFrontend/Private/MetasoundFrontendNodeClassRegistryPrivate.h
+++ b/Engine/Plugins/Runtime/Metasound/Source/MetasoundFrontend/Private/MetasoundFrontendNodeClassRegistryPrivate.h
@@ -34,13 +34,8 @@ namespace Metasound::Frontend
 	public:
 		virtual ~INodeTemplateRegistryEntry() = default;
 
-		/** Return FNodeClassInfo for the node class.
-		 *
-		 * Implementations of method should avoid any expensive operations
-		 * (e.g. loading from disk, allocating memory) as this method is called
-		 * frequently when querying nodes.
-		 */
-		virtual const FNodeClassInfo& GetClassInfo() const = 0;
+		UE_DEPRECATED(5.8, "Use GetFrontendClass() instead")
+		virtual const FNodeClassInfo& GetClassInfo() const;
 
 		/** Return a FMetasoundFrontendClass which describes the node. */
 		virtual const FMetasoundFrontendClass& GetFrontendClass() const = 0;
@@ -144,7 +139,7 @@ namespace Metasound::Frontend
 			FTopLevelAssetPath AssetPath;
 		};
 
-		void BuildAndRegisterGraphFromDocument(const FMetasoundFrontendDocument& InDocument, const FProxyDataCache& InProxyDataCache, FNodeClassInfo&& InNodeClassInfo, const FTopLevelAssetPath& AssetPath, TArrayView<const FGuid> InPageOrder);
+		void BuildAndRegisterGraphFromDocument(const FMetasoundFrontendDocument& InDocument, const FProxyDataCache& InProxyDataCache, const FMetasoundFrontendClassName& InClassName, const FTopLevelAssetPath& AssetPath, TArrayView<const FGuid> InPageOrder);
 
 		void AddRegistrationTask(const FGraphRegistryKey& InKey, FActiveRegistrationTaskInfo&& TaskInfo);
 		void RemoveRegistrationTask(const FGraphRegistryKey& InKey, FNodeClassRegistryTransaction::ETransactionType TransactionType);
@@ -157,7 +152,7 @@ namespace Metasound::Frontend
 		// indicating async registration/unregistration task(s) are no longer depending on it.
 		void RemoveDocumentReference(TScriptInterface<IMetaSoundDocumentInterface> DocumentInterface);
 
-		bool UnregisterNodeInternal(const FNodeClassRegistryKey& InKey, FNodeClassInfo* OutClassInfo = nullptr);
+		bool UnregisterNodeInternal(const FNodeClassRegistryKey& InKey);
 		FNodeClassRegistryKey RegisterNodeInternal(TUniquePtr<INodeClassRegistryEntry>&& InEntry);
 
 		static FNodeClassRegistry* LazySingleton;
diff --git a/Engine/Plugins/Runtime/Metasound/Source/MetasoundFrontend/Private/MetasoundFrontendNodeController.cpp b/Engine/Plugins/Runtime/Metasound/Source/MetasoundFrontend/Private/MetasoundFrontendNodeController.cpp
--- a/Engine/Plugins/Runtime/Metasound/Source/MetasoundFrontend/Private/MetasoundFrontendNodeController.cpp
+++ b/Engine/Plugins/Runtime/Metasound/Source/MetasoundFrontend/Private/MetasoundFrontendNodeController.cpp
@@ -1032,27 +1032,6 @@ namespace Metasound
 			return ClassInterfaceOverride;
 		}
 
-		bool FBaseNodeController::FindRegisteredClass(const FMetasoundFrontendClassMetadata& InMetadata, bool bInUseHighestMinorVersion, FMetasoundFrontendClass& OutRegisteredClass) const
-		{
-			bool bFoundInRegistry = false;
-
-			const INodeClassRegistry* Registry = INodeClassRegistry::Get();
-			checkf(nullptr != Registry, TEXT("The metasound node registry should always be available if the metasound plugin is loaded"));
-
-			if (bInUseHighestMinorVersion)
-			{
-				// Find class with same name and major version in the search engine
-				bFoundInRegistry = ISearchEngine::Get().FindClassWithHighestMinorVersion(InMetadata.GetClassName().ToNodeClassName(), InMetadata.GetVersion().Major, OutRegisteredClass);
-			}
-			else
-			{
-				// Find class with same metadata in the node registry.
-				bFoundInRegistry = Registry->FindFrontendClassFromRegistered(FNodeRegistryKey(InMetadata), OutRegisteredClass);
-			}
-
-			return bFoundInRegistry;
-		}
-
 		bool FBaseNodeController::ShouldReplaceExistingNodeConfig(TConstStructView<FMetaSoundFrontendNodeConfiguration> InRegisteredNodeConfig, TConstStructView<FMetaSoundFrontendNodeConfiguration> InExistingConfig) const
 		{
 			// The existing node configuration is not always replaced. Replacing
@@ -1177,8 +1156,22 @@ namespace Metasound
 			const FMetasoundFrontendClassInterface& NodeClassInterface = GetClassInterface();
 			Metasound::FNodeClassName NodeClassName = NodeClassMetadata.GetClassName().ToNodeClassName();
 
-			bool bDidFindClass = FindRegisteredClass(NodeClassMetadata, bInUseHighestMinorVersion, OutInterfaceUpdates.RegistryClass);
+			FNodeClassRegistryKey RegistryKey;
+			if (bInUseHighestMinorVersion)
+			{
+				FMetaSoundClassInfo ClassInfo;
+				const bool bIsRegistered = ISearchEngine::Get().FindRegisteredClass(NodeClassMetadata.GetClassName().ToNodeClassName(), NodeClassMetadata.GetVersion().Major, ClassInfo);
+				if (bIsRegistered)
+				{
+					RegistryKey = ClassInfo.ToRegistryKey();
+				}
+			}
+			else
+			{
+				RegistryKey = FNodeClassRegistryKey(NodeClassMetadata);
+			}
 
+			const bool bDidFindClass = INodeClassRegistry::GetChecked().FindFrontendClassFromRegistered(RegistryKey, OutInterfaceUpdates.RegistryClass);
 			if (!bDidFindClass)
 			{
 				// The class couldn't be found. Mark all relevant fields as removed. 
@@ -1233,11 +1226,13 @@ namespace Metasound
 			const FMetasoundFrontendClassMetadata& NodeClassMetadata = GetClassMetadata();
 			const IMetaSoundAssetManager& AssetManager = IMetaSoundAssetManager::GetChecked();
 
+PRAGMA_DISABLE_DEPRECATION_WARNINGS
 			FMetasoundFrontendClass RegistryClass;
 			if (!ISearchEngine::Get().FindClassWithHighestMinorVersion( NodeClassMetadata.GetClassName().ToNodeClassName(), NodeClassMetadata.GetVersion().Major, RegistryClass))
 			{
 				return false;
 			}
+PRAGMA_ENABLE_DEPRECATION_WARNINGS
 
 			// 1. Document's class version is somehow higher than registries, so can't update.
 			if (RegistryClass.Metadata.GetVersion() < NodeClassMetadata.GetVersion())
diff --git a/Engine/Plugins/Runtime/Metasound/Source/MetasoundFrontend/Private/MetasoundFrontendNodeTemplateRegistry.cpp b/Engine/Plugins/Runtime/Metasound/Source/MetasoundFrontend/Private/MetasoundFrontendNodeTemplateRegistry.cpp
--- a/Engine/Plugins/Runtime/Metasound/Source/MetasoundFrontend/Private/MetasoundFrontendNodeTemplateRegistry.cpp
+++ b/Engine/Plugins/Runtime/Metasound/Source/MetasoundFrontend/Private/MetasoundFrontendNodeTemplateRegistry.cpp
@@ -185,23 +185,16 @@ namespace Metasound::Frontend
 	{
 		class FTemplateRegistryEntry : public INodeTemplateRegistryEntry
 		{
-			const FNodeClassInfo ClassInfo;
 			const FMetasoundFrontendClass FrontendClass;
 
 		public:
 			FTemplateRegistryEntry(const INodeTemplate& InNodeTemplate)
-				: ClassInfo(InNodeTemplate.GetFrontendClass().Metadata)
-				, FrontendClass(InNodeTemplate.GetFrontendClass())
+				: FrontendClass(InNodeTemplate.GetFrontendClass())
 			{
 			}
 
 			virtual ~FTemplateRegistryEntry() = default;
 
-			virtual const FNodeClassInfo& GetClassInfo() const override
-			{
-				return ClassInfo;
-			}
-
 			/** Return a FMetasoundFrontendClass which describes the node. */
 			virtual const FMetasoundFrontendClass& GetFrontendClass() const override
 			{
diff --git a/Engine/Plugins/Runtime/Metasound/Source/MetasoundFrontend/Private/MetasoundFrontendSearchEngine.cpp b/Engine/Plugins/Runtime/Metasound/Source/MetasoundFrontend/Private/MetasoundFrontendSearchEngine.cpp
--- a/Engine/Plugins/Runtime/Metasound/Source/MetasoundFrontend/Private/MetasoundFrontendSearchEngine.cpp
+++ b/Engine/Plugins/Runtime/Metasound/Source/MetasoundFrontend/Private/MetasoundFrontendSearchEngine.cpp
@@ -32,10 +32,13 @@ namespace Metasound::Frontend
 	TArray<FMetasoundFrontendClass> ISearchEngine::FindAllClasses(bool bInIncludeAllVersions)
 	{
 		constexpr bool bIncludeUnloadedAssets = false;
+
+PRAGMA_DISABLE_DEPRECATION_WARNINGS
 		const TArray<FMetaSoundClassInfo> AllClasses = FindAllClasses(bInIncludeAllVersions
 			? ISearchEngine::EResultVersion::All
 			: ISearchEngine::EResultVersion::Highest,
 			bIncludeUnloadedAssets);
+PRAGMA_ENABLE_DEPRECATION_WARNINGS
 
 		TArray<FMetasoundFrontendClass> ToReturn;
 		for (const FMetaSoundClassInfo& ClassInfo : AllClasses)
diff --git a/Engine/Plugins/Runtime/Metasound/Source/MetasoundFrontend/Private/MetasoundFrontendSearchEngineCore.cpp b/Engine/Plugins/Runtime/Metasound/Source/MetasoundFrontend/Private/MetasoundFrontendSearchEngineCore.cpp
--- a/Engine/Plugins/Runtime/Metasound/Source/MetasoundFrontend/Private/MetasoundFrontendSearchEngineCore.cpp
+++ b/Engine/Plugins/Runtime/Metasound/Source/MetasoundFrontend/Private/MetasoundFrontendSearchEngineCore.cpp
@@ -469,42 +469,6 @@ namespace Metasound::Frontend
 		return false;
 	}
 
-	bool FSearchEngineCore::FindClassWithHighestMinorVersion(const FMetasoundFrontendClassName& InName, int32 InMajorVersion, FMetasoundFrontendClass& OutClass)
-	{
-		METASOUND_TRACE_CPUPROFILER_EVENT_SCOPE(metasound::FSearchEngineCore::FindClassWithHighestMinorVersion);
-		using namespace SearchEngineQuerySteps;
-
-		const FFrontendQueryKey QueryKey = CreateKeyFromFullClassNameAndMajorVersion(InName, InMajorVersion);
-		FNodeRegistryKey NodeRegistryKey; 
-		if (FindClassWithHighestMinorVersionQuery.UpdateAndFindResult(QueryKey, NodeRegistryKey))
-		{
-			if (FMetasoundFrontendRegistryContainer* NodeRegistry = FMetasoundFrontendRegistryContainer::Get())
-			{
-				return NodeRegistry->FindFrontendClassFromRegistered(NodeRegistryKey, OutClass);
-			}
-		}
-		return false;
-	}
-
-	TArray<FMetasoundFrontendInterface> FSearchEngineCore::FindUClassDefaultInterfaces(FName InUClassName)
-	{
-		using namespace SearchEngineQuerySteps;
-
-		METASOUND_TRACE_CPUPROFILER_EVENT_SCOPE(metasound::FSearchEngineCore::FindUClassDefaultInterfaces);
-
-		TArray<FMetasoundFrontendInterface> Interfaces;
-		TArray<FMetasoundFrontendVersion> InterfaceVersions = FindUClassDefaultInterfaceVersions(FTopLevelAssetPath(InUClassName.ToString()));
-		Algo::Transform(InterfaceVersions, Interfaces, [](const FMetasoundFrontendVersion& Version)
-		{
-			FMetasoundFrontendInterface Interface;
-			const FInterfaceRegistryKey Key = GetInterfaceRegistryKey(Version);
-			ensure(IInterfaceRegistry::Get().FindInterface(Key, Interface));
-			return Interface;
-		});
-
-		return Interfaces;
-	}
-
 	TArray<FMetasoundFrontendVersion> FSearchEngineCore::FindUClassDefaultInterfaceVersions(const FTopLevelAssetPath& InUClassPath)
 	{
 		using namespace SearchEngineQuerySteps;
diff --git a/Engine/Plugins/Runtime/Metasound/Source/MetasoundFrontend/Private/MetasoundFrontendSearchEngineCore.h b/Engine/Plugins/Runtime/Metasound/Source/MetasoundFrontend/Private/MetasoundFrontendSearchEngineCore.h
--- a/Engine/Plugins/Runtime/Metasound/Source/MetasoundFrontend/Private/MetasoundFrontendSearchEngineCore.h
+++ b/Engine/Plugins/Runtime/Metasound/Source/MetasoundFrontend/Private/MetasoundFrontendSearchEngineCore.h
@@ -325,11 +325,9 @@ namespace Metasound
 			virtual ~FSearchEngineCore() = default;
 
 			virtual void Prime() override;
-			virtual bool FindClassWithHighestMinorVersion(const FMetasoundFrontendClassName& InName, int32 InMajorVersion, FMetasoundFrontendClass& OutClass) override;
 
 			virtual bool FindRegisteredClass(const FMetasoundFrontendClassName& InName, int32 InMajorVersion, FMetaSoundClassInfo& OutClass) override;
 
-			virtual TArray<FMetasoundFrontendInterface> FindUClassDefaultInterfaces(FName InUClassName) override;
 			virtual TArray<FMetasoundFrontendVersion> FindUClassDefaultInterfaceVersions(const FTopLevelAssetPath& InUClassPath) override;
 			virtual TArray<FMetasoundFrontendVersion> FindAllRegisteredInterfacesWithName(FName InInterfaceName) override;
 			virtual bool FindInterfaceWithHighestVersion(FName InInterfaceName, FMetasoundFrontendInterface& OutInterface) override;
diff --git a/Engine/Plugins/Runtime/Metasound/Source/MetasoundFrontend/Private/MetasoundFrontendSearchEngineEditorOnly.h b/Engine/Plugins/Runtime/Metasound/Source/MetasoundFrontend/Private/MetasoundFrontendSearchEngineEditorOnly.h
--- a/Engine/Plugins/Runtime/Metasound/Source/MetasoundFrontend/Private/MetasoundFrontendSearchEngineEditorOnly.h
+++ b/Engine/Plugins/Runtime/Metasound/Source/MetasoundFrontend/Private/MetasoundFrontendSearchEngineEditorOnly.h
@@ -92,7 +92,7 @@ namespace Metasound::Frontend
 		virtual ~FSearchEngineEditorOnly() = default;
 
 		virtual void Prime() override;
-		virtual TArray<FMetaSoundClassInfo> FindAllClasses(ISearchEngine::EResultVersion VersionOptions, bool bIncludeUnloadedAssets) override;
+		virtual TArray<FMetaSoundClassInfo> FindAllClasses(ISearchEngine::EResultVersion VersionOptions, bool bIncludeUnloadedAssets = false) override;
 		virtual TArray<FMetaSoundClassInfo> FindClassesWithName(const FMetasoundFrontendClassName& InName, ESortByVersion SortByVersion = ESortByVersion::Yes) override;
 		virtual TArray<FMetasoundFrontendClass> FindClassesWithName(const FMetasoundFrontendClassName& InName, bool bInSortByVersion) override;
 		virtual bool FindClassWithHighestVersion(const FMetasoundFrontendClassName& InName, FMetaSoundClassInfo& OutClassInfo) override;
diff --git a/Engine/Plugins/Runtime/Metasound/Source/MetasoundFrontend/Private/MetasoundFrontendTransform.cpp b/Engine/Plugins/Runtime/Metasound/Source/MetasoundFrontend/Private/MetasoundFrontendTransform.cpp
--- a/Engine/Plugins/Runtime/Metasound/Source/MetasoundFrontend/Private/MetasoundFrontendTransform.cpp
+++ b/Engine/Plugins/Runtime/Metasound/Source/MetasoundFrontend/Private/MetasoundFrontendTransform.cpp
@@ -546,16 +546,16 @@ namespace Metasound
 					UpdatedClasses.Add(Class.ID);
 
 					// Check if a updated minor version exists.
-					FMetasoundFrontendClass ClassWithHighestMinorVersion;
-					const bool bFoundClassInSearchEngine = Frontend::ISearchEngine::Get().FindClassWithHighestMinorVersion(NodeClassKey.ClassName, NodeClassKey.Version.Major, ClassWithHighestMinorVersion);
+					FMetaSoundClassInfo ClassInfoWithHighestMinorVersion;
+					const bool bFoundClassInSearchEngine = Frontend::ISearchEngine::Get().FindRegisteredClass(NodeClassKey.ClassName, NodeClassKey.Version.Major, ClassInfoWithHighestMinorVersion);
 
-					if (bFoundClassInSearchEngine && (ClassWithHighestMinorVersion.Metadata.GetVersion() > NodeClassKey.Version))
+					if (bFoundClassInSearchEngine && (ClassInfoWithHighestMinorVersion.Version > NodeClassKey.Version))
 					{
-						const FMetasoundFrontendVersionNumber UpdateVersion = ClassWithHighestMinorVersion.Metadata.GetVersion();
+						const FMetasoundFrontendVersionNumber UpdateVersion = ClassInfoWithHighestMinorVersion.Version;
 
 						METASOUND_VERSIONING_LOG(Display, TEXT("Auto-Updating '%s' node class '%s': Newer version '%s' found."), *DebugAssetPath, *NodeClassKey.ClassName.ToString(), *UpdateVersion.ToString());
 
-						DependenciesToUpdate.Add(FDependencyPair(NodeClassKey, FNodeClassRegistryKey(ClassWithHighestMinorVersion.Metadata)));
+						DependenciesToUpdate.Add(FDependencyPair(NodeClassKey, ClassInfoWithHighestMinorVersion.ToRegistryKey()));
 					}
 					else if (InterfaceUpdates.ContainsChanges())
 					{
diff --git a/Engine/Plugins/Runtime/Metasound/Source/MetasoundFrontend/Private/MetasoundNodeRegistrationMacro.cpp b/Engine/Plugins/Runtime/Metasound/Source/MetasoundFrontend/Private/MetasoundNodeRegistrationMacro.cpp
--- a/Engine/Plugins/Runtime/Metasound/Source/MetasoundFrontend/Private/MetasoundNodeRegistrationMacro.cpp
+++ b/Engine/Plugins/Runtime/Metasound/Source/MetasoundFrontend/Private/MetasoundNodeRegistrationMacro.cpp
@@ -12,7 +12,6 @@ namespace Metasound::Frontend
 	{
 		FNodeRegistryEntryBase::FNodeRegistryEntryBase(const Metasound::FNodeClassMetadata& InMetadata, const FModuleInfo& InOwningModuleInfo)
 		: ClassMetadata(MakeShared<FNodeClassMetadata>(InMetadata))
-		, ClassInfo(FMetasoundFrontendClassMetadata::GenerateClassMetadata(InMetadata, EMetasoundFrontendClassType::External))
 		, FrontendClass(GenerateClass(InMetadata))
 #if WITH_EDITORONLY_DATA
 		, PluginName(InOwningModuleInfo.PluginName)
@@ -21,11 +20,6 @@ namespace Metasound::Frontend
 		{
 		}
 
-		const FNodeClassInfo& FNodeRegistryEntryBase::GetClassInfo() const
-		{
-			return ClassInfo;
-		}
-
 		const FMetasoundFrontendClass& FNodeRegistryEntryBase::GetFrontendClass() const
 		{
 			return FrontendClass;
diff --git a/Engine/Plugins/Runtime/Metasound/Source/MetasoundFrontend/Public/MetasoundAssetManager.h b/Engine/Plugins/Runtime/Metasound/Source/MetasoundFrontend/Public/MetasoundAssetManager.h
--- a/Engine/Plugins/Runtime/Metasound/Source/MetasoundFrontend/Public/MetasoundAssetManager.h
+++ b/Engine/Plugins/Runtime/Metasound/Source/MetasoundFrontend/Public/MetasoundAssetManager.h
@@ -321,7 +321,7 @@ PRAGMA_ENABLE_DEPRECATION_WARNINGS
 		virtual bool GetReferencedPresetHierarchy(FMetasoundAssetBase& InAsset, TArray<FMetasoundAssetBase*>& OutReferencedAssets) const = 0;
 #endif // WITH_EDITOR
 
-		// returns whether or not the given class is defined as a registered asset
+		// Returns whether or not the given ClassMetadata is both valid asset ClassMetadata and is registered.
 		virtual bool IsAssetClass(const FMetasoundFrontendClassMetadata& ClassMetadata) const = 0;
 
 		UE_DEPRECATED(5.6, "Moved to Iterate ClassInfo")
diff --git a/Engine/Plugins/Runtime/Metasound/Source/MetasoundFrontend/Public/MetasoundFrontendDocument.h b/Engine/Plugins/Runtime/Metasound/Source/MetasoundFrontend/Public/MetasoundFrontendDocument.h
--- a/Engine/Plugins/Runtime/Metasound/Source/MetasoundFrontend/Public/MetasoundFrontendDocument.h
+++ b/Engine/Plugins/Runtime/Metasound/Source/MetasoundFrontend/Public/MetasoundFrontendDocument.h
@@ -1556,7 +1556,7 @@ private:
 #endif // WITH_EDITORONLY_DATA
 
 	UPROPERTY(EditAnywhere, Category = Metasound, meta = (Bitmask, BitmaskEnum = "/Script/MetasoundFrontend.EMetasoundFrontendClassAccessFlags"))
-	uint16 AccessFlags = static_cast<uint8>(EMetasoundFrontendClassAccessFlags::Referenceable);
+	uint16 AccessFlags = static_cast<uint16>(EMetasoundFrontendClassAccessFlags::Referenceable);
 
 #if WITH_EDITORONLY_DATA
 	// If true, auto-update will manage (add and remove)
diff --git a/Engine/Plugins/Runtime/Metasound/Source/MetasoundFrontend/Public/MetasoundFrontendDocumentAccessPtr.h b/Engine/Plugins/Runtime/Metasound/Source/MetasoundFrontend/Public/MetasoundFrontendDocumentAccessPtr.h
--- a/Engine/Plugins/Runtime/Metasound/Source/MetasoundFrontend/Public/MetasoundFrontendDocumentAccessPtr.h
+++ b/Engine/Plugins/Runtime/Metasound/Source/MetasoundFrontend/Public/MetasoundFrontendDocumentAccessPtr.h
@@ -149,7 +149,10 @@ namespace Metasound
 			UE_API FClassAccessPtr GetDependencyWithID(const FGuid& InID);
 			UE_API FClassAccessPtr GetClassWithID(const FGuid& InID);
 			UE_API FClassAccessPtr GetClassWithMetadata(const FMetasoundFrontendClassMetadata& InMetadata);
+
+			UE_DEPRECATED(5.8, "FNodeClassinfo is actively being deprecated, use overload with FNodeClassRegistryKey instead")
 			UE_API FClassAccessPtr GetClassWithInfo(const FNodeClassInfo& InInfo);
+
 			UE_API FClassAccessPtr GetClassWithRegistryKey(const FNodeRegistryKey& InKey);
 
 			UE_API FConstGraphClassAccessPtr GetRootGraph() const;
@@ -157,7 +160,10 @@ namespace Metasound
 			UE_API FConstClassAccessPtr GetDependencyWithID(const FGuid& InID) const;
 			UE_API FConstClassAccessPtr GetClassWithID(const FGuid& InID) const;
 			UE_API FConstClassAccessPtr GetClassWithMetadata(const FMetasoundFrontendClassMetadata& InMetadata) const;
+
+			UE_DEPRECATED(5.8, "FNodeClassinfo is actively being deprecated, use overload with FNodeClassRegistryKey instead")
 			UE_API FConstClassAccessPtr GetClassWithInfo(const FNodeClassInfo& InInfo) const;
+
 			UE_API FConstClassAccessPtr GetClassWithRegistryKey(const FNodeRegistryKey& InKey) const;
 		};
 
@@ -174,7 +180,10 @@ namespace Metasound
 			UE_API FConstClassAccessPtr GetDependencyWithID(const FGuid& InID) const;
 			UE_API FConstClassAccessPtr GetClassWithID(const FGuid& InID) const;
 			UE_API FConstClassAccessPtr GetClassWithMetadata(const FMetasoundFrontendClassMetadata& InMetadata) const;
+
+			UE_DEPRECATED(5.8, "FNodeClassinfo is actively being deprecated, use overload with FNodeClassRegistryKey instead")
 			UE_API FConstClassAccessPtr GetClassWithInfo(const FNodeClassInfo& InInfo) const;
+
 			UE_API FConstClassAccessPtr GetClassWithRegistryKey(const FNodeRegistryKey& InKey) const;
 		};
 	}
diff --git a/Engine/Plugins/Runtime/Metasound/Source/MetasoundFrontend/Public/MetasoundFrontendNodeClassRegistry.h b/Engine/Plugins/Runtime/Metasound/Source/MetasoundFrontend/Public/MetasoundFrontendNodeClassRegistry.h
--- a/Engine/Plugins/Runtime/Metasound/Source/MetasoundFrontend/Public/MetasoundFrontendNodeClassRegistry.h
+++ b/Engine/Plugins/Runtime/Metasound/Source/MetasoundFrontend/Public/MetasoundFrontendNodeClassRegistry.h
@@ -41,13 +41,8 @@ namespace Metasound::Frontend
 	public:
 		virtual ~INodeClassRegistryEntry() = default;
 
-		/** Return FNodeClassInfo for the node class.
-			*
-			* Implementations of method should avoid any expensive operations 
-			* (e.g. loading from disk, allocating memory) as this method is called
-			* frequently when querying nodes.
-			*/
-		virtual const FNodeClassInfo& GetClassInfo() const = 0;
+		UE_DEPRECATED(5.8, "NodeClassInfo is actively being deprecated. Use FNodeClassRegistryKey instead (which can be generated from the frontend class)")
+		virtual const FNodeClassInfo& GetClassInfo() const { return FNodeClassInfo::GetInvalid(); }
 
 		/** Create a node given FDefaultNamedVertexNodeConstructorParams.
 			*
@@ -228,16 +223,22 @@ namespace Metasound::Frontend
 			}
 		}
 
+		UE_DEPRECATED(5.8, "Use constructor providing RegistryKey instead as NodeClassInfo is actively being deprecated")
 		UE_API FNodeClassRegistryTransaction(ETransactionType InType, const FNodeClassInfo& InNodeClassInfo, FTimeType InTimestamp);
 
+		UE_API FNodeClassRegistryTransaction(ETransactionType InType, FNodeClassRegistryKey InRegistryKey, FTimeType InTimestamp);
+
 		UE_API ETransactionType GetTransactionType() const;
+
+		UE_DEPRECATED(5.8, "NodeClassInfo is actively being deprecated. Use NodeRegistryKey instead (and SearchEngine to enquire about node class data via newer FMetaSoundClassInfo struct)")
 		UE_API const FNodeClassInfo& GetNodeClassInfo() const;
+
 		UE_API FNodeClassRegistryKey GetNodeRegistryKey() const;
 		UE_API FTimeType GetTimestamp() const;
 
 	private:
 		ETransactionType Type;
-		FNodeClassInfo NodeClassInfo;
+		FNodeClassRegistryKey RegistryKey;
 		FTimeType Timestamp;
 	};
 
diff --git a/Engine/Plugins/Runtime/Metasound/Source/MetasoundFrontend/Public/MetasoundFrontendRegistryKey.h b/Engine/Plugins/Runtime/Metasound/Source/MetasoundFrontend/Public/MetasoundFrontendRegistryKey.h
--- a/Engine/Plugins/Runtime/Metasound/Source/MetasoundFrontend/Public/MetasoundFrontendRegistryKey.h
+++ b/Engine/Plugins/Runtime/Metasound/Source/MetasoundFrontend/Public/MetasoundFrontendRegistryKey.h
@@ -14,7 +14,9 @@ struct FMetaSoundAssetKey;
 
 namespace Metasound::Frontend
 {
-	/** FNodeClassInfo contains a minimal set of information needed to find and query node classes. */
+	/** FNodeClassInfo contains a minimal set of information needed to find and query node classes. 
+	  * THIS CLASS IS ACTIVELY BEING DEPRECATED. USE FNodeClassRegistryKey instead.
+	  */
 	struct FNodeClassInfo
 	{
 		// ClassName of the given class
@@ -82,7 +84,10 @@ namespace Metasound::Frontend
 		UE_API FNodeClassRegistryKey(const FNodeClassMetadata& InNodeMetadata);
 		UE_API FNodeClassRegistryKey(const FMetasoundFrontendClassMetadata& InNodeMetadata);
 		UE_API FNodeClassRegistryKey(const FMetasoundFrontendGraphClass& InGraphClass);
+
+		UE_DEPRECATED(5.8, "FNodeClassInfo is actively being deprecated, use overload with NodeClassRegistryKey or otherwise")
 		UE_API FNodeClassRegistryKey(const FNodeClassInfo& InClassInfo);
+
 		UE_API FNodeClassRegistryKey(const FMetaSoundAssetKey& AssetKey);
 
 		inline friend bool operator==(const FNodeClassRegistryKey& InLHS, const FNodeClassRegistryKey& InRHS)
diff --git a/Engine/Plugins/Runtime/Metasound/Source/MetasoundFrontend/Public/MetasoundFrontendSearchEngine.h b/Engine/Plugins/Runtime/Metasound/Source/MetasoundFrontend/Public/MetasoundFrontendSearchEngine.h
--- a/Engine/Plugins/Runtime/Metasound/Source/MetasoundFrontend/Public/MetasoundFrontendSearchEngine.h
+++ b/Engine/Plugins/Runtime/Metasound/Source/MetasoundFrontend/Public/MetasoundFrontendSearchEngine.h
@@ -44,7 +44,8 @@ namespace Metasound::Frontend
 			/** Find the class info with the given ClassName & Major Version. Returns false if not found, true if found. */
 			virtual bool FindRegisteredClass(const FMetasoundFrontendClassName& InName, int32 InMajorVersion, FMetaSoundClassInfo& OutClass) = 0;
 
-			virtual bool FindClassWithHighestMinorVersion(const FMetasoundFrontendClassName& InName, int32 InMajorVersion, FMetasoundFrontendClass& OutClass) = 0;
+			UE_DEPRECATED(5.8, "Use FindRegisteredClass.  If a full class definition is required, use the resulting class info to look-up the class directly from the INodeClassRegistry.")
+			virtual bool FindClassWithHighestMinorVersion(const FMetasoundFrontendClassName& InName, int32 InMajorVersion, FMetasoundFrontendClass& OutClass);
 
 			/** Finds all registered interfaces with the given name. */
 			virtual TArray<FMetasoundFrontendVersion> FindAllRegisteredInterfacesWithName(FName InInterfaceName) = 0;
@@ -52,9 +53,6 @@ namespace Metasound::Frontend
 			/** Finds the registered interface with the highest version of the given name. Returns true if found, false if not. */
 			virtual bool FindInterfaceWithHighestVersion(FName InInterfaceName, FMetasoundFrontendInterface& OutInterface) = 0;
 
-			UE_DEPRECATED(5.3, "Use ISearchEngine::FindUClassDefaultInterfaceVersions using TopLevelAssetPath instead.")
-			virtual TArray<FMetasoundFrontendInterface> FindUClassDefaultInterfaces(FName InUClassName) = 0;
-
 			/** Returns all interfaces that are to be added to a given document when it is initialized on an object with the given class**/
 			virtual TArray<FMetasoundFrontendVersion> FindUClassDefaultInterfaceVersions(const FTopLevelAssetPath& InUClassPath) = 0;
 
@@ -63,11 +61,9 @@ namespace Metasound::Frontend
 				* Options - If set to 'All', include all versions (i.e. deprecated classes and versions of classes that are not the highest major version).
 				* bIncludeUnloadedAssets - If true, include class info for assets not loaded/registered.
 				*/
-			virtual TArray<FMetaSoundClassInfo> FindAllClasses(EResultVersion Options, bool bIncludeUnloadedAssets) = 0;
+			virtual TArray<FMetaSoundClassInfo> FindAllClasses(EResultVersion Options, bool bIncludeUnloadedAssets = false) = 0;
 
-			/** (To be deprecated) Find all FMetasoundFrontendClasses.
-				* (Optional) Include all versions (i.e. deprecated classes and versions of classes that are not the highest major version).
-				*/
+			UE_DEPRECATED(5.8, "Use FindAllClasses overload that returns ClassInfo instead")
 			virtual TArray<FMetasoundFrontendClass> FindAllClasses(bool bInIncludeAllVersions);
 
 			/** Find all class info with the given ClassName.
diff --git a/Engine/Plugins/Runtime/Metasound/Source/MetasoundFrontend/Public/MetasoundNodeRegistrationMacro.h b/Engine/Plugins/Runtime/Metasound/Source/MetasoundFrontend/Public/MetasoundNodeRegistrationMacro.h
--- a/Engine/Plugins/Runtime/Metasound/Source/MetasoundFrontend/Public/MetasoundNodeRegistrationMacro.h
+++ b/Engine/Plugins/Runtime/Metasound/Source/MetasoundFrontend/Public/MetasoundNodeRegistrationMacro.h
@@ -36,8 +36,6 @@ namespace Metasound::Frontend
 
 			virtual ~FNodeRegistryEntryBase() = default;
 
-			UE_API virtual const FNodeClassInfo& GetClassInfo() const override;
-
 			UE_API virtual const FMetasoundFrontendClass& GetFrontendClass() const override;
 
 			// Unhide CreateNode overloads which exist on INodeClassRegistryEntry
@@ -64,10 +62,9 @@ namespace Metasound::Frontend
 			UE_API TSharedRef<const FNodeClassMetadata> GetNodeClassMetadata() const;
 
 		private:
-			
 			TSharedRef<FNodeClassMetadata> ClassMetadata;
-			FNodeClassInfo ClassInfo;
 			FMetasoundFrontendClass FrontendClass;
+
 #if WITH_EDITORONLY_DATA
 			FName PluginName;
 			FName ModuleName;
