BugHawk fix-ups: Missing class info removal/UE_LOG fix-up from prior refactor.  Minor deprecation casing clean-up.

#rnx

diff --git a/Engine/Plugins/Runtime/Metasound/Source/MetasoundFrontend/Private/MetasoundFrontendNodeClassRegistry.cpp b/Engine/Plugins/Runtime/Metasound/Source/MetasoundFrontend/Private/MetasoundFrontendNodeClassRegistry.cpp
--- a/Engine/Plugins/Runtime/Metasound/Source/MetasoundFrontend/Private/MetasoundFrontendNodeClassRegistry.cpp
+++ b/Engine/Plugins/Runtime/Metasound/Source/MetasoundFrontend/Private/MetasoundFrontendNodeClassRegistry.cpp
@@ -137,7 +137,9 @@ namespace Metasound::Frontend
 				}
 				else
 				{
-					UE_LOG(LogMetaSound, Error, TEXT("Cannot create MetaSound node from class %s due to prior failure to build graph"), *ClassInfo.ClassName.ToString());
+					UE_LOG(LogMetaSound, Error,
+						TEXT("Cannot create MetaSound node from class %s due to prior failure to build graph"),
+						*FrontendClass.Metadata.GetClassName().ToString());
 					return TUniquePtr<INode>();
 				}
 			}
@@ -191,10 +193,10 @@ namespace Metasound::Frontend
 				// No node configuration supported for this node type, so only compatible if setting to invalid (null) configuration
 				return !InNodeConfiguration.IsValid();
 			}
+
 		private:
 			FMetasoundFrontendClass FrontendClass;
 			TSet<FMetasoundFrontendVersion> Interfaces;
-			FNodeClassInfo ClassInfo;
 			TSharedPtr<const IGraph> Graph;
 			FTopLevelAssetPath AssetPath;
 		};
diff --git a/Engine/Plugins/Runtime/Metasound/Source/MetasoundFrontend/Public/MetasoundFrontendDocumentAccessPtr.h b/Engine/Plugins/Runtime/Metasound/Source/MetasoundFrontend/Public/MetasoundFrontendDocumentAccessPtr.h
--- a/Engine/Plugins/Runtime/Metasound/Source/MetasoundFrontend/Public/MetasoundFrontendDocumentAccessPtr.h
+++ b/Engine/Plugins/Runtime/Metasound/Source/MetasoundFrontend/Public/MetasoundFrontendDocumentAccessPtr.h
@@ -150,7 +150,7 @@ namespace Metasound
 			UE_API FClassAccessPtr GetClassWithID(const FGuid& InID);
 			UE_API FClassAccessPtr GetClassWithMetadata(const FMetasoundFrontendClassMetadata& InMetadata);
 
-			UE_DEPRECATED(5.8, "FNodeClassinfo is actively being deprecated, use overload with FNodeClassRegistryKey instead")
+			UE_DEPRECATED(5.8, "FNodeClassInfo is actively being deprecated, use overload with FNodeClassRegistryKey instead")
 			UE_API FClassAccessPtr GetClassWithInfo(const FNodeClassInfo& InInfo);
 
 			UE_API FClassAccessPtr GetClassWithRegistryKey(const FNodeRegistryKey& InKey);
@@ -161,7 +161,7 @@ namespace Metasound
 			UE_API FConstClassAccessPtr GetClassWithID(const FGuid& InID) const;
 			UE_API FConstClassAccessPtr GetClassWithMetadata(const FMetasoundFrontendClassMetadata& InMetadata) const;
 
-			UE_DEPRECATED(5.8, "FNodeClassinfo is actively being deprecated, use overload with FNodeClassRegistryKey instead")
+			UE_DEPRECATED(5.8, "FNodeClassInfo is actively being deprecated, use overload with FNodeClassRegistryKey instead")
 			UE_API FConstClassAccessPtr GetClassWithInfo(const FNodeClassInfo& InInfo) const;
 
 			UE_API FConstClassAccessPtr GetClassWithRegistryKey(const FNodeRegistryKey& InKey) const;
@@ -181,7 +181,7 @@ namespace Metasound
 			UE_API FConstClassAccessPtr GetClassWithID(const FGuid& InID) const;
 			UE_API FConstClassAccessPtr GetClassWithMetadata(const FMetasoundFrontendClassMetadata& InMetadata) const;
 
-			UE_DEPRECATED(5.8, "FNodeClassinfo is actively being deprecated, use overload with FNodeClassRegistryKey instead")
+			UE_DEPRECATED(5.8, "FNodeClassInfo is actively being deprecated, use overload with FNodeClassRegistryKey instead")
 			UE_API FConstClassAccessPtr GetClassWithInfo(const FNodeClassInfo& InInfo) const;
 
 			UE_API FConstClassAccessPtr GetClassWithRegistryKey(const FNodeRegistryKey& InKey) const;
