Fix MetaSound patch details panel showing blank output format row

diff --git a/Engine/Plugins/Runtime/Metasound/Source/MetasoundEditor/Private/MetasoundDetailCustomization.cpp b/Engine/Plugins/Runtime/Metasound/Source/MetasoundEditor/Private/MetasoundDetailCustomization.cpp
--- a/Engine/Plugins/Runtime/Metasound/Source/MetasoundEditor/Private/MetasoundDetailCustomization.cpp
+++ b/Engine/Plugins/Runtime/Metasound/Source/MetasoundEditor/Private/MetasoundDetailCustomization.cpp
@@ -187,7 +187,7 @@ namespace Metasound::Editor
 
 				// Invalid for UMetaSounds
 				TSharedPtr<IPropertyHandle> OutputFormat = DetailLayout.GetProperty(GET_MEMBER_NAME_CHECKED(UMetaSoundSource, OutputFormat));
-				if (OutputFormat.IsValid())
+				if (OutputFormat.IsValid() && OutputFormat->IsValidHandle())
 				{
 					if (MetaSoundSource.IsValid())
 					{
