Fixed missing parameter in metasound editor for some projects

diff --git a/Engine/Plugins/Runtime/Metasound/Source/MetasoundEditor/Private/MetasoundDetailCustomization.cpp b/Engine/Plugins/Runtime/Metasound/Source/MetasoundEditor/Private/MetasoundDetailCustomization.cpp
--- a/Engine/Plugins/Runtime/Metasound/Source/MetasoundEditor/Private/MetasoundDetailCustomization.cpp
+++ b/Engine/Plugins/Runtime/Metasound/Source/MetasoundEditor/Private/MetasoundDetailCustomization.cpp
@@ -228,8 +228,9 @@ namespace Metasound::Editor
 					.ValueContent()
 					[
 						OutputFormatValueWidget
-					];
-
+					]
+					.RowTag("OutputFormat");
+					
 					OutputFormat->MarkHiddenByCustomization();
 				}
 
