- Optimize MetaSound interface search & registry
- Remove interface from SearchEngine FFrontendQueryValue to avoid memory bloat
- Interface is only used by a single editor query and can be retrieved directly from the interface registry instead avoiding the required duplicate data and extending entry size for other partitions
- Saves 144 bytes per partition entry in editor, 280 bytes per partition entry at runtime
- Gets interface registry one step closer to being able to be dynamic, where interfaces can be defined at runtime and unregistered as well as registered

#rnx

diff --git a/Engine/Plugins/Experimental/TechAudioTools/Source/TechAudioToolsMetaSound/Private/ViewModels/MetaSoundViewModelConversionFunctions.cpp b/Engine/Plugins/Experimental/TechAudioTools/Source/TechAudioToolsMetaSound/Private/ViewModels/MetaSoundViewModelConversionFunctions.cpp
--- a/Engine/Plugins/Experimental/TechAudioTools/Source/TechAudioToolsMetaSound/Private/ViewModels/MetaSoundViewModelConversionFunctions.cpp
+++ b/Engine/Plugins/Experimental/TechAudioTools/Source/TechAudioToolsMetaSound/Private/ViewModels/MetaSoundViewModelConversionFunctions.cpp
@@ -1,4 +1,4 @@
-﻿// Copyright Epic Games, Inc. All Rights Reserved.
+// Copyright Epic Games, Inc. All Rights Reserved.
 
 #include "ViewModels/MetaSoundViewModelConversionFunctions.h"
 
@@ -46,8 +46,8 @@ bool UMetaSoundViewModelConversionFunctions::IsInterfaceMember(const FName& Memb
 	FName ParamName;
 	Audio::FParameterPath::SplitName(MemberName, InterfaceNamespace, ParamName);
 
-	FMetasoundFrontendInterface FoundInterface;
-	const bool bIsValidNamespace = Metasound::Frontend::ISearchEngine::Get().FindInterfaceWithHighestVersion(InterfaceNamespace, FoundInterface);
+	FMetasoundFrontendVersion FoundInterface;
+	const bool bIsValidNamespace = Metasound::Frontend::ISearchEngine::Get().FindHighestInterfaceVersion(InterfaceNamespace, FoundInterface);
 	const bool bResult = !InterfaceNamespace.IsNone() && bIsValidNamespace;
 
 	return bInvert ? !bResult : bResult;
diff --git a/Engine/Plugins/Runtime/Metasound/Source/MetasoundEditor/Private/MetasoundDetailCustomization.cpp b/Engine/Plugins/Runtime/Metasound/Source/MetasoundEditor/Private/MetasoundDetailCustomization.cpp
--- a/Engine/Plugins/Runtime/Metasound/Source/MetasoundEditor/Private/MetasoundDetailCustomization.cpp
+++ b/Engine/Plugins/Runtime/Metasound/Source/MetasoundEditor/Private/MetasoundDetailCustomization.cpp
@@ -817,22 +817,32 @@ namespace Metasound::Editor
 
 				if (Builder.IsValid() && InSelectInfo != ESelectInfo::OnNavigation)
 				{
-					FMetaSoundFrontendDocumentBuilder& DocBuilder = Builder->GetBuilder();
-					UObject& MetaSound = DocBuilder.CastDocumentObjectChecked<UObject>();
-					FMetasoundFrontendInterface InterfaceToAdd;
-					const FName InterfaceName { *NameToAdd.Get() };
-					if (ensure(ISearchEngine::Get().FindInterfaceWithHighestVersion(InterfaceName, InterfaceToAdd)))
+					FMetasoundFrontendVersion InterfaceVersion;
+					if (ensure(ISearchEngine::Get().FindHighestInterfaceVersion(FName{ *NameToAdd.Get() }, InterfaceVersion)))
 					{
-						const FScopedTransaction Transaction(FText::Format(LOCTEXT("AddInterfaceTransactionFormat", "Add MetaSound Interface '{0}'"), FText::FromString(InterfaceToAdd.Metadata.Version.ToString())));
-						MetaSound.Modify();
-						FModifyInterfaceOptions Options({ }, { InterfaceToAdd });
-						Options.bSetDefaultNodeLocations = false; // Don't automatically add nodes to ed graph
-						DocBuilder.ModifyInterfaces(MoveTemp(Options));
-					}
+						FMetaSoundFrontendDocumentBuilder& DocBuilder = Builder->GetBuilder();
+						UObject& MetaSound = DocBuilder.CastDocumentObjectChecked<UObject>();
 
-					UpdateInterfaceNames();
-					InterfaceComboBox->RefreshOptions();
-					FGraphBuilder::RegisterGraphWithFrontend(MetaSound);
+						FMetasoundFrontendInterface InterfaceToAdd;
+						if (ensure(IInterfaceRegistry::Get().FindInterface(InterfaceVersion, InterfaceToAdd)))
+						{
+							const FScopedTransaction Transaction(FText::Format(
+								LOCTEXT("AddInterfaceTransactionFormat", "Add MetaSound Interface '{0}'"),
+								FText::FromString(InterfaceToAdd.Metadata.Version.ToString())));
+
+							MetaSound.Modify();
+							FModifyInterfaceOptions Options({ }, { InterfaceToAdd });
+							Options.bSetDefaultNodeLocations = false; // Don't automatically add nodes to ed graph
+							DocBuilder.ModifyInterfaces(MoveTemp(Options));
+						}
+
+						UpdateInterfaceNames();
+						InterfaceComboBox->RefreshOptions();
+
+						FMetaSoundAssetRegistrationOptions Options = FGraphBuilder::GetDefaultRegistrationOptions();
+						Options.bIgnoreIfLiveAuditioning = true;
+						FGraphBuilder::RegisterGraphWithFrontend(MetaSound, MoveTemp(Options));
+					}
 				}
 			})
 			.Content()
@@ -872,13 +882,20 @@ namespace Metasound::Editor
 						return;
 					}
 
+					FMetasoundFrontendVersion OutVersion;
 					TArray<FMetasoundFrontendInterface> InheritedInterfaces;
-					Algo::Transform(ImplementedInterfaceNames, InheritedInterfaces, [](const FName& Name)
-					{
-						FMetasoundFrontendInterface Interface;
-						ISearchEngine::Get().FindInterfaceWithHighestVersion(Name, Interface);
-						return Interface;
-					});
+					Algo::TransformIf(ImplementedInterfaceNames, InheritedInterfaces,
+						[&](const FName& Name)
+						{
+							return ISearchEngine::Get().FindHighestInterfaceVersion(Name, OutVersion);
+						},
+						[&](const FName& Name)
+						{
+							FMetasoundFrontendInterface Interface;
+							const bool bFoundInterface = IInterfaceRegistry::Get().FindInterface(OutVersion, Interface);
+							checkf(bFoundInterface, TEXT("Interface not found in registry, indicating SearchEngine has entry and unregistration was not evaluated."));
+							return Interface;
+						});
 
 					UObject& MetaSound = DocBuilder.CastDocumentObjectChecked<UObject>();
 					{
@@ -909,13 +926,13 @@ namespace Metasound::Editor
 		{
 			using namespace Frontend;
 
-			FMetasoundFrontendInterface InterfaceEntry;
-			if (!ensure(ISearchEngine::Get().FindInterfaceWithHighestVersion(InInterfaceName, InterfaceEntry)))
+			FMetasoundFrontendVersion InterfaceVersion;
+			if (!ensure(ISearchEngine::Get().FindHighestInterfaceVersion(InInterfaceName, InterfaceVersion)))
 			{
 				return SNullWidget::NullWidget;
 			}
 
-			TSharedRef<SWidget> RemoveButtonWidget = PropertyCustomizationHelpers::MakeDeleteButton(FSimpleDelegate::CreateLambda([this, InInterfaceName, InterfaceEntry]()
+			TSharedRef<SWidget> RemoveButtonWidget = PropertyCustomizationHelpers::MakeDeleteButton(FSimpleDelegate::CreateLambda([this, InInterfaceName, InterfaceVersion]()
 			{
 				using namespace Frontend;
 
@@ -932,11 +949,15 @@ namespace Metasound::Editor
 
 				UObject& MetaSound = DocBuilder.CastDocumentObjectChecked<UObject>();
 				{
-					const FScopedTransaction Transaction(FText::Format(LOCTEXT("RemoveInterfaceTransactionFormat", "Remove MetaSound Interface '{0}'"), FText::FromString(InterfaceEntry.Metadata.Version.ToString())));
-					MetaSound.Modify();
-					FModifyInterfaceOptions Options({ InterfaceEntry }, {});
-					Options.bSetDefaultNodeLocations = false; // Don't automatically add nodes to ed graph
-					DocBuilder.ModifyInterfaces(MoveTemp(Options));
+					FMetasoundFrontendInterface InterfaceToAdd;
+					if (ensure(IInterfaceRegistry::Get().FindInterface(InterfaceVersion, InterfaceToAdd)))
+					{
+						const FScopedTransaction Transaction(FText::Format(LOCTEXT("RemoveInterfaceTransactionFormat", "Remove MetaSound Interface '{0}'"), FText::FromString(InterfaceVersion.ToString())));
+						MetaSound.Modify();
+						FModifyInterfaceOptions Options({ MoveTemp(InterfaceToAdd) }, { });
+						Options.bSetDefaultNodeLocations = false; // Don't automatically add nodes to ed graph
+						DocBuilder.ModifyInterfaces(MoveTemp(Options));
+					}
 				}
 
 				UpdateInterfaceNames();
@@ -953,7 +974,7 @@ namespace Metasound::Editor
 				.AutoWidth()
 				[
 					SNew(STextBlock)
-					.Text(FText::FromName(InterfaceEntry.Metadata.Version.Name))
+					.Text(FText::FromName(InterfaceVersion.Name))
 				]
 			+ SHorizontalBox::Slot()
 				.Padding(2.0f)
@@ -992,11 +1013,11 @@ namespace Metasound::Editor
 			{
 				using namespace Frontend;
 
-				const FInterfaceRegistryKey Key = GetInterfaceRegistryKey(Version);
-				if (const IInterfaceRegistryEntry* Entry = IInterfaceRegistry::Get().FindInterfaceRegistryEntry(Key))
+				FMetasoundFrontendInterface Interface;
+				if (IInterfaceRegistry::Get().FindInterface(Version, Interface))
 				{
 					auto FindClassOptionsPredicate = [&ClassName](const FMetasoundFrontendInterfaceUClassOptions& Options) { return Options.ClassPath == ClassName; };
-					if (const FMetasoundFrontendInterfaceUClassOptions* Options = Entry->GetInterface().Metadata.UClassOptions.FindByPredicate(FindClassOptionsPredicate))
+					if (const FMetasoundFrontendInterfaceUClassOptions* Options = Interface.Metadata.UClassOptions.FindByPredicate(FindClassOptionsPredicate))
 					{
 						return Options->bIsModifiable;
 					}
@@ -1011,14 +1032,14 @@ namespace Metasound::Editor
 			const TSet<FMetasoundFrontendVersion>& InheritedInterfaces = Builder->GetBuilder().GetConstDocumentChecked().Interfaces;
 			Algo::TransformIf(InheritedInterfaces, ImplementedInterfaceNames, CanAddOrRemoveInterface, GetVersionName);
 
-			TArray<FMetasoundFrontendInterface> Interfaces = Frontend::ISearchEngine::Get().FindAllInterfaces();
-			for (const FMetasoundFrontendInterface& Interface : Interfaces)
+			TArray<FMetasoundFrontendVersion> InterfaceVersions = Frontend::ISearchEngine::Get().FindAllInterfaceVersions();
+			for (const FMetasoundFrontendVersion& Version : InterfaceVersions)
 			{
-				if (!ImplementedInterfaceNames.Contains(Interface.Metadata.Version.Name))
+				if (!ImplementedInterfaceNames.Contains(Version.Name))
 				{
-					if (CanAddOrRemoveInterface(Interface.Metadata.Version))
+					if (CanAddOrRemoveInterface(Version))
 					{
-						FString Name = Interface.Metadata.Version.Name.ToString();
+						FString Name = Version.Name.ToString();
 						AddableInterfaceNames.Add(MakeShared<FString>(MoveTemp(Name)));
 					}
 				}
diff --git a/Engine/Plugins/Runtime/Metasound/Source/MetasoundEditor/Private/MetasoundEditorGraph.cpp b/Engine/Plugins/Runtime/Metasound/Source/MetasoundEditor/Private/MetasoundEditorGraph.cpp
--- a/Engine/Plugins/Runtime/Metasound/Source/MetasoundEditor/Private/MetasoundEditorGraph.cpp
+++ b/Engine/Plugins/Runtime/Metasound/Source/MetasoundEditor/Private/MetasoundEditorGraph.cpp
@@ -474,14 +474,19 @@ bool UMetasoundEditorGraphVertex::NameContainsInterfaceNamespace(FMetasoundFront
 	FName ParamName;
 	Audio::FParameterPath::SplitName(MemberName, InterfaceNamespace, ParamName);
 
-	FMetasoundFrontendInterface FoundInterface;
-	if (!InterfaceNamespace.IsNone() && ISearchEngine::Get().FindInterfaceWithHighestVersion(InterfaceNamespace, FoundInterface))
+	if (!InterfaceNamespace.IsNone())
 	{
-		if (OutInterface)
+		FMetasoundFrontendVersion InterfaceVersion;
+		const bool bFoundVersion = ISearchEngine::Get().FindHighestInterfaceVersion(InterfaceNamespace, InterfaceVersion);
+		if (bFoundVersion)
 		{
-			*OutInterface = MoveTemp(FoundInterface);
+			if (OutInterface)
+			{
+				const bool bIsValid = IInterfaceRegistry::Get().FindInterface(InterfaceVersion, *OutInterface);
+				checkf(bIsValid, TEXT("Version found in search engine must be registered with interface registry"));
+			}
+			return true;
 		}
-		return true;
 	}
 
 	if (OutInterface)
diff --git a/Engine/Plugins/Runtime/Metasound/Source/MetasoundEditor/Private/MetasoundNodeDetailCustomization.cpp b/Engine/Plugins/Runtime/Metasound/Source/MetasoundEditor/Private/MetasoundNodeDetailCustomization.cpp
--- a/Engine/Plugins/Runtime/Metasound/Source/MetasoundEditor/Private/MetasoundNodeDetailCustomization.cpp
+++ b/Engine/Plugins/Runtime/Metasound/Source/MetasoundEditor/Private/MetasoundNodeDetailCustomization.cpp
@@ -590,9 +590,9 @@ namespace Metasound
 						{
 							// Don't display trigger simulation widget if its a trigger
 							// provided by an interface that does not support transmission.
-							const FInterfaceRegistryKey Key = GetInterfaceRegistryKey(Input->GetInterfaceVersion());
-							const IInterfaceRegistryEntry* Entry = IInterfaceRegistry::Get().FindInterfaceRegistryEntry(Key);
-							if (!Entry || Entry->GetRouterName() == Audio::IParameterTransmitter::RouterName)
+							FName RouterName;
+							const bool bFound = IInterfaceRegistry::Get().FindInterfaceRouter(Input->GetInterfaceVersion(), RouterName);
+							if (!bFound || RouterName == Audio::IParameterTransmitter::RouterName)
 							{
 								EnablementAttribute = true;
 								return SMetaSoundGraphNode::CreateTriggerSimulationWidget(*Literal, MoveTemp(VisibilityAttribute), MoveTemp(EnablementAttribute));
diff --git a/Engine/Plugins/Runtime/Metasound/Source/MetasoundEditor/Private/SMetasoundGraphNode.cpp b/Engine/Plugins/Runtime/Metasound/Source/MetasoundEditor/Private/SMetasoundGraphNode.cpp
--- a/Engine/Plugins/Runtime/Metasound/Source/MetasoundEditor/Private/SMetasoundGraphNode.cpp
+++ b/Engine/Plugins/Runtime/Metasound/Source/MetasoundEditor/Private/SMetasoundGraphNode.cpp
@@ -182,9 +182,9 @@ namespace Metasound
 
 						// Don't enable trigger simulation widget if its a trigger provided by an interface
 						// that does not support transmission.
-						const FInterfaceRegistryKey Key = GetInterfaceRegistryKey(Vertex->GetInterfaceVersion());
-						const IInterfaceRegistryEntry* Entry = IInterfaceRegistry::Get().FindInterfaceRegistryEntry(Key);
-						if (Entry && Entry->GetRouterName() != Audio::IParameterTransmitter::RouterName)
+						FName RouterName;
+						if (IInterfaceRegistry::Get().FindInterfaceRouter(Vertex->GetInterfaceVersion(), RouterName)
+							&& RouterName != Audio::IParameterTransmitter::RouterName)
 						{
 							return EVisibility::Hidden;
 						}
diff --git a/Engine/Plugins/Runtime/Metasound/Source/MetasoundEngine/Private/Interfaces/MetasoundInterface.cpp b/Engine/Plugins/Runtime/Metasound/Source/MetasoundEngine/Private/Interfaces/MetasoundInterface.cpp
--- a/Engine/Plugins/Runtime/Metasound/Source/MetasoundEngine/Private/Interfaces/MetasoundInterface.cpp
+++ b/Engine/Plugins/Runtime/Metasound/Source/MetasoundEngine/Private/Interfaces/MetasoundInterface.cpp
@@ -26,8 +26,35 @@
 
 namespace Metasound::Engine
 {
+	namespace InterfacePrivate
+	{
+		struct FDefaultTransform : public Frontend::IBuilderVersionTransform
+		{
+			const FMetasoundFrontendVersion Version;
+
+			FDefaultTransform(FMetasoundFrontendVersion InVersion)
+				: Version(InVersion)
+			{
+			}
+
+			virtual ~FDefaultTransform() = default;
+
+			virtual const FMetasoundFrontendVersion& GetVersion() const override
+			{
+				return Version;
+			}
+
+			virtual bool Transform(FMetaSoundFrontendDocumentBuilder& InOutBuilder) const override
+			{
+				InOutBuilder.AddInterface(Version.Name);
+				return true;
+			}
+		};
+	}
+
 	FInterfaceRegistryEntry::FInterfaceRegistryEntry(FMetasoundFrontendInterface&& InInterface, FName InRouterName, bool bInIsDeprecated)
 		: Interface(MoveTemp(InInterface))
+		, UpdateTransform(MakeShareable(new InterfacePrivate::FDefaultTransform(Interface.Metadata.Version)))
 		, RouterName(InRouterName)
 		, bIsDeprecated(bInIsDeprecated)
 	{
@@ -35,17 +62,53 @@ namespace Metasound::Engine
 
 	FInterfaceRegistryEntry::FInterfaceRegistryEntry(const FMetasoundFrontendInterface& InInterface, FName InRouterName, bool bInIsDeprecated)
 		: Interface(InInterface)
+		, UpdateTransform(MakeShareable(new InterfacePrivate::FDefaultTransform(Interface.Metadata.Version)))
+		, RouterName(InRouterName)
+		, bIsDeprecated(bInIsDeprecated)
+	{
+	}
+
+	FInterfaceRegistryEntry::FInterfaceRegistryEntry(FMetasoundFrontendInterface InInterface, TUniquePtr<Frontend::IBuilderVersionTransform>&& InUpdateTransform, FName InRouterName, bool bInIsDeprecated)
+		: Interface(MoveTemp(InInterface))
+		, UpdateTransform(MakeShareable(InUpdateTransform.Release()))
 		, RouterName(InRouterName)
 		, bIsDeprecated(bInIsDeprecated)
 	{
+		check(UpdateTransform.IsValid());
 	}
 
 	FInterfaceRegistryEntry::FInterfaceRegistryEntry(const FMetasoundFrontendInterface& InInterface, TUniquePtr<Frontend::IDocumentTransform>&& InUpdateTransform, FName InRouterName, bool bInIsDeprecated)
 		: Interface(InInterface)
-		, UpdateTransform(MoveTemp(InUpdateTransform))
 		, RouterName(InRouterName)
 		, bIsDeprecated(bInIsDeprecated)
 	{
+		using FDocTransform = TUniquePtr<Frontend::IDocumentTransform>;
+
+		struct FDefaultBuilderTransform : public Frontend::IBuilderVersionTransform
+		{
+			FDocTransform DocTransform;
+			FMetasoundFrontendVersion Version;
+
+			FDefaultBuilderTransform(const FMetasoundFrontendVersion& UpdateVersion, FDocTransform&& InTransform)
+				: DocTransform(MoveTemp(InTransform))
+				, Version(UpdateVersion)
+			{
+			}
+
+			virtual ~FDefaultBuilderTransform() = default;
+
+			virtual const FMetasoundFrontendVersion& GetVersion() const
+			{
+				return Version;
+			}
+
+			virtual bool Transform(FMetaSoundFrontendDocumentBuilder& InOutBuilder) const override
+			{
+				DocTransform->Transform(InOutBuilder.GetMetasoundAsset().GetDocumentHandle());
+				return true;
+			}
+		};
+		UpdateTransform = MakeShareable(new FDefaultBuilderTransform(Interface.Metadata.Version, MoveTemp(InUpdateTransform)));
 	}
 
 	FName FInterfaceRegistryEntry::GetRouterName() const
@@ -64,23 +127,18 @@ namespace Metasound::Engine
 	}
 
 	bool FInterfaceRegistryEntry::UpdateRootGraphInterface(Frontend::FDocumentHandle InDocument) const
-	{
-		if (UpdateTransform.IsValid())
-		{
-			return UpdateTransform->Transform(InDocument);
-		}
-		
+	{		
 		return false;
 	}
 
 	bool FInterfaceRegistryEntry::UpdateRootGraphInterface(FMetaSoundFrontendDocumentBuilder& InDocumentBuilder) const
 	{
-		if (UpdateTransform.IsValid())
-		{
-			return UpdateRootGraphInterface(InDocumentBuilder.GetMetasoundAsset().GetDocumentHandle());
-		}
+		return false;
+	}
 
-		return InDocumentBuilder.AddInterface(GetInterface().Metadata.Version.Name);
+	TSharedRef<Frontend::IBuilderVersionTransform> FInterfaceRegistryEntry::GetUpdateTransform() const
+	{
+		return UpdateTransform.ToSharedRef();
 	}
 
 	void RegisterAudioFormatInterfaces()
diff --git a/Engine/Plugins/Runtime/Metasound/Source/MetasoundEngine/Private/Interfaces/MetasoundInterface.h b/Engine/Plugins/Runtime/Metasound/Source/MetasoundEngine/Private/Interfaces/MetasoundInterface.h
--- a/Engine/Plugins/Runtime/Metasound/Source/MetasoundEngine/Private/Interfaces/MetasoundInterface.h
+++ b/Engine/Plugins/Runtime/Metasound/Source/MetasoundEngine/Private/Interfaces/MetasoundInterface.h
@@ -16,6 +16,9 @@ namespace Metasound::Engine
 	public:
 		FInterfaceRegistryEntry(FMetasoundFrontendInterface&& InInterface, FName InRouterName = IDataReference::RouterName, bool bInIsDeprecated = false);
 		FInterfaceRegistryEntry(const FMetasoundFrontendInterface& InInterface, FName InRouterName = IDataReference::RouterName, bool bInIsDeprecated = false);
+		FInterfaceRegistryEntry(FMetasoundFrontendInterface InInterface, TUniquePtr<Frontend::IBuilderVersionTransform>&& InUpdateTransform, FName InRouterName, bool bInIsDeprecated = false);
+
+		// Ctor is here for back compat.  Future registry entries that provide update transforms should use overload providing builder version transform
 		FInterfaceRegistryEntry(const FMetasoundFrontendInterface& InInterface, TUniquePtr<Frontend::IDocumentTransform>&& InUpdateTransform, FName InRouterName = IDataReference::RouterName, bool bInIsDeprecated = false);
 
 		virtual FName GetRouterName() const override;
@@ -23,10 +26,11 @@ namespace Metasound::Engine
 		virtual bool IsDeprecated() const override;
 		virtual bool UpdateRootGraphInterface(Frontend::FDocumentHandle InDocument) const override;
 		virtual bool UpdateRootGraphInterface(FMetaSoundFrontendDocumentBuilder& InDocumentBuilder) const override;
+		virtual TSharedRef<Frontend::IBuilderVersionTransform> GetUpdateTransform() const override;
 
 	private:
 		FMetasoundFrontendInterface Interface;
-		TUniquePtr<Frontend::IDocumentTransform> UpdateTransform;
+		TSharedPtr<Frontend::IBuilderVersionTransform> UpdateTransform;
 		FName RouterName;
 		bool bIsDeprecated = false;
 	};
diff --git a/Engine/Plugins/Runtime/Metasound/Source/MetasoundEngine/Private/MetasoundBuilderSubsystem.cpp b/Engine/Plugins/Runtime/Metasound/Source/MetasoundEngine/Private/MetasoundBuilderSubsystem.cpp
--- a/Engine/Plugins/Runtime/Metasound/Source/MetasoundEngine/Private/MetasoundBuilderSubsystem.cpp
+++ b/Engine/Plugins/Runtime/Metasound/Source/MetasoundEngine/Private/MetasoundBuilderSubsystem.cpp
@@ -1424,8 +1424,8 @@ bool UMetaSoundBuilderSubsystem::IsInterfaceRegistered(FName InInterfaceName) co
 {
 	using namespace Metasound::Frontend;
 
-	FMetasoundFrontendInterface Interface;
-	return ISearchEngine::Get().FindInterfaceWithHighestVersion(InInterfaceName, Interface);
+	FMetasoundFrontendVersion InterfaceVersion;
+	return ISearchEngine::Get().FindHighestInterfaceVersion(InInterfaceName, InterfaceVersion);
 }
 
 void UMetaSoundBuilderSubsystem::RegisterBuilder(FName BuilderName, UMetaSoundBuilderBase* Builder)
diff --git a/Engine/Plugins/Runtime/Metasound/Source/MetasoundEngine/Private/MetasoundSource.cpp b/Engine/Plugins/Runtime/Metasound/Source/MetasoundEngine/Private/MetasoundSource.cpp
--- a/Engine/Plugins/Runtime/Metasound/Source/MetasoundEngine/Private/MetasoundSource.cpp
+++ b/Engine/Plugins/Runtime/Metasound/Source/MetasoundEngine/Private/MetasoundSource.cpp
@@ -1830,22 +1830,23 @@ Metasound::TSortedVertexNameMap<UMetaSoundSource::FRuntimeInput> UMetaSoundSourc
 
 	auto GetInputName = [](const FMetasoundFrontendClassInput& InInput) { return InInput.Name; };
 
-	const IDataTypeRegistry& Registry = IDataTypeRegistry::Get();
-	const FMetasoundFrontendDocument& Doc = GetConstDocument();
+	const IDataTypeRegistry& DataTypeRegistry = IDataTypeRegistry::Get();
+	const IInterfaceRegistry& InterfaceRegistry = IInterfaceRegistry::Get();
 
-	TArray<const IInterfaceRegistryEntry*> Interfaces;
-	FMetaSoundFrontendDocumentBuilder::FindDeclaredInterfaces(Doc, Interfaces);
+	const FMetasoundFrontendDocument& Doc = GetConstDocument();
 
 	// Inputs which are controlled by an interface are private unless
 	// their router name is `Audio::IParameterTransmitter::RouterName`
 	TSet<FVertexName> PrivateInputs;
-	for (const IInterfaceRegistryEntry* InterfaceEntry : Interfaces)
+	for (const FMetasoundFrontendVersion& InterfaceVersion : Doc.Interfaces)
 	{
-		if (InterfaceEntry)
+		FName RouterName;
+		if (IInterfaceRegistry::Get().FindInterfaceRouter(InterfaceVersion, RouterName))
 		{
-			if (InterfaceEntry->GetRouterName() != Audio::IParameterTransmitter::RouterName)
+			if (RouterName != Audio::IParameterTransmitter::RouterName)
 			{
-				const FMetasoundFrontendInterface& Interface = InterfaceEntry->GetInterface();
+				FMetasoundFrontendInterface Interface;
+				InterfaceRegistry.FindInterface(InterfaceVersion, Interface);
 				Algo::Transform(Interface.Inputs, PrivateInputs, GetInputName);
 			}
 		}
@@ -1857,14 +1858,14 @@ Metasound::TSortedVertexNameMap<UMetaSoundSource::FRuntimeInput> UMetaSoundSourc
 	{
 		if (!PrivateInputs.Contains(Input.Name))
 		{
-			PublicInputs.Add(Input.Name, CreateRuntimeInput(Registry, Input, bCreateUObjectProxies, InPageOrder));
+			PublicInputs.Add(Input.Name, CreateRuntimeInput(DataTypeRegistry, Input, bCreateUObjectProxies, InPageOrder));
 		}
 	}
 
-	// Add the parameter pack input that ALL Metasounds have
+	// Add the parameter pack input that ALL MetaSounds have
 	FMetasoundFrontendClassInput ParameterPackInput = UMetasoundParameterPack::GetClassInput();
 	const FMetasoundFrontendLiteral& DefaultLiteral = ParameterPackInput.FindConstDefaultChecked(Frontend::DefaultPageID);
-	FAudioParameter ParameterPackDefaultParameter = SourcePrivate::MakeAudioParameter(Registry, ParameterPackInput.Name, ParameterPackInput.TypeName, DefaultLiteral, bCreateUObjectProxies) ;
+	FAudioParameter ParameterPackDefaultParameter = SourcePrivate::MakeAudioParameter(DataTypeRegistry, ParameterPackInput.Name, ParameterPackInput.TypeName, DefaultLiteral, bCreateUObjectProxies) ;
 	PublicInputs.Add(ParameterPackInput.Name, FRuntimeInput{ParameterPackInput.Name, ParameterPackInput.TypeName, ParameterPackInput.AccessType, ParameterPackDefaultParameter, true /* bIsTransmittable */});
 	
 	return PublicInputs;
diff --git a/Engine/Plugins/Runtime/Metasound/Source/MetasoundFrontend/Private/Interfaces/MetasoundFrontendInterfaceBindingRegistry.cpp b/Engine/Plugins/Runtime/Metasound/Source/MetasoundFrontend/Private/Interfaces/MetasoundFrontendInterfaceBindingRegistry.cpp
--- a/Engine/Plugins/Runtime/Metasound/Source/MetasoundFrontend/Private/Interfaces/MetasoundFrontendInterfaceBindingRegistry.cpp
+++ b/Engine/Plugins/Runtime/Metasound/Source/MetasoundFrontend/Private/Interfaces/MetasoundFrontendInterfaceBindingRegistry.cpp
@@ -35,8 +35,7 @@ namespace Metasound::Frontend
 	bool FInterfaceBindingRegistry::FindInterfaceBindingEntries(const FMetasoundFrontendVersion& InInputInterfaceVersion, TArray<const FInterfaceBindingRegistryEntry*>& OutEntries) const
 	{
 		TArray<const TUniquePtr<FInterfaceBindingRegistryEntry>*> FoundEntries;
-		const FInterfaceRegistryKey Key = GetInterfaceRegistryKey(InInputInterfaceVersion);
-		Entries.MultiFindPointer(Key, FoundEntries);
+		Entries.MultiFindPointer(InInputInterfaceVersion, FoundEntries);
 
 		OutEntries.Reset();
 		Algo::Transform(FoundEntries, OutEntries, [](const TUniquePtr<FInterfaceBindingRegistryEntry>* FoundEntry)
@@ -50,15 +49,13 @@ namespace Metasound::Frontend
 
 	void FInterfaceBindingRegistry::RegisterInterfaceBinding(const FMetasoundFrontendVersion& InInputInterfaceVersion, TUniquePtr<FInterfaceBindingRegistryEntry>&& InEntry)
 	{
-		FInterfaceRegistryKey Key = GetInterfaceRegistryKey(InInputInterfaceVersion);
-		Entries.Add(MoveTemp(Key), MoveTemp(InEntry));
+		Entries.Add(InInputInterfaceVersion, MoveTemp(InEntry));
 	}
 
 	bool FInterfaceBindingRegistry::UnregisterInterfaceBinding(const FMetasoundFrontendVersion& InInputInterfaceVersion, const FMetasoundFrontendVersion& InOutputInterfaceVersion)
 	{
 		TArray<const TUniquePtr<FInterfaceBindingRegistryEntry>*> FoundEntries;
-		const FInterfaceRegistryKey Key = GetInterfaceRegistryKey(InInputInterfaceVersion);
-		Entries.MultiFindPointer(Key, FoundEntries);
+		Entries.MultiFindPointer(InInputInterfaceVersion, FoundEntries);
 
 		for (const TUniquePtr<FInterfaceBindingRegistryEntry>* EntryPtr : FoundEntries)
 		{
@@ -66,7 +63,7 @@ namespace Metasound::Frontend
 			check(Entry);
 			if (Entry->GetOutputInterfaceVersion() == InOutputInterfaceVersion)
 			{
-				if (Entries.RemoveSingle(Key, *EntryPtr) > 0)
+				if (Entries.RemoveSingle(InInputInterfaceVersion, *EntryPtr) > 0)
 				{
 					return true;
 				}
@@ -78,8 +75,7 @@ namespace Metasound::Frontend
 
 	bool FInterfaceBindingRegistry::UnregisterAllInterfaceBindings(const FMetasoundFrontendVersion& InInputInterfaceVersion)
 	{
-		FInterfaceRegistryKey Key = GetInterfaceRegistryKey(InInputInterfaceVersion);
-		return Entries.Remove(Key) > 0;
+		return Entries.Remove(InInputInterfaceVersion) > 0;
 	}
 
 	IInterfaceBindingRegistry& IInterfaceBindingRegistry::Get()
diff --git a/Engine/Plugins/Runtime/Metasound/Source/MetasoundFrontend/Private/Interfaces/MetasoundFrontendInterfaceBindingRegistryPrivate.h b/Engine/Plugins/Runtime/Metasound/Source/MetasoundFrontend/Private/Interfaces/MetasoundFrontendInterfaceBindingRegistryPrivate.h
--- a/Engine/Plugins/Runtime/Metasound/Source/MetasoundFrontend/Private/Interfaces/MetasoundFrontendInterfaceBindingRegistryPrivate.h
+++ b/Engine/Plugins/Runtime/Metasound/Source/MetasoundFrontend/Private/Interfaces/MetasoundFrontendInterfaceBindingRegistryPrivate.h
@@ -26,7 +26,8 @@ namespace Metasound::Frontend
 
 		virtual bool UnregisterAllInterfaceBindings(const FMetasoundFrontendVersion& InInputInterfaceVersion) override;
 
-		// MultiMap of interface bindings keyed by the shared input interface registry key for an array of connection-priority sorted bindings
-		TMultiMap<FInterfaceRegistryKey, TUniquePtr<FInterfaceBindingRegistryEntry>> Entries;
+	private:
+		// MultiMap of interface bindings keyed by the shared input interface version for an array of connection-priority sorted bindings
+		TMultiMap<FMetasoundFrontendVersion, TUniquePtr<FInterfaceBindingRegistryEntry>> Entries;
 	};
 } // namespace Metasound::Frontend
diff --git a/Engine/Plugins/Runtime/Metasound/Source/MetasoundFrontend/Private/Interfaces/MetasoundFrontendInterfaceRegistry.cpp b/Engine/Plugins/Runtime/Metasound/Source/MetasoundFrontend/Private/Interfaces/MetasoundFrontendInterfaceRegistry.cpp
--- a/Engine/Plugins/Runtime/Metasound/Source/MetasoundFrontend/Private/Interfaces/MetasoundFrontendInterfaceRegistry.cpp
+++ b/Engine/Plugins/Runtime/Metasound/Source/MetasoundFrontend/Private/Interfaces/MetasoundFrontendInterfaceRegistry.cpp
@@ -7,6 +7,7 @@
 #include "MetasoundFrontendInterfaceRegistryPrivate.h"
 #include "MetasoundFrontendRegistryTransaction.h"
 #include "MetasoundTrace.h"
+#include "Misc/ScopeLock.h"
 
 
 #ifndef UE_METASOUND_ENABLE_INTERFACE_VALIDATION
@@ -27,76 +28,136 @@ namespace Metasound::Frontend
 		FInterfaceRegistryTransaction::FTimeType TransactionTime = FPlatformTime::Cycles64();
 		if (InEntry.IsValid())
 		{
-			const FInterfaceRegistryKey Key = GetInterfaceRegistryKey(InEntry->GetInterface());
-			if (IsValidInterfaceRegistryKey(Key))
+			const FMetasoundFrontendVersion& Version = InEntry->GetInterface().Metadata.Version;
+			if (ContainsInterface(Version))
 			{
-				if (const IInterfaceRegistryEntry* Entry = FindInterfaceRegistryEntry(Key))
-				{
-					UE_LOG(LogMetaSound, Warning, TEXT("Registration of interface overwriting previously registered interface [RegistryKey: %s]"), *Key);
+				UE_LOG(LogMetaSound, Warning, TEXT("Registration of interface overwriting previously registered interface [%s]"), *Version.ToString());
 						
-					FInterfaceRegistryTransaction Transaction{FInterfaceRegistryTransaction::ETransactionType::InterfaceUnregistration, Key, Entry->GetInterface().Metadata.Version, TransactionTime};
-					TransactionBuffer->AddTransaction(MoveTemp(Transaction));
-				}
+				FInterfaceRegistryTransaction Transaction { FInterfaceRegistryTransaction::ETransactionType::InterfaceUnregistration, Version, TransactionTime };
+				TransactionBuffer->AddTransaction(MoveTemp(Transaction));
+			}
 
 #if UE_METASOUND_ENABLE_INTERFACE_VALIDATION
-				// Don't run vertex name validation warning for deprecated registry entries. Some of these may be
-				// versioning schema that have a subsequent version or versions fixing the very problem this log is reporting.
-				if (!InEntry->IsDeprecated())
+			// Don't run vertex name validation warning for deprecated registry entries. Some of these may be
+			// versioning schema that have a subsequent version or versions fixing the very problem this log is reporting.
+			if (!InEntry->IsDeprecated())
+			{
+				const FName InterfaceNamespace = InEntry->GetInterface().Version.Name;
+				auto LogIfMismatch = [this, &InterfaceNamespace](FName VertexName)
 				{
-					const FName InterfaceNamespace = InEntry->GetInterface().Version.Name;
-					auto LogIfMismatch = [this, &InterfaceNamespace](FName VertexName)
-					{
-						FName VertexNamespace;
-						if (!IsInterfaceVertexNameValid(InterfaceNamespace, VertexName, &VertexNamespace))
-						{
-							UE_LOG(LogMetaSound, Warning, TEXT("Interface '%s' contains vertex '%s' with mismatched namespace '%s': "
-								"All interface-defined vertices' must start with matching interface namespace (See AUDIO_PARAMETER_INTERFACE_MEMBER_DEFINE/AUDIO_PARAMETER_INTERFACE_NAMESPACE macro to ensure convention is followed). "
-								"Failing to fix relationship via interface versioning will fail validation/cook in future builds."), *InterfaceNamespace.ToString(), *VertexName.ToString(), *VertexNamespace.ToString());
-						}
-					};
-
-					for (const FMetasoundFrontendClassVertex& Vertex : InEntry->GetInterface().Inputs)
+					FName VertexNamespace;
+					if (!IsInterfaceVertexNameValid(InterfaceNamespace, VertexName, &VertexNamespace))
 					{
-						LogIfMismatch(Vertex.Name);
+						UE_LOG(LogMetaSound, Warning, TEXT("Interface '%s' contains vertex '%s' with mismatched namespace '%s': "
+							"All interface-defined vertices' must start with matching interface namespace (See AUDIO_PARAMETER_INTERFACE_MEMBER_DEFINE/AUDIO_PARAMETER_INTERFACE_NAMESPACE macro to ensure convention is followed). "
+							"Failing to fix relationship via interface versioning will fail validation/cook in future builds."), *InterfaceNamespace.ToString(), *VertexName.ToString(), *VertexNamespace.ToString());
 					}
+				};
 
-					for (const FMetasoundFrontendClassVertex& Vertex : InEntry->GetInterface().Outputs)
-					{
-						LogIfMismatch(Vertex.Name);
-					}
+				for (const FMetasoundFrontendClassVertex& Vertex : InEntry->GetInterface().Inputs)
+				{
+					LogIfMismatch(Vertex.Name);
 				}
+
+				for (const FMetasoundFrontendClassVertex& Vertex : InEntry->GetInterface().Outputs)
+				{
+					LogIfMismatch(Vertex.Name);
+				}
+			}
 #endif // UE_METASOUND_ENABLE_INTERFACE_VALIDATION
 
-				FInterfaceRegistryTransaction Transaction{FInterfaceRegistryTransaction::ETransactionType::InterfaceRegistration, Key, InEntry->GetInterface().Metadata.Version, TransactionTime};
-				TransactionBuffer->AddTransaction(MoveTemp(Transaction));
-				Entries.Add(Key, MoveTemp(InEntry));
-				return true;
+			FInterfaceRegistryTransaction Transaction{FInterfaceRegistryTransaction::ETransactionType::InterfaceRegistration, Version, TransactionTime};
+			TransactionBuffer->AddTransaction(MoveTemp(Transaction));
+
+			{
+				UE::TScopeLock Lock(EntriesCriticalSection);
+				Entries.Add(Version, MoveTemp(InEntry));
+			}
+			return true;
+		}
+
+		return false;
+	}
+
+	bool FInterfaceRegistry::ContainsInput(const FMetasoundFrontendVersion& InVersion, const FMetasoundFrontendVertex& InInput) const
+	{
+		auto IsVertex = [&InInput](const FMetasoundFrontendVertex& InterfaceInput)
+		{
+			return FMetasoundFrontendVertex::IsFunctionalEquivalent(InInput, InterfaceInput);
+		};
+
+		{
+			UE::TScopeLock Lock(EntriesCriticalSection);
+			if (const TUniquePtr<IInterfaceRegistryEntry>* Entry = Entries.Find(InVersion))
+			{
+				check(Entry);
+				return (*Entry)->GetInterface().Inputs.ContainsByPredicate(IsVertex);
 			}
 		}
 
 		return false;
 	}
 
-	const IInterfaceRegistryEntry* FInterfaceRegistry::FindInterfaceRegistryEntry(const FInterfaceRegistryKey& InKey) const
+	bool FInterfaceRegistry::ContainsInterface(const FMetasoundFrontendVersion& InVersion) const
+	{
+		UE::TScopeLock Lock(EntriesCriticalSection);
+		return Entries.Contains(InVersion);
+	}
+
+	bool FInterfaceRegistry::ContainsOutput(const FMetasoundFrontendVersion& InVersion, const FMetasoundFrontendVertex& InOutput) const
 	{
-		if (const TUniquePtr<IInterfaceRegistryEntry>* Entry = Entries.Find(InKey))
+		auto IsVertex = [&InOutput](const FMetasoundFrontendVertex& InterfaceInput)
 		{
-			return Entry->Get();
+			return FMetasoundFrontendVertex::IsFunctionalEquivalent(InOutput, InterfaceInput);
+		};
+
+		{
+			UE::TScopeLock Lock(EntriesCriticalSection);
+			if (const TUniquePtr<IInterfaceRegistryEntry>* Entry = Entries.Find(InVersion))
+			{
+				check(Entry);
+				return (*Entry)->GetInterface().Inputs.ContainsByPredicate(IsVertex);
+			}
 		}
-		return nullptr;
+
+		return false;
 	}
 
-	bool FInterfaceRegistry::FindInterface(const FInterfaceRegistryKey& InKey, FMetasoundFrontendInterface& OutInterface) const
+	bool FInterfaceRegistry::FindInterface(const FMetasoundFrontendVersion& InVersion, FMetasoundFrontendInterface& OutInterface) const
 	{
-		if (const IInterfaceRegistryEntry* Entry = FindInterfaceRegistryEntry(InKey))
+		UE::TScopeLock Lock(EntriesCriticalSection);
+		if (const TUniquePtr<IInterfaceRegistryEntry>* Entry = Entries.Find(InVersion))
 		{
-			OutInterface = Entry->GetInterface();
+			OutInterface = Entry->Get()->GetInterface();
 			return true;
 		}
 
 		return false;
 	}
 
+	bool FInterfaceRegistry::FindInterfaceRouter(const FMetasoundFrontendVersion& InVersion, FName& OutRouter) const
+	{
+		UE::TScopeLock Lock(EntriesCriticalSection);
+		if (const TUniquePtr<IInterfaceRegistryEntry>* Entry = Entries.Find(InVersion))
+		{
+			OutRouter = Entry->Get()->GetRouterName();
+			return true;
+		}
+
+		return false;
+	}
+
+	TSharedPtr<IBuilderVersionTransform> FInterfaceRegistry::FindInterfaceUpdateTransform(const FMetasoundFrontendVersion& InVersion) const
+	{
+		UE::TScopeLock Lock(EntriesCriticalSection);
+		if (const TUniquePtr<IInterfaceRegistryEntry>* Entry = Entries.Find(InVersion))
+		{
+			return Entry->Get()->GetUpdateTransform();
+		}
+
+		return { };
+	}
+
 	TUniquePtr<FInterfaceTransactionStream> FInterfaceRegistry::CreateTransactionStream()
 	{
 		return MakeUnique<FInterfaceTransactionStream>(TransactionBuffer);
@@ -115,17 +176,25 @@ namespace Metasound::Frontend
 
 	FInterfaceRegistryKey GetInterfaceRegistryKey(const FMetasoundFrontendInterface& InInterface)
 	{
+PRAGMA_DISABLE_DEPRECATION_WARNINGS
 		return GetInterfaceRegistryKey(InInterface.Metadata.Version);
+PRAGMA_ENABLE_DEPRECATION_WARNINGS
 	}
 
-	FInterfaceRegistryTransaction::FInterfaceRegistryTransaction(ETransactionType InType, const FInterfaceRegistryKey& InKey, const FMetasoundFrontendVersion& InInterfaceVersion, FInterfaceRegistryTransaction::FTimeType InTimestamp)
+	FInterfaceRegistryTransaction::FInterfaceRegistryTransaction(ETransactionType InType, const FInterfaceRegistryKey&, const FMetasoundFrontendVersion& InInterfaceVersion, FInterfaceRegistryTransaction::FTimeType InTimestamp)
 	: Type(InType)
-	, Key(InKey)
 	, InterfaceVersion(InInterfaceVersion)
 	, Timestamp(InTimestamp)
 	{
 	}
 
+	FInterfaceRegistryTransaction::FInterfaceRegistryTransaction(ETransactionType InType, const FMetasoundFrontendVersion& InInterfaceVersion, FInterfaceRegistryTransaction::FTimeType InTimestamp)
+		: Type(InType)
+		, InterfaceVersion(InInterfaceVersion)
+		, Timestamp(InTimestamp)
+	{
+	}
+
 	FInterfaceRegistryTransaction::ETransactionType FInterfaceRegistryTransaction::GetTransactionType() const
 	{
 		return Type;
@@ -138,7 +207,8 @@ namespace Metasound::Frontend
 
 	const FInterfaceRegistryKey& FInterfaceRegistryTransaction::GetInterfaceRegistryKey() const
 	{
-		return Key;
+		static const FInterfaceRegistryKey InvalidKey;
+		return InvalidKey;
 	}
 
 	FInterfaceRegistryTransaction::FTimeType FInterfaceRegistryTransaction::GetTimestamp() const
diff --git a/Engine/Plugins/Runtime/Metasound/Source/MetasoundFrontend/Private/Interfaces/MetasoundFrontendInterfaceRegistryPrivate.h b/Engine/Plugins/Runtime/Metasound/Source/MetasoundFrontend/Private/Interfaces/MetasoundFrontendInterfaceRegistryPrivate.h
--- a/Engine/Plugins/Runtime/Metasound/Source/MetasoundFrontend/Private/Interfaces/MetasoundFrontendInterfaceRegistryPrivate.h
+++ b/Engine/Plugins/Runtime/Metasound/Source/MetasoundFrontend/Private/Interfaces/MetasoundFrontendInterfaceRegistryPrivate.h
@@ -5,6 +5,7 @@
 #include "MetasoundFrontendController.h"
 #include "MetasoundFrontendDocument.h"
 #include "MetasoundFrontendRegistryTransaction.h"
+#include "Misc/TransactionallySafeCriticalSection.h"
 #include "UObject/NoExportTypes.h"
 #include "UObject/TopLevelAssetPath.h"
 
@@ -22,9 +23,14 @@ namespace Metasound::Frontend
 
 		virtual ~FInterfaceRegistry() = default;
 
+		virtual bool ContainsInput(const FMetasoundFrontendVersion& InVersion, const FMetasoundFrontendVertex& InInput) const override;
+		virtual bool ContainsInterface(const FMetasoundFrontendVersion& InVersion) const override;
+		virtual bool ContainsOutput(const FMetasoundFrontendVersion& InVersion, const FMetasoundFrontendVertex& InOutput) const override;
+
+		virtual bool FindInterface(const FMetasoundFrontendVersion& InVersion, FMetasoundFrontendInterface& OutInterface) const override;
+		virtual bool FindInterfaceRouter(const FMetasoundFrontendVersion& InVersion, FName& OutRouter) const override;
+		virtual TSharedPtr<IBuilderVersionTransform> FindInterfaceUpdateTransform(const FMetasoundFrontendVersion& InVersion) const override;
 		virtual bool RegisterInterface(TUniquePtr<IInterfaceRegistryEntry>&& InEntry) override;
-		virtual const IInterfaceRegistryEntry* FindInterfaceRegistryEntry(const FInterfaceRegistryKey& InKey) const override;
-		virtual bool FindInterface(const FInterfaceRegistryKey& InKey, FMetasoundFrontendInterface& OutInterface) const override;
 
 		TUniquePtr<FInterfaceTransactionStream> CreateTransactionStream();
 
@@ -33,7 +39,8 @@ namespace Metasound::Frontend
 
 		using FInterfaceTransactionBuffer = TTransactionBuffer<FInterfaceRegistryTransaction>;
 
-		TMap<FInterfaceRegistryKey, TUniquePtr<IInterfaceRegistryEntry>> Entries;
+		mutable FTransactionallySafeCriticalSection EntriesCriticalSection;
+		TMap<FMetasoundFrontendVersion, TUniquePtr<IInterfaceRegistryEntry>> Entries;
 
 		TSharedRef<FInterfaceTransactionBuffer> TransactionBuffer;
 	};
diff --git a/Engine/Plugins/Runtime/Metasound/Source/MetasoundFrontend/Private/MetasoundAssetBase.cpp b/Engine/Plugins/Runtime/Metasound/Source/MetasoundFrontend/Private/MetasoundAssetBase.cpp
--- a/Engine/Plugins/Runtime/Metasound/Source/MetasoundFrontend/Private/MetasoundAssetBase.cpp
+++ b/Engine/Plugins/Runtime/Metasound/Source/MetasoundFrontend/Private/MetasoundAssetBase.cpp
@@ -130,7 +130,7 @@ namespace Metasound
 				FTopLevelAssetPath AssetPath;
 			};
 
-			void GetUpdatePathForDocument(const FMetasoundFrontendVersion& InCurrentVersion, const FMetasoundFrontendVersion& InTargetVersion, TArray<const IInterfaceRegistryEntry*>& OutUpgradePath)
+			void GetUpdatePathForDocument(const FMetasoundFrontendVersion& InCurrentVersion, const FMetasoundFrontendVersion& InTargetVersion, TArray<TSharedPtr<IBuilderVersionTransform>>& OutUpgradePath)
 			{
 				if (InCurrentVersion.Name == InTargetVersion.Name)
 				{
@@ -153,24 +153,21 @@ namespace Metasound
 					// Get registry entries from registry keys.
 					auto GetRegistryEntry = [](const FMetasoundFrontendVersion& InVersion)
 					{
-						FInterfaceRegistryKey Key = GetInterfaceRegistryKey(InVersion);
-						return IInterfaceRegistry::Get().FindInterfaceRegistryEntry(Key);
+						return IInterfaceRegistry::Get().FindInterfaceUpdateTransform(InVersion);
 					};
 					Algo::Transform(RegisteredVersions, OutUpgradePath, GetRegistryEntry);
 				}
 			}
 
-			bool UpdateDocumentInterface(const TArray<const IInterfaceRegistryEntry*>& InUpgradePath, const FMetasoundFrontendVersion& InterfaceVersion, FMetaSoundFrontendDocumentBuilder& Builder)
+			bool UpdateDocumentInterface(const TArray<TSharedPtr<IBuilderVersionTransform>>& InUpgradePath, const FMetasoundFrontendVersion& InterfaceVersion, FMetaSoundFrontendDocumentBuilder& Builder)
 			{
 				const FMetasoundFrontendVersionNumber* LastVersionUpdated = nullptr;
-				for (const IInterfaceRegistryEntry* Entry : InUpgradePath)
+				for (const TSharedPtr<IBuilderVersionTransform>& Transform : InUpgradePath)
 				{
-					if (ensure(nullptr != Entry))
+					check(Transform.IsValid());
+					if (Transform->Transform(Builder))
 					{
-						if (Entry->UpdateRootGraphInterface(Builder))
-						{
-							LastVersionUpdated = &Entry->GetInterface().Metadata.Version.Number;
-						}
+						LastVersionUpdated = &Transform->GetVersion().Number;
 					}
 				}
 
@@ -901,13 +898,13 @@ bool FMetasoundAssetBase::TryUpdateInterfaceFromVersion(const FMetasoundFrontend
 	using namespace Metasound::Frontend;
 	using namespace AssetBasePrivate;
 
-	FMetasoundFrontendInterface TargetInterface = GetInterfaceToVersion(Version);
-	if (TargetInterface.Metadata.Version.IsValid())
+	const FMetasoundFrontendVersion TargetInterfaceVersion = GetTargetInterfaceVersion(Version);
+	if (TargetInterfaceVersion.IsValid())
 	{
-		TArray<const IInterfaceRegistryEntry*> UpgradePath;
-		GetUpdatePathForDocument(Version, TargetInterface.Metadata.Version, UpgradePath);
+		TArray<TSharedPtr<IBuilderVersionTransform>> UpgradePath;
+		GetUpdatePathForDocument(Version, TargetInterfaceVersion, UpgradePath);
 		const bool bUpdated = UpdateDocumentInterface(UpgradePath, Version, Builder);
-		ensureMsgf(bUpdated, TEXT("Target interface '%s' was out-of-date but interface failed to be updated"), *TargetInterface.Metadata.Version.ToString());
+		ensureMsgf(bUpdated, TEXT("Target interface '%s' was out-of-date but interface failed to be updated"), *TargetInterfaceVersion.ToString());
 		return bUpdated;
 	}
 
@@ -928,13 +925,13 @@ bool FMetasoundAssetBase::VersionDependencies(FMetaSoundFrontendDocumentBuilder&
 	return bDocumentModified;
 }
 
-FMetasoundFrontendInterface FMetasoundAssetBase::GetInterfaceToVersion(const FMetasoundFrontendVersion& InterfaceVersion) const
+FMetasoundFrontendVersion FMetasoundAssetBase::GetTargetInterfaceVersion(const FMetasoundFrontendVersion& InterfaceVersion) const
 {
 	using namespace Metasound::Frontend;
 
 	// Find registered target interface.
-	FMetasoundFrontendInterface TargetInterface;
-	bool bFoundTargetInterface = ISearchEngine::Get().FindInterfaceWithHighestVersion(InterfaceVersion.Name, TargetInterface);
+	FMetasoundFrontendVersion TargetInterfaceVersion;
+	bool bFoundTargetInterface = ISearchEngine::Get().FindHighestInterfaceVersion(InterfaceVersion.Name, TargetInterfaceVersion);
 	if (!bFoundTargetInterface)
 	{
 		UE_LOG(LogMetaSound, Warning,
@@ -945,12 +942,12 @@ FMetasoundFrontendInterface FMetasoundAssetBase::GetInterfaceToVersion(const FMe
 		return { };
 	}
 
-	if (TargetInterface.Metadata.Version == InterfaceVersion)
+	if (TargetInterfaceVersion == InterfaceVersion)
 	{
 		return { };
 	}
 
-	return TargetInterface;
+	return TargetInterfaceVersion;
 }
 
 #if WITH_EDITORONLY_DATA
diff --git a/Engine/Plugins/Runtime/Metasound/Source/MetasoundFrontend/Private/MetasoundFrontendDocumentBuilder.cpp b/Engine/Plugins/Runtime/Metasound/Source/MetasoundFrontend/Private/MetasoundFrontendDocumentBuilder.cpp
--- a/Engine/Plugins/Runtime/Metasound/Source/MetasoundFrontend/Private/MetasoundFrontendDocumentBuilder.cpp
+++ b/Engine/Plugins/Runtime/Metasound/Source/MetasoundFrontend/Private/MetasoundFrontendDocumentBuilder.cpp
@@ -96,14 +96,20 @@ namespace Metasound::Frontend
 			FName ParamName;
 			Audio::FParameterPath::SplitName(VertexName, InterfaceNamespace, ParamName);
 
-			FMetasoundFrontendInterface FoundInterface;
-			if (!InterfaceNamespace.IsNone() && ISearchEngine::Get().FindInterfaceWithHighestVersion(InterfaceNamespace, FoundInterface))
+			if (!InterfaceNamespace.IsNone())
 			{
-				if (OutInterface)
+				FMetasoundFrontendVersion InterfaceVersion;
+				if (ISearchEngine::Get().FindHighestInterfaceVersion(InterfaceNamespace, InterfaceVersion))
 				{
-					*OutInterface = MoveTemp(FoundInterface);
+					if (OutInterface)
+					{
+						FMetasoundFrontendInterface FoundInterface;
+						const bool bInterfaceFound = IInterfaceRegistry::Get().FindInterface(InterfaceVersion, FoundInterface);
+						checkf(bInterfaceFound, TEXT("MetaSound Interface found via search engine must exist in registry"))
+						*OutInterface = MoveTemp(FoundInterface);
+					}
+					return true;
 				}
-				return true;
 			}
 
 			if (OutInterface)
@@ -702,7 +708,7 @@ namespace Metasound::Frontend
 		Algo::Transform(InInterfaceVersionsToRemove, InterfacesToRemove, [](const FMetasoundFrontendVersion& Version)
 		{
 			FMetasoundFrontendInterface Interface;
-			const bool bFromInterfaceFound = IInterfaceRegistry::Get().FindInterface(GetInterfaceRegistryKey(Version), Interface);
+			const bool bFromInterfaceFound = IInterfaceRegistry::Get().FindInterface(Version, Interface);
 			if (!ensureAlways(bFromInterfaceFound))
 			{
 				UE_LOG(LogMetaSound, Error, TEXT("Failed to find interface '%s' to remove"), *Version.ToString());
@@ -713,7 +719,7 @@ namespace Metasound::Frontend
 		Algo::Transform(InInterfaceVersionsToAdd, InterfacesToAdd, [](const FMetasoundFrontendVersion& Version)
 		{
 			FMetasoundFrontendInterface Interface;
-			const bool bToInterfaceFound = IInterfaceRegistry::Get().FindInterface(GetInterfaceRegistryKey(Version), Interface);
+			const bool bToInterfaceFound = IInterfaceRegistry::Get().FindInterface(Version, Interface);
 			if (!ensureAlways(bToInterfaceFound))
 			{
 				UE_LOG(LogMetaSound, Error, TEXT("Failed to find interface '%s' to add"), *Version.ToString());
@@ -1050,10 +1056,11 @@ bool FMetaSoundFrontendDocumentBuilder::AddEdgesFromMatchingInterfaceNodeOutputs
 	TSet<FNamedEdge> EdgesToMake;
 	for (const FMetasoundFrontendVersion& Version : CommonInterfaces)
 	{
-		const FInterfaceRegistryKey InterfaceKey = GetInterfaceRegistryKey(Version);
-		if (const IInterfaceRegistryEntry* RegistryEntry = IInterfaceRegistry::Get().FindInterfaceRegistryEntry(InterfaceKey))
+		FMetasoundFrontendInterface Interface;
+		const bool bFound = IInterfaceRegistry::Get().FindInterface(Version, Interface);
+		if (bFound)
 		{
-			Algo::Transform(RegistryEntry->GetInterface().Outputs, EdgesToMake, [this, &NodeCache, &InterfaceCache, &PageID, InNodeID](const FMetasoundFrontendClassOutput& Output)
+			Algo::Transform(Interface.Outputs, EdgesToMake, [this, &NodeCache, &InterfaceCache, &PageID, InNodeID](const FMetasoundFrontendClassOutput& Output)
 			{
 				const FMetasoundFrontendGraph& Graph = GetDocumentChecked().RootGraph.FindConstGraphChecked(PageID);
 				const FMetasoundFrontendVertex* NodeVertex = NodeCache.FindOutputVertex(InNodeID, Output.Name);
@@ -1097,10 +1104,11 @@ bool FMetaSoundFrontendDocumentBuilder::AddEdgesFromMatchingInterfaceNodeInputsT
 	const FMetasoundFrontendGraph& Graph = GetDocumentChecked().RootGraph.FindConstGraphChecked(PageID);
 	for (const FMetasoundFrontendVersion& Version : CommonInterfaces)
 	{
-		const FInterfaceRegistryKey InterfaceKey = GetInterfaceRegistryKey(Version);
-		if (const IInterfaceRegistryEntry* RegistryEntry = IInterfaceRegistry::Get().FindInterfaceRegistryEntry(InterfaceKey))
+		FMetasoundFrontendInterface Interface;
+		const bool bFound = IInterfaceRegistry::Get().FindInterface(Version, Interface);
+		if (bFound)
 		{
-			Algo::Transform(RegistryEntry->GetInterface().Inputs, EdgesToMake, [this, &Graph, &NodeCache, &InterfaceCache, InNodeID](const FMetasoundFrontendClassInput& Input)
+			Algo::Transform(Interface.Inputs, EdgesToMake, [this, &Graph, &NodeCache, &InterfaceCache, InNodeID](const FMetasoundFrontendClassInput& Input)
 			{
 				const FMetasoundFrontendVertex* NodeVertex = NodeCache.FindInputVertex(InNodeID, Input.Name);
 				check(NodeVertex);
@@ -1720,15 +1728,19 @@ bool FMetaSoundFrontendDocumentBuilder::AddInterface(FName InterfaceName, bool b
 {
 	using namespace Metasound::Frontend;
 
-	FMetasoundFrontendInterface Interface;
-	if (ISearchEngine::Get().FindInterfaceWithHighestVersion(InterfaceName, Interface))
+	FMetasoundFrontendVersion InterfaceVersion;
+	if (ISearchEngine::Get().FindHighestInterfaceVersion(InterfaceName, InterfaceVersion))
 	{
-		if (GetDocumentChecked().Interfaces.Contains(Interface.Metadata.Version))
+		if (GetDocumentChecked().Interfaces.Contains(InterfaceVersion))
 		{
 			UE_LOG(LogMetaSound, VeryVerbose, TEXT("MetaSound interface '%s' already found on document. MetaSoundBuilder skipping add request."), *InterfaceName.ToString());
 			return true;
 		}
 
+		FMetasoundFrontendInterface Interface;
+		const bool bFoundInterface = IInterfaceRegistry::Get().FindInterface(InterfaceVersion, Interface);
+		check(bFoundInterface);
+
 		// Get all versions of the interface, excluding the latest
 		TArray<FMetasoundFrontendVersion> PreviousVersions = ISearchEngine::Get().FindAllRegisteredInterfacesWithName(Interface.Metadata.Version.Name);
 		PreviousVersions.Remove(Interface.Metadata.Version);
@@ -1739,30 +1751,24 @@ bool FMetaSoundFrontendDocumentBuilder::AddInterface(FName InterfaceName, bool b
 		for (const FMetasoundFrontendVersion& PreviousVersion : PreviousVersions)
 		{
 			FMetasoundFrontendInterface PreviousVersionInterface;
-			if (IInterfaceRegistry::Get().FindInterface(GetInterfaceRegistryKey(PreviousVersion), PreviousVersionInterface))
+			if (IInterfaceRegistry::Get().FindInterface(PreviousVersion, PreviousVersionInterface))
 			{
 				PreviousVersionInterfaces.Add(MoveTemp(PreviousVersionInterface));
 			}
 		}
 
-		const FInterfaceRegistryKey Key = GetInterfaceRegistryKey(Interface.Metadata.Version);
-		if (const IInterfaceRegistryEntry* Entry = IInterfaceRegistry::Get().FindInterfaceRegistryEntry(Key))
+		const FTopLevelAssetPath BuilderClassPath = GetBuilderClassPath();
+		auto FindClassOptionsPredicate = [&BuilderClassPath](const FMetasoundFrontendInterfaceUClassOptions& Options) { return Options.ClassPath == BuilderClassPath; };
+		const FMetasoundFrontendInterfaceUClassOptions* ClassOptions = Interface.Metadata.UClassOptions.FindByPredicate(FindClassOptionsPredicate);
+		if (ClassOptions && bAddUserModifiableInterfaceOnly && !ClassOptions->bIsModifiable)
 		{
-			const FTopLevelAssetPath BuilderClassPath = GetBuilderClassPath();
-			auto FindClassOptionsPredicate = [&BuilderClassPath](const FMetasoundFrontendInterfaceUClassOptions& Options) { return Options.ClassPath == BuilderClassPath; };
-			const FMetasoundFrontendInterfaceUClassOptions* ClassOptions = Entry->GetInterface().Metadata.UClassOptions.FindByPredicate(FindClassOptionsPredicate);
-			if (ClassOptions && bAddUserModifiableInterfaceOnly && !ClassOptions->bIsModifiable)
-			{
-				UE_LOG(LogMetaSound, Error, TEXT("DocumentBuilder failed to add MetaSound Interface '%s' to document: is not set to be user modifiable for given UClass '%s'"), *InterfaceName.ToString(), *BuilderClassPath.ToString());
-				return false;
-			}
-
-			// Remove old versions and add new interfaces simultaneously to support node reconnection
-			TArray<FMetasoundFrontendInterface> InterfacesToAdd;
-			InterfacesToAdd.Add(Entry->GetInterface());
-			FModifyInterfaceOptions Options(MoveTemp(PreviousVersionInterfaces), MoveTemp(InterfacesToAdd));
-			return ModifyInterfaces(MoveTemp(Options));
+			UE_LOG(LogMetaSound, Error, TEXT("DocumentBuilder failed to add MetaSound Interface '%s' to document: is not set to be user modifiable for given UClass '%s'"), *InterfaceName.ToString(), *BuilderClassPath.ToString());
+			return false;
 		}
+
+		// Remove old versions and add new interfaces simultaneously to support node reconnection
+		FModifyInterfaceOptions Options(MoveTemp(PreviousVersionInterfaces), { MoveTemp(Interface) });
+		return ModifyInterfaces(MoveTemp(Options));
 	}
 
 	return false;
@@ -2069,17 +2075,17 @@ void FMetaSoundFrontendDocumentBuilder::CacheRegistryMetadata() const
 	const TSet<FMetasoundFrontendVersion>& InterfaceVersions = Document.Interfaces;
 	for (const FMetasoundFrontendVersion& Version : InterfaceVersions)
 	{
-		const FInterfaceRegistryKey InterfaceKey = GetInterfaceRegistryKey(Version);
-		const IInterfaceRegistryEntry* Entry = IInterfaceRegistry::Get().FindInterfaceRegistryEntry(InterfaceKey);
+		FMetasoundFrontendInterface Interface;
+		const bool bFound = IInterfaceRegistry::Get().FindInterface(Version, Interface);
 
-		UE_CLOG(nullptr == Entry, LogMetaSound, Error,
+		UE_CLOG(!bFound, LogMetaSound, Error,
 			TEXT("Failed to find interface (%s) when caching registry data for %s. "
 			"MetaSound inputs and outputs for asset may not function correctly."),
 			*Version.ToString(), *GetDebugName());
 
-		if (Entry)
+		if (bFound)
 		{
-			for (const FMetasoundFrontendClassInput& InterfaceInput : Entry->GetInterface().Inputs)
+			for (const FMetasoundFrontendClassInput& InterfaceInput : Interface.Inputs)
 			{
 				const FNameDataTypePair NameDataTypePair = FNameDataTypePair(InterfaceInput.Name, InterfaceInput.TypeName);
 				if (FMetasoundFrontendClassInput* Input = Inputs.FindRef(NameDataTypePair))
@@ -2089,7 +2095,7 @@ void FMetaSoundFrontendDocumentBuilder::CacheRegistryMetadata() const
 				}
 			}
 
-			for (const FMetasoundFrontendClassOutput& InterfaceOutput : Entry->GetInterface().Outputs)
+			for (const FMetasoundFrontendClassOutput& InterfaceOutput : Interface.Outputs)
 			{
 				const FNameDataTypePair NameDataTypePair = FNameDataTypePair(InterfaceOutput.Name, InterfaceOutput.TypeName);
 				if (FMetasoundFrontendClassOutput* Output = Outputs.FindRef(NameDataTypePair))
@@ -2828,30 +2834,12 @@ TConstStructView<FMetasoundFrontendClassInterface> FMetaSoundFrontendDocumentBui
 
 bool FMetaSoundFrontendDocumentBuilder::FindDeclaredInterfaces(TArray<const Metasound::Frontend::IInterfaceRegistryEntry*>& OutInterfaces) const
 {
-	return FindDeclaredInterfaces(GetConstDocumentChecked(), OutInterfaces);
+	return false;
 }
 
 bool FMetaSoundFrontendDocumentBuilder::FindDeclaredInterfaces(const FMetasoundFrontendDocument& InDocument, TArray<const Metasound::Frontend::IInterfaceRegistryEntry*>& OutInterfaces)
 {
-	using namespace Metasound;
-	using namespace Metasound::Frontend;
-
-	bool bInterfacesFound = true;
-
-	Algo::Transform(InDocument.Interfaces, OutInterfaces, [&bInterfacesFound](const FMetasoundFrontendVersion& Version)
-	{
-		const FInterfaceRegistryKey InterfaceKey = GetInterfaceRegistryKey(Version);
-		const IInterfaceRegistryEntry* RegistryEntry = IInterfaceRegistry::Get().FindInterfaceRegistryEntry(InterfaceKey);
-		if (!RegistryEntry)
-		{
-			bInterfacesFound = false;
-			UE_LOG(LogMetaSound, Warning, TEXT("No registered interface matching interface version on document [InterfaceVersion:%s]"), *Version.ToString());
-		}
-
-		return RegistryEntry;
-	});
-
-	return bInterfacesFound;
+	return false;
 }
 
 const FMetasoundFrontendClass* FMetaSoundFrontendDocumentBuilder::FindDependency(const FGuid& InClassID) const
@@ -3001,12 +2989,16 @@ bool FMetaSoundFrontendDocumentBuilder::FindInterfaceInputNodes(FName InterfaceN
 	OutInputs.Reset();
 
 	const FGuid& PageID = InPageID ? *InPageID : BuildPageID;
-	FMetasoundFrontendInterface Interface;
+	FMetasoundFrontendVersion InterfaceVersion;
 	const TSet<FMetasoundFrontendVersion>& Interfaces = GetConstDocumentChecked().Interfaces;
-	if (ISearchEngine::Get().FindInterfaceWithHighestVersion(InterfaceName, Interface))
+	if (ISearchEngine::Get().FindHighestInterfaceVersion(InterfaceName, InterfaceVersion))
 	{
-		if (Interfaces.Contains(Interface.Metadata.Version))
+		if (Interfaces.Contains(InterfaceVersion))
 		{
+			FMetasoundFrontendInterface Interface;
+			const bool bFoundInterface = IInterfaceRegistry::Get().FindInterface(InterfaceVersion, Interface);
+			check(bFoundInterface);
+
 			const IDocumentGraphNodeCache& NodeCache = DocumentCache->GetNodeCache(PageID);
 			const IDocumentGraphInterfaceCache& InterfaceCache = DocumentCache->GetInterfaceCache();
 
@@ -3045,36 +3037,40 @@ bool FMetaSoundFrontendDocumentBuilder::FindInterfaceOutputNodes(FName Interface
 
 	const FGuid& PageID = InPageID ? *InPageID : BuildPageID;
 
-	FMetasoundFrontendInterface Interface;
 	const TSet<FMetasoundFrontendVersion>& Interfaces = GetConstDocumentChecked().Interfaces;
-	if (ISearchEngine::Get().FindInterfaceWithHighestVersion(InterfaceName, Interface))
+	FMetasoundFrontendVersion InterfaceVersion;
+	if (ISearchEngine::Get().FindHighestInterfaceVersion(InterfaceName, InterfaceVersion))
 	{
-		if (Interfaces.Contains(Interface.Metadata.Version))
+		if (Interfaces.Contains(InterfaceVersion))
 		{
-			const IDocumentGraphNodeCache& NodeCache = DocumentCache->GetNodeCache(PageID);
-			const IDocumentGraphInterfaceCache& InterfaceCache = DocumentCache->GetInterfaceCache();
-
-			TArray<const FMetasoundFrontendNode*> InterfaceOutputs;
-			for (const FMetasoundFrontendClassOutput& Output : Interface.Outputs)
+			FMetasoundFrontendInterface Interface;
+			if (IInterfaceRegistry::Get().FindInterface(InterfaceVersion, Interface))
 			{
-				const FMetasoundFrontendClassOutput* ClassOutput = InterfaceCache.FindOutput(Output.Name);
-				if (!ClassOutput)
-				{
-					return false;
-				}
+				const IDocumentGraphNodeCache& NodeCache = DocumentCache->GetNodeCache(PageID);
+				const IDocumentGraphInterfaceCache& InterfaceCache = DocumentCache->GetInterfaceCache();
 
-				if (const FMetasoundFrontendNode* Node = NodeCache.FindNode(ClassOutput->NodeID))
+				TArray<const FMetasoundFrontendNode*> InterfaceOutputs;
+				for (const FMetasoundFrontendClassOutput& Output : Interface.Outputs)
 				{
-					InterfaceOutputs.Add(Node);
-				}
-				else
-				{
-					return false;
+					const FMetasoundFrontendClassOutput* ClassOutput = InterfaceCache.FindOutput(Output.Name);
+					if (!ClassOutput)
+					{
+						return false;
+					}
+
+					if (const FMetasoundFrontendNode* Node = NodeCache.FindNode(ClassOutput->NodeID))
+					{
+						InterfaceOutputs.Add(Node);
+					}
+					else
+					{
+						return false;
+					}
 				}
-			}
 
-			OutOutputs = MoveTemp(InterfaceOutputs);
-			return true;
+				OutOutputs = MoveTemp(InterfaceOutputs);
+				return true;
+			}
 		}
 	}
 
@@ -4428,10 +4424,10 @@ bool FMetaSoundFrontendDocumentBuilder::IsInterfaceDeclared(FName InInterfaceNam
 {
 	using namespace Metasound::Frontend;
 
-	FMetasoundFrontendInterface Interface;
-	if (ISearchEngine::Get().FindInterfaceWithHighestVersion(InInterfaceName, Interface))
+	FMetasoundFrontendVersion InterfaceVersion;
+	if (ISearchEngine::Get().FindHighestInterfaceVersion(InInterfaceName, InterfaceVersion))
 	{
-		return IsInterfaceDeclared(Interface.Metadata.Version);
+		return IsInterfaceDeclared(InterfaceVersion);
 	}
 
 	return false;
@@ -5374,30 +5370,29 @@ bool FMetaSoundFrontendDocumentBuilder::RemoveInterface(FName InterfaceName)
 {
 	using namespace Metasound::Frontend;
 
-	FMetasoundFrontendInterface Interface;
-	if (ISearchEngine::Get().FindInterfaceWithHighestVersion(InterfaceName, Interface))
+	FMetasoundFrontendVersion InterfaceVersion;
+	if (ISearchEngine::Get().FindHighestInterfaceVersion(InterfaceName, InterfaceVersion))
 	{
-		if (!GetDocumentChecked().Interfaces.Contains(Interface.Metadata.Version))
+		if (!GetDocumentChecked().Interfaces.Contains(InterfaceVersion))
 		{
 			UE_LOG(LogMetaSound, VeryVerbose, TEXT("MetaSound interface '%s' not found on document. MetaSoundBuilder skipping remove request."), *InterfaceName.ToString());
 			return true;
 		}
 
-		const FInterfaceRegistryKey Key = GetInterfaceRegistryKey(Interface.Metadata.Version);
-		if (const IInterfaceRegistryEntry* Entry = IInterfaceRegistry::Get().FindInterfaceRegistryEntry(Key))
+		FMetasoundFrontendInterface Interface;
+		const bool bFoundInterface = IInterfaceRegistry::Get().FindInterface(InterfaceVersion, Interface);
+		if (bFoundInterface)
 		{
 			const FTopLevelAssetPath BuilderClassPath = GetBuilderClassPath();
 			auto FindClassOptionsPredicate = [&BuilderClassPath](const FMetasoundFrontendInterfaceUClassOptions& Options) { return Options.ClassPath == BuilderClassPath; };
-			const FMetasoundFrontendInterfaceUClassOptions* ClassOptions = Entry->GetInterface().Metadata.UClassOptions.FindByPredicate(FindClassOptionsPredicate);
+			const FMetasoundFrontendInterfaceUClassOptions* ClassOptions = Interface.Metadata.UClassOptions.FindByPredicate(FindClassOptionsPredicate);
 			if (ClassOptions && !ClassOptions->bIsModifiable)
 			{
 				UE_LOG(LogMetaSound, Error, TEXT("DocumentBuilder failed to remove MetaSound Interface '%s' to document: is not set to be modifiable for given UClass '%s'"), *InterfaceName.ToString(), *BuilderClassPath.ToString());
 				return false;
 			}
 
-			TArray<FMetasoundFrontendInterface> InterfacesToRemove;
-			InterfacesToRemove.Add(Entry->GetInterface());
-			FModifyInterfaceOptions Options(MoveTemp(InterfacesToRemove), { });
+			FModifyInterfaceOptions Options({ MoveTemp(Interface) }, { });
 			return ModifyInterfaces(MoveTemp(Options));
 		}
 	}
diff --git a/Engine/Plugins/Runtime/Metasound/Source/MetasoundFrontend/Private/MetasoundFrontendDocumentVersioning.cpp b/Engine/Plugins/Runtime/Metasound/Source/MetasoundFrontend/Private/MetasoundFrontendDocumentVersioning.cpp
--- a/Engine/Plugins/Runtime/Metasound/Source/MetasoundFrontend/Private/MetasoundFrontendDocumentVersioning.cpp
+++ b/Engine/Plugins/Runtime/Metasound/Source/MetasoundFrontend/Private/MetasoundFrontendDocumentVersioning.cpp
@@ -344,7 +344,7 @@ PRAGMA_ENABLE_DEPRECATION_WARNINGS
 					for (int32 i = 0; i < NumPreexistingInterfaceVersions; i++)
 					{
 						FMetasoundFrontendInterface Interface;
-						if (InterfaceRegistry.FindInterface(GetInterfaceRegistryKey(PreexistingInterfaceVersions[i]), Interface))
+						if (InterfaceRegistry.FindInterface(PreexistingInterfaceVersions[i], Interface))
 						{
 							CandidateInterfaces.Add(Interface);
 						}
@@ -533,11 +533,10 @@ PRAGMA_ENABLE_DEPRECATION_WARNINGS
 
 				for (const FMetasoundFrontendVersion& Version : InterfaceVersions)
 				{
-					FInterfaceRegistryKey RegistryKey = GetInterfaceRegistryKey(Version);
-					const IInterfaceRegistryEntry* Entry = IInterfaceRegistry::Get().FindInterfaceRegistryEntry(RegistryKey);
-					if (ensure(Entry))
+					FMetasoundFrontendInterface Interface;
+					const bool bFound = IInterfaceRegistry::Get().FindInterface(Version, Interface);
+					if (ensure(bFound))
 					{
-						const FMetasoundFrontendInterface& Interface = Entry->GetInterface();
 						Algo::Transform(Interface.Inputs, InterfaceInputs, [](const FMetasoundFrontendClassInput& Input)
 						{
 							return FNameDataTypePair(Input.Name, Input.TypeName);
diff --git a/Engine/Plugins/Runtime/Metasound/Source/MetasoundFrontend/Private/MetasoundFrontendInvalidController.h b/Engine/Plugins/Runtime/Metasound/Source/MetasoundFrontend/Private/MetasoundFrontendInvalidController.h
--- a/Engine/Plugins/Runtime/Metasound/Source/MetasoundFrontend/Private/MetasoundFrontendInvalidController.h
+++ b/Engine/Plugins/Runtime/Metasound/Source/MetasoundFrontend/Private/MetasoundFrontendInvalidController.h
@@ -252,7 +252,10 @@ namespace Metasound
 			virtual const FText& GetDescription() const override { return Invalid::GetInvalidText(); }
 #endif // WITH_EDITOR
 
+			UE_DEPRECATED(5.8, "Controllers are being actively deprecated, and this call was not thread safe in certain contexts")
 			virtual const FMetasoundFrontendVersion& GetInterfaceVersion() const { return FMetasoundFrontendVersion::GetInvalid(); }
+
+			UE_DEPRECATED(5.8, "Controllers are being actively deprecated, and this call was not thread safe in certain contexts")
 			virtual bool IsInterfaceMember() const override { return false; }
 
 #if WITH_EDITORONLY_DATA
diff --git a/Engine/Plugins/Runtime/Metasound/Source/MetasoundFrontend/Private/MetasoundFrontendNodeClassRegistry.cpp b/Engine/Plugins/Runtime/Metasound/Source/MetasoundFrontend/Private/MetasoundFrontendNodeClassRegistry.cpp
--- a/Engine/Plugins/Runtime/Metasound/Source/MetasoundFrontend/Private/MetasoundFrontendNodeClassRegistry.cpp
+++ b/Engine/Plugins/Runtime/Metasound/Source/MetasoundFrontend/Private/MetasoundFrontendNodeClassRegistry.cpp
@@ -794,9 +794,6 @@ PRAGMA_ENABLE_DEPRECATION_WARNINGS
 		METASOUND_TRACE_CPUPROFILER_EVENT_SCOPE_TEXT(*InRegistryKey.ToString(TEXT("FNodeClassRegistry::UnregisterGraph")));
 #endif // (UE_BUILD_DEBUG || UE_BUILD_DEVELOPMENT)
 
-		check(IsInGameThread());
-		check(InDocumentInterface);
-
 		// Do not attempt to unregister a MetaSound with an invalid registry key
 		if (!InRegistryKey.IsValid())
 		{
diff --git a/Engine/Plugins/Runtime/Metasound/Source/MetasoundFrontend/Private/MetasoundFrontendNodeController.cpp b/Engine/Plugins/Runtime/Metasound/Source/MetasoundFrontend/Private/MetasoundFrontendNodeController.cpp
--- a/Engine/Plugins/Runtime/Metasound/Source/MetasoundFrontend/Private/MetasoundFrontendNodeController.cpp
+++ b/Engine/Plugins/Runtime/Metasound/Source/MetasoundFrontend/Private/MetasoundFrontendNodeController.cpp
@@ -1560,18 +1560,9 @@ PRAGMA_ENABLE_DEPRECATION_WARNINGS
 						const FMetasoundFrontendVertex& Output = Node->Interface.Outputs.Last();
 						for (const FMetasoundFrontendVersion& InterfaceVersion : OwningDocument->GetInterfaceVersions())
 						{
-							FInterfaceRegistryKey InterfaceKey = GetInterfaceRegistryKey(InterfaceVersion);
-							if (const IInterfaceRegistryEntry* Entry = IInterfaceRegistry::Get().FindInterfaceRegistryEntry(InterfaceKey))
+							if (IInterfaceRegistry::Get().ContainsOutput(InterfaceVersion, Output))
 							{
-								auto IsOutput = [&Output](const FMetasoundFrontendClassOutput& InterfaceOutput)
-								{
-									return FMetasoundFrontendVertex::IsFunctionalEquivalent(Output, InterfaceOutput);
-								};
-
-								if (Entry->GetInterface().Outputs.ContainsByPredicate(IsOutput))
-								{
-									return InterfaceVersion;
-								}
+								return InterfaceVersion;
 							}
 						}
 					}
@@ -1739,19 +1730,9 @@ PRAGMA_ENABLE_DEPRECATION_WARNINGS
 						for (const FMetasoundFrontendVersion& InterfaceVersion : InterfaceVersions)
 						{
 							// If the node is on the root graph, test if it is in the interfaces required inputs.
-							FMetasoundFrontendInterface Interface;
-							FInterfaceRegistryKey InterfaceKey = GetInterfaceRegistryKey(InterfaceVersion);
-							if (const IInterfaceRegistryEntry* Entry = IInterfaceRegistry::Get().FindInterfaceRegistryEntry(InterfaceKey))
+							if (IInterfaceRegistry::Get().ContainsInput(InterfaceVersion, Input))
 							{
-								auto IsInput = [&Input](const FMetasoundFrontendClassInput& InterfaceInput)
-								{
-									return FMetasoundFrontendVertex::IsFunctionalEquivalent(Input, InterfaceInput);
-								};
-
-								if (Entry->GetInterface().Inputs.ContainsByPredicate(IsInput))
-								{
-									return InterfaceVersion;
-								}
+								return InterfaceVersion;
 							}
 						}
 					}
diff --git a/Engine/Plugins/Runtime/Metasound/Source/MetasoundFrontend/Private/MetasoundFrontendSearchEngineCore.cpp b/Engine/Plugins/Runtime/Metasound/Source/MetasoundFrontend/Private/MetasoundFrontendSearchEngineCore.cpp
--- a/Engine/Plugins/Runtime/Metasound/Source/MetasoundFrontend/Private/MetasoundFrontendSearchEngineCore.cpp
+++ b/Engine/Plugins/Runtime/Metasound/Source/MetasoundFrontend/Private/MetasoundFrontendSearchEngineCore.cpp
@@ -54,22 +54,21 @@ namespace Metasound::Frontend
 			}
 		}
 
-		FFrontendQueryKey FMapInterfaceRegistryTransactionsToInterfaceRegistryKeys::Map(const FFrontendQueryEntry& InEntry) const
+		FFrontendQueryKey FMapInterfaceRegistryTransactionsToInterfaceRegistryVersions::Map(const FFrontendQueryEntry& InEntry) const
 		{
-			FInterfaceRegistryKey RegistryKey;
+			FMetasoundFrontendVersion Version;
 			if (ensure(InEntry.Value.IsType<FInterfaceRegistryTransaction>()))
 			{
-				RegistryKey = InEntry.Value.Get<FInterfaceRegistryTransaction>().GetInterfaceRegistryKey();
+				Version = InEntry.Value.Get<FInterfaceRegistryTransaction>().GetInterfaceVersion();
 			}
-			return FFrontendQueryKey{RegistryKey};
+			return FFrontendQueryKey { Version.ToString() };
 		}
 
 		FFrontendQueryKey FMapInterfaceRegistryTransactionToInterfaceName::Map(const FFrontendQueryEntry& InEntry) const
 		{
-			FInterfaceRegistryKey RegistryKey;
 			if (ensure(InEntry.Value.IsType<FInterfaceRegistryTransaction>()))
 			{
-				return FFrontendQueryKey{InEntry.Value.Get<FInterfaceRegistryTransaction>().GetInterfaceVersion().Name};
+				return FFrontendQueryKey { InEntry.Value.Get<FInterfaceRegistryTransaction>().GetInterfaceVersion().Name };
 			}
 			return FFrontendQueryKey();
 		}
@@ -113,53 +112,33 @@ namespace Metasound::Frontend
 			return false;
 		}
 
-		void FTransformInterfaceRegistryTransactionToInterface::Transform(FFrontendQueryEntry::FValue& InValue) const
-		{
-			FMetasoundFrontendInterface Interface;
-			if (ensure(InValue.IsType<FInterfaceRegistryTransaction>()))
-			{
-				FInterfaceRegistryKey RegistryKey = InValue.Get<FInterfaceRegistryTransaction>().GetInterfaceRegistryKey();
-				IInterfaceRegistry::Get().FindInterface(RegistryKey, Interface);
-			}
-			InValue.Set<FMetasoundFrontendInterface>(MoveTemp(Interface));
-		}
-
-		void FTransformInterfaceToInterfaceVersion::Transform(FFrontendQueryEntry::FValue& InValue) const
-		{
-			const FMetasoundFrontendInterface* Interface = InValue.TryGet<FMetasoundFrontendInterface>();
-			if (ensure(Interface))
-			{
-				InValue.Set<FMetasoundFrontendVersion>(Interface->Metadata.Version);
-			}
-		}
-
-		FFrontendQueryKey FMapInterfaceToInterfaceNameAndMajorVersion::GetKey(const FString& InName, int32 InMajorVersion)
+		FFrontendQueryKey FMapVersionToNameAndMajorVersion::GetKey(const FString& InName, int32 InMajorVersion)
 		{
 			return FFrontendQueryKey{FString::Format(TEXT("{0}_{1}"), { InName, InMajorVersion })};
 		}
 
-		FFrontendQueryKey FMapInterfaceToInterfaceNameAndMajorVersion::Map(const FFrontendQueryEntry& InEntry) const
+		FFrontendQueryKey FMapVersionToNameAndMajorVersion::Map(const FFrontendQueryEntry& InEntry) const
 		{
 			FString Key;
-			if (ensure(InEntry.Value.IsType<FMetasoundFrontendInterface>()))
+			if (ensure(InEntry.Value.IsType<FMetasoundFrontendVersion>()))
 			{
-				const FMetasoundFrontendInterface& Interface = InEntry.Value.Get<FMetasoundFrontendInterface>();
+				const FMetasoundFrontendVersion& InVersion = InEntry.Value.Get<FMetasoundFrontendVersion>();
 
-				return FMapInterfaceToInterfaceNameAndMajorVersion::GetKey(Interface.Metadata.Version.Name.ToString(), Interface.Metadata.Version.Number.Major);
+				return FMapVersionToNameAndMajorVersion::GetKey(InVersion.Name.ToString(), InVersion.Number.Major);
 			}
 			return FFrontendQueryKey{Key};
 		}
 
-		void FReduceInterfacesToHighestVersion::Reduce(const FFrontendQueryKey& InKey, FFrontendQueryPartition& InOutEntries) const
+		void FReduceToHighestVersion::Reduce(const FFrontendQueryKey& InKey, FFrontendQueryPartition& InOutEntries) const
 		{
 			FFrontendQueryEntry* HighestVersionEntry = nullptr;
 			FMetasoundFrontendVersionNumber HighestVersion = FMetasoundFrontendVersionNumber::GetInvalid();
 
 			for (FFrontendQueryEntry& Entry : InOutEntries)
 			{
-				if (ensure(Entry.Value.IsType<FMetasoundFrontendInterface>()))
+				if (ensure(Entry.Value.IsType<FMetasoundFrontendVersion>()))
 				{
-					const FMetasoundFrontendVersionNumber& VersionNumber = Entry.Value.Get<FMetasoundFrontendInterface>().Metadata.Version.Number;
+					const FMetasoundFrontendVersionNumber& VersionNumber = Entry.Value.Get<FMetasoundFrontendVersion>().Number;
 					if (VersionNumber > HighestVersion)
 					{
 						HighestVersionEntry = &Entry;
@@ -183,17 +162,22 @@ namespace Metasound::Frontend
 		TArray<FFrontendQueryKey> FMapUClassToDefaultInterface::Map(const FFrontendQueryEntry& InEntry) const
 		{
 			TArray<FFrontendQueryKey> Keys;
-			if (const FMetasoundFrontendInterface* Interface = InEntry.Value.TryGet<FMetasoundFrontendInterface>())
+			if (const FMetasoundFrontendVersion* InterfaceVersion = InEntry.Value.TryGet<FMetasoundFrontendVersion>())
 			{
-				auto IsDefaultInterface = [](const FMetasoundFrontendInterfaceUClassOptions& InOptions)
-				{
-					return InOptions.bIsDefault;
-				};
-				auto GetUClassQueryKey = [](const FMetasoundFrontendInterfaceUClassOptions& InOptions)
+				FMetasoundFrontendInterface Interface;
+				const bool bFound = FInterfaceRegistry::Get().FindInterface(*InterfaceVersion, Interface);
+				if (ensure(bFound))
 				{
-					return CreateKeyFromUClassPathName(InOptions.ClassPath);
-				};
-				Algo::TransformIf(Interface->Metadata.UClassOptions, Keys, IsDefaultInterface, GetUClassQueryKey);
+					auto IsDefaultInterface = [](const FMetasoundFrontendInterfaceUClassOptions& InOptions)
+					{
+						return InOptions.bIsDefault;
+					};
+					auto GetUClassQueryKey = [](const FMetasoundFrontendInterfaceUClassOptions& InOptions)
+					{
+						return CreateKeyFromUClassPathName(InOptions.ClassPath);
+					};
+					Algo::TransformIf(Interface.Metadata.UClassOptions, Keys, IsDefaultInterface, GetUClassQueryKey);
+				}
 			}
 
 			return Keys;
@@ -229,14 +213,14 @@ namespace Metasound::Frontend
 			}
 		};
 
-		FFrontendQueryKey FMapInterfaceToInterfaceName::Map(const FFrontendQueryEntry& InEntry) const
+		FFrontendQueryKey FMapVersionToName::Map(const FFrontendQueryEntry& InEntry) const
 		{
 			FName Key;
-			if (ensure(InEntry.Value.IsType<FMetasoundFrontendInterface>()))
+			if (ensure(InEntry.Value.IsType<FMetasoundFrontendVersion>()))
 			{
-				Key = InEntry.Value.Get<FMetasoundFrontendInterface>().Metadata.Version.Name;
+				Key = InEntry.Value.Get<FMetasoundFrontendVersion>().Name;
 			}
-			return FFrontendQueryKey{Key};
+			return FFrontendQueryKey { Key };
 		}
 
 		void FTransformInterfaceRegistryTransactionToInterfaceVersion::Transform(FFrontendQueryEntry::FValue& InValue) const
@@ -270,25 +254,14 @@ namespace Metasound::Frontend
 			return MetaSoundClass;
 		}
 
-		TArray<FMetasoundFrontendInterface> BuildArrayOfInterfacesFromPartition(const FFrontendQueryPartition& InPartition)
+		FMetasoundFrontendVersion BuildSingleVersionFromPartition(const FFrontendQueryPartition& InPartition)
 		{
-			TArray<FMetasoundFrontendInterface> Result;
-			for (const FFrontendQueryEntry& Entry : InPartition)
-			{
-				check(Entry.Value.IsType<FMetasoundFrontendInterface>());
-				Result.Add(Entry.Value.Get<FMetasoundFrontendInterface>());
-			}
-			return Result;
-		}
-
-		FMetasoundFrontendInterface BuildSingleInterfaceFromPartition(const FFrontendQueryPartition& InPartition)
-		{
-			FMetasoundFrontendInterface MetaSoundInterface;
+			FMetasoundFrontendVersion Version;
 			if (ensureMsgf(InPartition.Num() > 0, TEXT("Cannot retrieve class from empty partition")))
 			{
-				MetaSoundInterface = InPartition.CreateConstIterator()->Value.Get<FMetasoundFrontendInterface>();
+				Version = InPartition.CreateConstIterator()->Value.Get<FMetasoundFrontendVersion>();
 			}
-			return MetaSoundInterface;
+			return Version;
 		}
 
 		TArray<FMetasoundFrontendVersion> BuildArrayOfVersionsFromPartition(const FFrontendQueryPartition& InPartition)
@@ -346,7 +319,7 @@ namespace Metasound::Frontend
 		using namespace SearchEngineQuerySteps;
 		FFrontendQuery Query;
 		Query.AddStep<FInterfaceRegistryTransactionSource>()
-			.AddStep<FMapInterfaceRegistryTransactionsToInterfaceRegistryKeys>()
+			.AddStep<FMapInterfaceRegistryTransactionsToInterfaceRegistryVersions>()
 			.AddStep<FReduceInterfaceRegistryTransactionsToCurrentStatus>()
 			.AddStep<FMapInterfaceRegistryTransactionToInterfaceName>()
 			.AddStep<FTransformInterfaceRegistryTransactionToInterfaceVersion>();
@@ -359,23 +332,23 @@ namespace Metasound::Frontend
 		return BuildArrayOfVersionsFromPartition(InPartition);
 	}
 
-	FFrontendQuery FFindInterfaceWithHighestVersionQueryPolicy::CreateQuery()
+	FFrontendQuery FFindHighestVersionOfInterfaceQueryPolicy::CreateQuery()
 	{
 		using namespace SearchEngineQuerySteps;
 		FFrontendQuery Query;
 		Query.AddStep<FInterfaceRegistryTransactionSource>()
-			.AddStep<FMapInterfaceRegistryTransactionsToInterfaceRegistryKeys>()
+			.AddStep<FMapInterfaceRegistryTransactionsToInterfaceRegistryVersions>()
 			.AddStep<FReduceInterfaceRegistryTransactionsToCurrentStatus>()
-			.AddStep<FTransformInterfaceRegistryTransactionToInterface>()
-			.AddStep<FMapInterfaceToInterfaceName>()
-			.AddStep<FReduceInterfacesToHighestVersion>();
+			.AddStep<FTransformInterfaceRegistryTransactionToInterfaceVersion>()
+			.AddStep<FMapVersionToName>()
+			.AddStep<FReduceToHighestVersion>();
 		return Query;
 	}
 
-	FMetasoundFrontendInterface FFindInterfaceWithHighestVersionQueryPolicy::BuildResult(const FFrontendQueryPartition& InPartition)
+	FMetasoundFrontendVersion FFindHighestVersionOfInterfaceQueryPolicy::BuildResult(const FFrontendQueryPartition& InPartition)
 	{
 		using namespace SearchEngineQuerySteps;
-		return BuildSingleInterfaceFromPartition(InPartition);
+		return BuildSingleVersionFromPartition(InPartition);
 	}
 
 	FFrontendQuery FFindAllInterfacesQueryPolicy::CreateQuery()
@@ -383,19 +356,19 @@ namespace Metasound::Frontend
 		using namespace SearchEngineQuerySteps;
 		FFrontendQuery Query;
 		Query.AddStep<FInterfaceRegistryTransactionSource>()
-			.AddStep<FMapInterfaceRegistryTransactionsToInterfaceRegistryKeys>()
+			.AddStep<FMapInterfaceRegistryTransactionsToInterfaceRegistryVersions>()
 			.AddStep<FReduceInterfaceRegistryTransactionsToCurrentStatus>()
-			.AddStep<FTransformInterfaceRegistryTransactionToInterface>()
-			.AddStep<FMapInterfaceToInterfaceNameAndMajorVersion>()
-			.AddStep<FReduceInterfacesToHighestVersion>()
+			.AddStep<FTransformInterfaceRegistryTransactionToInterfaceVersion>()
+			.AddStep<FMapVersionToNameAndMajorVersion>()
+			.AddStep<FReduceToHighestVersion>()
 			.AddStep<FMapToNull>();
 		return Query;
 	}
 
-	TArray<FMetasoundFrontendInterface> FFindAllInterfacesQueryPolicy::BuildResult(const FFrontendQueryPartition& InPartition)
+	TArray<FMetasoundFrontendVersion> FFindAllInterfacesQueryPolicy::BuildResult(const FFrontendQueryPartition& InPartition)
 	{
 		using namespace SearchEngineQuerySteps;
-		return BuildArrayOfInterfacesFromPartition(InPartition);
+		return BuildArrayOfVersionsFromPartition(InPartition);
 	}
 
 	FFrontendQuery FFindAllInterfacesIncludingAllVersionsQueryPolicy::CreateQuery()
@@ -403,17 +376,17 @@ namespace Metasound::Frontend
 		using namespace SearchEngineQuerySteps;
 		FFrontendQuery Query;
 		Query.AddStep<FInterfaceRegistryTransactionSource>()
-			.AddStep<FMapInterfaceRegistryTransactionsToInterfaceRegistryKeys>()
+			.AddStep<FMapInterfaceRegistryTransactionsToInterfaceRegistryVersions>()
 			.AddStep<FReduceInterfaceRegistryTransactionsToCurrentStatus>()
-			.AddStep<FTransformInterfaceRegistryTransactionToInterface>()
+			.AddStep<FTransformInterfaceRegistryTransactionToInterfaceVersion>()
 			.AddStep<FMapToNull>();
 		return Query;
 	}
 
-	TArray<FMetasoundFrontendInterface> FFindAllInterfacesIncludingAllVersionsQueryPolicy::BuildResult(const FFrontendQueryPartition& InPartition)
+	TArray<FMetasoundFrontendVersion> FFindAllInterfacesIncludingAllVersionsQueryPolicy::BuildResult(const FFrontendQueryPartition& InPartition)
 	{
 		using namespace SearchEngineQuerySteps;
-		return BuildArrayOfInterfacesFromPartition(InPartition);
+		return BuildArrayOfVersionsFromPartition(InPartition);
 	}
 
 	FFrontendQuery FFindAllDefaultInterfaceVersionsForUClassQueryPolicy::CreateQuery()
@@ -421,11 +394,10 @@ namespace Metasound::Frontend
 		using namespace SearchEngineQuerySteps;
 		FFrontendQuery Query;
 		Query.AddStep<FInterfaceRegistryTransactionSource>()
-			.AddStep<FMapInterfaceRegistryTransactionsToInterfaceRegistryKeys>()
+			.AddStep<FMapInterfaceRegistryTransactionsToInterfaceRegistryVersions>()
 			.AddStep<FReduceInterfaceRegistryTransactionsToCurrentStatus>()
-			.AddStep<FTransformInterfaceRegistryTransactionToInterface>()
-			.AddStep<FMapUClassToDefaultInterface>()
-			.AddStep<FTransformInterfaceToInterfaceVersion>();
+			.AddStep<FTransformInterfaceRegistryTransactionToInterfaceVersion>()
+			.AddStep<FMapUClassToDefaultInterface>();
 		return Query;
 	}
 
@@ -438,17 +410,17 @@ namespace Metasound::Frontend
 	void FSearchEngineCore::Prime()
 	{
 		METASOUND_LLM_SCOPE;
-		METASOUND_TRACE_CPUPROFILER_EVENT_SCOPE(metasound::FSearchEngineCore::Prime);
+		METASOUND_TRACE_CPUPROFILER_EVENT_SCOPE(MetaSound::FSearchEngineCore::Prime);
 
 		FindClassWithHighestMinorVersionQuery.Prime();
 		FindAllRegisteredInterfacesWithNameQuery.Prime();
-		FindInterfaceWithHighestVersionQuery.Prime();
+		FindHighestVersionOfInterfaceQueryPolicy.Prime();
 		FindAllDefaultInterfaceVersionsForUClassQuery.Prime();
 	}
 
 	bool FSearchEngineCore::FindRegisteredClass(const FMetasoundFrontendClassName& InName, int32 InMajorVersion, FMetaSoundClassInfo& OutClass)
 	{
-		METASOUND_TRACE_CPUPROFILER_EVENT_SCOPE(metasound::FSearchEngineCore::FindRegisteredClass);
+		METASOUND_TRACE_CPUPROFILER_EVENT_SCOPE(MetaSound::FSearchEngineCore::FindRegisteredClass);
 		using namespace SearchEngineQuerySteps;
 
 		const FFrontendQueryKey QueryKey = CreateKeyFromFullClassNameAndMajorVersion(InName, InMajorVersion);
@@ -473,7 +445,7 @@ namespace Metasound::Frontend
 	{
 		using namespace SearchEngineQuerySteps;
 
-		METASOUND_TRACE_CPUPROFILER_EVENT_SCOPE(metasound::FSearchEngineCore::FindUClassDefaultInterfaceVersions);
+		METASOUND_TRACE_CPUPROFILER_EVENT_SCOPE(MetaSound::FSearchEngineCore::FindUClassDefaultInterfaceVersions);
 
 		const FFrontendQueryKey Key = CreateKeyFromUClassPathName(InUClassPath);
 		return FindAllDefaultInterfaceVersionsForUClassQuery.UpdateAndFindResult(Key);
@@ -481,17 +453,32 @@ namespace Metasound::Frontend
 
 	TArray<FMetasoundFrontendVersion> FSearchEngineCore::FindAllRegisteredInterfacesWithName(FName InInterfaceName)
 	{
-		METASOUND_TRACE_CPUPROFILER_EVENT_SCOPE(metasound::FSearchEngineCore::FindAllRegisteredInterfacesWithName);
+		METASOUND_TRACE_CPUPROFILER_EVENT_SCOPE(MetaSound::FSearchEngineCore::FindAllRegisteredInterfacesWithName);
 
 		const FFrontendQueryKey Key{InInterfaceName};
 		return FindAllRegisteredInterfacesWithNameQuery.UpdateAndFindResult(Key);
 	}
 
+	bool FSearchEngineCore::FindHighestInterfaceVersion(FName InInterfaceName, FMetasoundFrontendVersion& OutVersion)
+	{
+		METASOUND_TRACE_CPUPROFILER_EVENT_SCOPE(MetaSound::FSearchEngineCore::FindHighestInterfaceVersion);
+
+		const FFrontendQueryKey Key { InInterfaceName };
+
+		return FindHighestVersionOfInterfaceQueryPolicy.UpdateAndFindResult(Key, OutVersion);
+	}
+
 	bool FSearchEngineCore::FindInterfaceWithHighestVersion(FName InInterfaceName, FMetasoundFrontendInterface& OutInterface)
 	{
-		METASOUND_TRACE_CPUPROFILER_EVENT_SCOPE(metasound::FSearchEngineCore::FindInterfaceWithHighestVersion);
+		METASOUND_TRACE_CPUPROFILER_EVENT_SCOPE(MetaSound::FSearchEngineCore::FindInterfaceWithHighestVersion);
 
-		const FFrontendQueryKey Key{InInterfaceName};
-		return FindInterfaceWithHighestVersionQuery.UpdateAndFindResult(Key, OutInterface);
+		FMetasoundFrontendVersion Version;
+		const bool bVersionFound = FindHighestInterfaceVersion(InInterfaceName, Version);
+		if (bVersionFound)
+		{
+			return IInterfaceRegistry::Get().FindInterface(Version, OutInterface);
+		}
+
+		return false;
 	}
 } // namespace Metasound::Frontend
diff --git a/Engine/Plugins/Runtime/Metasound/Source/MetasoundFrontend/Private/MetasoundFrontendSearchEngineCore.h b/Engine/Plugins/Runtime/Metasound/Source/MetasoundFrontend/Private/MetasoundFrontendSearchEngineCore.h
--- a/Engine/Plugins/Runtime/Metasound/Source/MetasoundFrontend/Private/MetasoundFrontendSearchEngineCore.h
+++ b/Engine/Plugins/Runtime/Metasound/Source/MetasoundFrontend/Private/MetasoundFrontendSearchEngineCore.h
@@ -37,7 +37,7 @@ namespace Metasound
 				TUniquePtr<FInterfaceTransactionStream> TransactionStream;
 			};
 
-			class FMapInterfaceRegistryTransactionsToInterfaceRegistryKeys : public IFrontendQueryMapStep
+			class FMapInterfaceRegistryTransactionsToInterfaceRegistryVersions : public IFrontendQueryMapStep
 			{
 			public:
 				virtual FFrontendQueryKey Map(const FFrontendQueryEntry& InEntry) const override;
@@ -64,35 +64,14 @@ namespace Metasound
 				virtual TArray<FFrontendQueryKey> Map(const FFrontendQueryEntry& InEntry) const override;
 			};
 
-			class FReduceInterfaceRegistryTransactionsToDefault : public IFrontendQueryReduceStep
-			{
-			public:
-				virtual void Reduce(const FFrontendQueryKey& InKey, FFrontendQueryPartition& InOutEntries) const override;
-			private:
-				static FInterfaceRegistryTransaction::FTimeType GetTransactionTimestamp(const FFrontendQueryEntry& InEntry);
-				static bool IsValidTransactionOfType(FInterfaceRegistryTransaction::ETransactionType InType, const FFrontendQueryEntry* InEntry);
-			};
-
-			class FTransformInterfaceRegistryTransactionToInterface : public IFrontendQueryTransformStep
-			{
-			public:
-				virtual void Transform(FFrontendQueryEntry::FValue& InValue) const override;
-			};
-
-			class FTransformInterfaceToInterfaceVersion : public IFrontendQueryTransformStep
-			{
-			public:
-				virtual void Transform(FFrontendQueryEntry::FValue& InValue) const override;
-			};
-
-			class FMapInterfaceToInterfaceNameAndMajorVersion : public IFrontendQueryMapStep
+			class FMapVersionToNameAndMajorVersion : public IFrontendQueryMapStep
 			{
 			public:
 				static FFrontendQueryKey GetKey(const FString& InName, int32 InMajorVersion);
 				virtual FFrontendQueryKey Map(const FFrontendQueryEntry& InEntry) const override;
 			};
 
-			class FReduceInterfacesToHighestVersion : public IFrontendQueryReduceStep
+			class FReduceToHighestVersion : public IFrontendQueryReduceStep
 			{
 			public:
 				virtual void Reduce(const FFrontendQueryKey& InKey, FFrontendQueryPartition& InOutEntries) const override;
@@ -106,7 +85,7 @@ namespace Metasound
 				virtual bool Filter(const FFrontendQueryEntry& InEntry) const override;
 			};
 
-			class FMapInterfaceToInterfaceName : public IFrontendQueryMapStep
+			class FMapVersionToName : public IFrontendQueryMapStep
 			{
 			public:
 				virtual FFrontendQueryKey Map(const FFrontendQueryEntry& InEntry) const override;
@@ -123,10 +102,7 @@ namespace Metasound
 
 			FMetaSoundClassInfo BuildSingleClassFromPartition(const FFrontendQueryPartition& InPartition);
 
-			// Convert an array of partitions into an array of metasound frontend interfaces.
-			TArray<FMetasoundFrontendInterface> BuildArrayOfInterfacesFromPartition(const FFrontendQueryPartition& InPartition);
-
-			FMetasoundFrontendInterface BuildSingleInterfaceFromPartition(const FFrontendQueryPartition& InPartition);
+			FMetasoundFrontendVersion BuildSingleVersionFromPartition(const FFrontendQueryPartition& InPartition);
 
 			// Convert an array of partitions into an array of metasound frontend versions.
 			TArray<FMetasoundFrontendVersion> BuildArrayOfVersionsFromPartition(const FFrontendQueryPartition& InPartition);
@@ -277,18 +253,18 @@ namespace Metasound
 		};
 
 		// Policy for finding the highest version registered interfaces by name.
-		struct FFindInterfaceWithHighestVersionQueryPolicy
+		struct FFindHighestVersionOfInterfaceQueryPolicy
 		{
-			using ResultType = FMetasoundFrontendInterface;
+			using ResultType = FMetasoundFrontendVersion;
 
 			static FFrontendQuery CreateQuery();
-			static FMetasoundFrontendInterface BuildResult(const FFrontendQueryPartition& InPartition);
+			static FMetasoundFrontendVersion BuildResult(const FFrontendQueryPartition& InPartition);
 		};
 
 		// Policy for finding all registered interfaces. 
 		struct FFindAllInterfacesQueryPolicy
 		{
-			using ResultType = TArray<FMetasoundFrontendInterface>;
+			using ResultType = TArray<FMetasoundFrontendVersion>;
 
 			static FFrontendQuery CreateQuery();
 			static ResultType BuildResult(const FFrontendQueryPartition& InPartition);
@@ -297,10 +273,10 @@ namespace Metasound
 		// Policy for finding all registered interfaces (including deprecated). 
 		struct FFindAllInterfacesIncludingAllVersionsQueryPolicy
 		{
-			using ResultType = TArray<FMetasoundFrontendInterface>;
+			using ResultType = TArray<FMetasoundFrontendVersion>;
 
 			static FFrontendQuery CreateQuery();
-			static TArray<FMetasoundFrontendInterface> BuildResult(const FFrontendQueryPartition& InPartition);
+			static TArray<FMetasoundFrontendVersion> BuildResult(const FFrontendQueryPartition& InPartition);
 		};
 
 		// Policy for finding all default interface versions for a given UClass
@@ -330,13 +306,15 @@ namespace Metasound
 
 			virtual TArray<FMetasoundFrontendVersion> FindUClassDefaultInterfaceVersions(const FTopLevelAssetPath& InUClassPath) override;
 			virtual TArray<FMetasoundFrontendVersion> FindAllRegisteredInterfacesWithName(FName InInterfaceName) override;
+
+			virtual bool FindHighestInterfaceVersion(FName InInterfaceName, FMetasoundFrontendVersion& OutVersion) override;
 			virtual bool FindInterfaceWithHighestVersion(FName InInterfaceName, FMetasoundFrontendInterface& OutInterface) override;
 
 		private:
 			TSearchEngineQuery<FFindClassWithHighestMinorVersionQueryPolicy> FindClassWithHighestMinorVersionQuery;
 			TSearchEngineQuery<FFindAllRegisteredInterfacesWithNameQueryPolicy> FindAllRegisteredInterfacesWithNameQuery;
 			TSearchEngineQuery<FFindAllDefaultInterfaceVersionsForUClassQueryPolicy> FindAllDefaultInterfaceVersionsForUClassQuery;
-			TSearchEngineQuery<FFindInterfaceWithHighestVersionQueryPolicy> FindInterfaceWithHighestVersionQuery;
+			TSearchEngineQuery<FFindHighestVersionOfInterfaceQueryPolicy> FindHighestVersionOfInterfaceQueryPolicy;
 		};
 	}
 }
diff --git a/Engine/Plugins/Runtime/Metasound/Source/MetasoundFrontend/Private/MetasoundFrontendSearchEngineEditorOnly.cpp b/Engine/Plugins/Runtime/Metasound/Source/MetasoundFrontend/Private/MetasoundFrontendSearchEngineEditorOnly.cpp
--- a/Engine/Plugins/Runtime/Metasound/Source/MetasoundFrontend/Private/MetasoundFrontendSearchEngineEditorOnly.cpp
+++ b/Engine/Plugins/Runtime/Metasound/Source/MetasoundFrontend/Private/MetasoundFrontendSearchEngineEditorOnly.cpp
@@ -253,9 +253,9 @@ namespace Metasound::Frontend
 		return false;
 	}
 
-	TArray<FMetasoundFrontendInterface> FSearchEngineEditorOnly::FindAllInterfaces(bool bInIncludeAllVersions)
+	TArray<FMetasoundFrontendVersion> FSearchEngineEditorOnly::FindAllInterfaceVersions(bool bInIncludeAllVersions)
 	{
-		METASOUND_TRACE_CPUPROFILER_EVENT_SCOPE(Metasound::FSearchEngineEditorOnly::FindAllInterfaces);
+		METASOUND_TRACE_CPUPROFILER_EVENT_SCOPE(Metasound::FSearchEngineEditorOnly::FindAllInterfaceVersions);
 
 		const FFrontendQueryKey NullKey;
 		if (bInIncludeAllVersions)
diff --git a/Engine/Plugins/Runtime/Metasound/Source/MetasoundFrontend/Private/MetasoundFrontendSearchEngineEditorOnly.h b/Engine/Plugins/Runtime/Metasound/Source/MetasoundFrontend/Private/MetasoundFrontendSearchEngineEditorOnly.h
--- a/Engine/Plugins/Runtime/Metasound/Source/MetasoundFrontend/Private/MetasoundFrontendSearchEngineEditorOnly.h
+++ b/Engine/Plugins/Runtime/Metasound/Source/MetasoundFrontend/Private/MetasoundFrontendSearchEngineEditorOnly.h
@@ -96,7 +96,7 @@ namespace Metasound::Frontend
 		virtual TArray<FMetaSoundClassInfo> FindClassesWithName(const FMetasoundFrontendClassName& InName, ESortByVersion SortByVersion = ESortByVersion::Yes) override;
 		virtual TArray<FMetasoundFrontendClass> FindClassesWithName(const FMetasoundFrontendClassName& InName, bool bInSortByVersion) override;
 		virtual bool FindClassWithHighestVersion(const FMetasoundFrontendClassName& InName, FMetaSoundClassInfo& OutClassInfo) override;
-		virtual TArray<FMetasoundFrontendInterface> FindAllInterfaces(bool bInIncludeAllVersions) override;
+		virtual TArray<FMetasoundFrontendVersion> FindAllInterfaceVersions(bool bInIncludeAllVersions) override;
 
 	private:
 		TSearchEngineQuery<FFindAllAssetClassInfoQueryPolicy> FindAllAssetClassInfoQuery;
diff --git a/Engine/Plugins/Runtime/Metasound/Source/MetasoundFrontend/Private/MetasoundFrontendTransform.cpp b/Engine/Plugins/Runtime/Metasound/Source/MetasoundFrontend/Private/MetasoundFrontendTransform.cpp
--- a/Engine/Plugins/Runtime/Metasound/Source/MetasoundFrontend/Private/MetasoundFrontendTransform.cpp
+++ b/Engine/Plugins/Runtime/Metasound/Source/MetasoundFrontend/Private/MetasoundFrontendTransform.cpp
@@ -92,7 +92,7 @@ namespace Metasound
 			Algo::Transform(InInterfaceVersionsToRemove, InterfacesToRemove, [](const FMetasoundFrontendVersion& Version)
 			{
 				FMetasoundFrontendInterface Interface;
-				const bool bFromInterfaceFound = IInterfaceRegistry::Get().FindInterface(GetInterfaceRegistryKey(Version), Interface);
+				const bool bFromInterfaceFound = IInterfaceRegistry::Get().FindInterface(Version, Interface);
 				if (!ensureAlways(bFromInterfaceFound))
 				{
 					METASOUND_VERSIONING_LOG(Error, TEXT("Failed to find interface '%s' to remove"), *Version.ToString());
@@ -103,7 +103,7 @@ namespace Metasound
 			Algo::Transform(InInterfaceVersionsToAdd, InterfacesToAdd, [](const FMetasoundFrontendVersion& Version)
 			{
 				FMetasoundFrontendInterface Interface;
-				const bool bToInterfaceFound = IInterfaceRegistry::Get().FindInterface(GetInterfaceRegistryKey(Version), Interface);
+				const bool bToInterfaceFound = IInterfaceRegistry::Get().FindInterface(Version, Interface);
 				if (!ensureAlways(bToInterfaceFound))
 				{
 					METASOUND_VERSIONING_LOG(Error, TEXT("Failed to find interface '%s' to add"), *Version.ToString());
diff --git a/Engine/Plugins/Runtime/Metasound/Source/MetasoundFrontend/Public/Interfaces/MetasoundFrontendInterfaceRegistry.h b/Engine/Plugins/Runtime/Metasound/Source/MetasoundFrontend/Public/Interfaces/MetasoundFrontendInterfaceRegistry.h
--- a/Engine/Plugins/Runtime/Metasound/Source/MetasoundFrontend/Public/Interfaces/MetasoundFrontendInterfaceRegistry.h
+++ b/Engine/Plugins/Runtime/Metasound/Source/MetasoundFrontend/Public/Interfaces/MetasoundFrontendInterfaceRegistry.h
@@ -4,18 +4,29 @@
 
 #include "MetasoundFrontendController.h"
 #include "MetasoundFrontendDocument.h"
+#include "Templates/SharedPointer.h"
 #include "UObject/NoExportTypes.h"
 
+
 #define UE_API METASOUNDFRONTEND_API
 
 namespace Metasound::Frontend
 {
+	// Forward Declarations
+	class IBuilderVersionTransform;
+
+	// Interface keys are actively being deprecated: Frontend Version is now directly used as hashable interface identifier")
 	using FInterfaceRegistryKey = FString;
 	using FRegistryTransactionID = int32;
 
-	METASOUNDFRONTEND_API bool IsValidInterfaceRegistryKey(const FInterfaceRegistryKey& InKey);
-	METASOUNDFRONTEND_API FInterfaceRegistryKey GetInterfaceRegistryKey(const FMetasoundFrontendVersion& InInterfaceVersion);
-	METASOUNDFRONTEND_API FInterfaceRegistryKey GetInterfaceRegistryKey(const FMetasoundFrontendInterface& InInterface);
+	UE_DEPRECATED(5.8, "Interface keys are deprecated: Frontend Version is now directly used as hashable interface identifier")
+	UE_API bool IsValidInterfaceRegistryKey(const FInterfaceRegistryKey& InKey);
+
+	UE_DEPRECATED(5.8, "Interface keys are deprecated: Frontend Version is now directly used as hashable interface identifier")
+	UE_API FInterfaceRegistryKey GetInterfaceRegistryKey(const FMetasoundFrontendVersion& InInterfaceVersion);
+
+	UE_DEPRECATED(5.8, "Interface keys are deprecated: Frontend Version is now directly used as hashable interface identifier")
+	UE_API FInterfaceRegistryKey GetInterfaceRegistryKey(const FMetasoundFrontendInterface& InInterface);
 
 	class IInterfaceRegistryEntry
 	{
@@ -31,11 +42,15 @@ namespace Metasound::Frontend
 		// Name of routing system used to update interface inputs (ex. ParameterInterface or DataReference).
 		virtual FName GetRouterName() const = 0;
 
-		// How to update a given document if versioning is required to this interface from a deprecated version.
-		UE_DEPRECATED(5.7, "Please implement UpdateRootGraphInterface(FMetaSoundFrontendDocumentBuilder& InDocumentBuilder) instead")
+		UE_DEPRECATED(5.7, "Transform accessor must be thread safe, and thus a shared pointer is passed via 'GetUpdateTransform'")
 		virtual bool UpdateRootGraphInterface(FDocumentHandle InDocument) const { return false; };
 
-		virtual bool UpdateRootGraphInterface(FMetaSoundFrontendDocumentBuilder& InDocumentBuilder) const = 0;
+		UE_DEPRECATED(5.8, "Transform accessor must be thread safe, and thus a shared pointer is passed via 'GetUpdateTransform'")
+		virtual bool UpdateRootGraphInterface(FMetaSoundFrontendDocumentBuilder& InDocumentBuilder) const { return false; }
+
+		// Returns transform indicating how to update a given document if versioning is required to this interface from a deprecated version.
+		// If not specified, update simply adds the given interface as a new interface to the given builder's document.
+		virtual TSharedRef<IBuilderVersionTransform> GetUpdateTransform() const = 0;
 	};
 
 	class FInterfaceRegistryTransaction
@@ -51,17 +66,21 @@ namespace Metasound::Frontend
 			Invalid
 		};
 
+		UE_DEPRECATED(5.8, "Interface keys are deprecated: Frontend Version is now directly used as hashable interface identifier")
 		UE_API FInterfaceRegistryTransaction(ETransactionType InType, const FInterfaceRegistryKey& InKey, const FMetasoundFrontendVersion& InInterfaceVersion, FTimeType InTimestamp);
 
+		UE_API FInterfaceRegistryTransaction(ETransactionType InType, const FMetasoundFrontendVersion& InInterfaceVersion, FTimeType InTimestamp);
+
 		UE_API ETransactionType GetTransactionType() const;
 		UE_API const FMetasoundFrontendVersion& GetInterfaceVersion() const;
+
+		UE_DEPRECATED(5.8, "Interface keys are deprecated: Frontend Version is now directly used as hashable interface identifier")
 		UE_API const FInterfaceRegistryKey& GetInterfaceRegistryKey() const;
+
 		UE_API FTimeType GetTimestamp() const;
 
 	private:
-
 		ETransactionType Type;
-		FInterfaceRegistryKey Key;
 		FMetasoundFrontendVersion InterfaceVersion;
 		FTimeType Timestamp;
 	};
@@ -76,11 +95,26 @@ namespace Metasound::Frontend
 		// Register an interface
 		virtual bool RegisterInterface(TUniquePtr<IInterfaceRegistryEntry>&& InEntry) = 0;
 
-		// Find an interface entry with the given key. Returns null if entry not found with given key.
-		virtual const IInterfaceRegistryEntry* FindInterfaceRegistryEntry(const FInterfaceRegistryKey& InKey) const = 0;
+		virtual bool ContainsInput(const FMetasoundFrontendVersion& InVersion, const FMetasoundFrontendVertex& InInput) const = 0;
+
+		virtual bool ContainsInterface(const FMetasoundFrontendVersion& InVersion) const = 0;
+
+		virtual bool ContainsOutput(const FMetasoundFrontendVersion& InVersion, const FMetasoundFrontendVertex& InOutput) const = 0;
+
+		UE_DEPRECATED(5.8, "Interface keys are deprecated: Frontend Version is now directly used as hashable interface identifier. In addition, direct access to entry is not thread safe, use FindInterface call instead.")
+		virtual const IInterfaceRegistryEntry* FindInterfaceRegistryEntry(const FInterfaceRegistryKey& InKey) const { return nullptr; }
+
+		UE_DEPRECATED(5.8, "Interface keys are deprecated: Frontend Version is now directly used as hashable interface identifier")
+		virtual bool FindInterface(const FInterfaceRegistryKey& InKey, FMetasoundFrontendInterface& OutInterface) const { return false; }
+
+		// Find an interface with the given version. Returns true if interface is found, false if not.
+		virtual bool FindInterface(const FMetasoundFrontendVersion& InVersion, FMetasoundFrontendInterface& OutInterface) const = 0;
+
+		// Returns the name of the router for the given interface.
+		virtual bool FindInterfaceRouter(const FMetasoundFrontendVersion& InVersion, FName& OutRouter) const = 0;
 
-		// Find an interface with the given key. Returns true if interface is found, false if not.
-		virtual bool FindInterface(const FInterfaceRegistryKey& InKey, FMetasoundFrontendInterface& OutInterface) const = 0;
+		// Find the interface update transform associated with the given version if found.
+		virtual TSharedPtr<IBuilderVersionTransform> FindInterfaceUpdateTransform(const FMetasoundFrontendVersion& InVersion) const = 0;
 	};
 } // namespace Metasound::Frontend
 
diff --git a/Engine/Plugins/Runtime/Metasound/Source/MetasoundFrontend/Public/MetasoundAssetBase.h b/Engine/Plugins/Runtime/Metasound/Source/MetasoundFrontend/Public/MetasoundAssetBase.h
--- a/Engine/Plugins/Runtime/Metasound/Source/MetasoundFrontend/Public/MetasoundAssetBase.h
+++ b/Engine/Plugins/Runtime/Metasound/Source/MetasoundFrontend/Public/MetasoundAssetBase.h
@@ -236,8 +236,8 @@ private:
 	UE_API bool VersionDependencies(FMetaSoundFrontendDocumentBuilder& Builder, bool bInLogWarningsOnDroppedConnection);
 
 	// Returns new interface to be versioned to from the given version. If no interface versioning is
-	// required, returns invalid interface (interface with no name and invalid version number).
-	UE_API FMetasoundFrontendInterface GetInterfaceToVersion(const FMetasoundFrontendVersion& InterfaceVersion) const;
+	// required, returns invalid interface version (no name and invalid version number).
+	FMetasoundFrontendVersion GetTargetInterfaceVersion(const FMetasoundFrontendVersion& InterfaceVersion) const;
 
 #if WITH_EDITORONLY_DATA
 	bool bVersionedOnLoad = false;
diff --git a/Engine/Plugins/Runtime/Metasound/Source/MetasoundFrontend/Public/MetasoundFrontendController.h b/Engine/Plugins/Runtime/Metasound/Source/MetasoundFrontend/Public/MetasoundFrontendController.h
--- a/Engine/Plugins/Runtime/Metasound/Source/MetasoundFrontend/Public/MetasoundFrontendController.h
+++ b/Engine/Plugins/Runtime/Metasound/Source/MetasoundFrontend/Public/MetasoundFrontendController.h
@@ -454,10 +454,7 @@ namespace Metasound
 			virtual FOutputHandle GetOutputWithVertexName(const FVertexName& InName) = 0;
 			virtual FConstOutputHandle GetConstOutputWithVertexName(const FVertexName& InName) const = 0;
 
-			/** Returns true if node is required to satisfy a document interface. */
 			virtual bool IsInterfaceMember() const = 0;
-
-			/** Returns interface version if node is a required member of a given interface, otherwise returns invalid version. */
 			virtual const FMetasoundFrontendVersion& GetInterfaceVersion() const = 0;
 
 			/**
@@ -1072,6 +1069,7 @@ namespace Metasound
 			virtual FGraphHandle AddDuplicateSubgraph(const IGraphController& InGraph) = 0;
 
 			virtual const TSet<FMetasoundFrontendVersion>& GetInterfaceVersions() const = 0;
+
 			virtual void AddInterfaceVersion(const FMetasoundFrontendVersion& InVersion) = 0;
 			virtual void RemoveInterfaceVersion(const FMetasoundFrontendVersion& InVersion) = 0;
 			virtual void ClearInterfaceVersions() = 0;
diff --git a/Engine/Plugins/Runtime/Metasound/Source/MetasoundFrontend/Public/MetasoundFrontendDocumentBuilder.h b/Engine/Plugins/Runtime/Metasound/Source/MetasoundFrontend/Public/MetasoundFrontendDocumentBuilder.h
--- a/Engine/Plugins/Runtime/Metasound/Source/MetasoundFrontend/Public/MetasoundFrontendDocumentBuilder.h
+++ b/Engine/Plugins/Runtime/Metasound/Source/MetasoundFrontend/Public/MetasoundFrontendDocumentBuilder.h
@@ -308,7 +308,10 @@ PRAGMA_ENABLE_EXPERIMENTAL_WARNINGS
 	UE_API UMetaSoundFrontendMemberMetadata* FindMemberMetadata(const FGuid& InMemberID);
 #endif // WITH_EDITORONLY_DATA
 
+	UE_DEPRECATED(5.8, "Direct access to IInterfaceRegistryEntry is deprecated as it is not thread safe. Must now get copy of interface directly from the InterfaceRegistry.")
 	static UE_API bool FindDeclaredInterfaces(const FMetasoundFrontendDocument& InDocument, TArray<const Metasound::Frontend::IInterfaceRegistryEntry*>& OutInterfaces);
+
+	UE_DEPRECATED(5.8, "Direct access to IInterfaceRegistryEntry is deprecated as it is not thread safe. Must now get copy of interface directly from the InterfaceRegistry.")
 	UE_API bool FindDeclaredInterfaces(TArray<const Metasound::Frontend::IInterfaceRegistryEntry*>& OutInterfaces) const;
 
 	UE_API const FMetasoundFrontendClass* FindDependency(const FGuid& InClassID) const;
diff --git a/Engine/Plugins/Runtime/Metasound/Source/MetasoundFrontend/Public/MetasoundFrontendQuery.h b/Engine/Plugins/Runtime/Metasound/Source/MetasoundFrontend/Public/MetasoundFrontendQuery.h
--- a/Engine/Plugins/Runtime/Metasound/Source/MetasoundFrontend/Public/MetasoundFrontendQuery.h
+++ b/Engine/Plugins/Runtime/Metasound/Source/MetasoundFrontend/Public/MetasoundFrontendQuery.h
@@ -250,8 +250,7 @@ namespace Metasound
 		Frontend::FAssetManagerTransaction,
 		Frontend::FInterfaceRegistryTransaction,
 		Frontend::FMetaSoundClassInfo,
-		Frontend::FNodeClassRegistryTransaction,
-		FMetasoundFrontendInterface
+		Frontend::FNodeClassRegistryTransaction
 	>;
 
 	/** FFrontendQueryEntry represents one value in the query. It contains an ID,
diff --git a/Engine/Plugins/Runtime/Metasound/Source/MetasoundFrontend/Public/MetasoundFrontendSearchEngine.h b/Engine/Plugins/Runtime/Metasound/Source/MetasoundFrontend/Public/MetasoundFrontendSearchEngine.h
--- a/Engine/Plugins/Runtime/Metasound/Source/MetasoundFrontend/Public/MetasoundFrontendSearchEngine.h
+++ b/Engine/Plugins/Runtime/Metasound/Source/MetasoundFrontend/Public/MetasoundFrontendSearchEngine.h
@@ -50,6 +50,9 @@ namespace Metasound::Frontend
 			/** Finds all registered interfaces with the given name. */
 			virtual TArray<FMetasoundFrontendVersion> FindAllRegisteredInterfacesWithName(FName InInterfaceName) = 0;
 
+			/** Finds the registered interface version with the highest version of the given name. Returns true if found, false if not. */
+			virtual bool FindHighestInterfaceVersion(FName InInterfaceName, FMetasoundFrontendVersion& OutVersion) = 0;
+
 			/** Finds the registered interface with the highest version of the given name. Returns true if found, false if not. */
 			virtual bool FindInterfaceWithHighestVersion(FName InInterfaceName, FMetasoundFrontendInterface& OutInterface) = 0;
 
@@ -83,8 +86,11 @@ namespace Metasound::Frontend
 			/** (To be deprecated) Find the highest version of a class with the given ClassName. Returns false if not found, true if found. */
 			virtual bool FindClassWithHighestVersion(const FMetasoundFrontendClassName& InName, FMetasoundFrontendClass& OutClass);
 
-			/** Returns array with all registered interfaces. Optionally, include interface versions that are not the highest version. */
-			virtual TArray<FMetasoundFrontendInterface> FindAllInterfaces(bool bInIncludeAllVersions = false) = 0;
+			/** Returns array with all registered interface versions. Optionally, include interface versions that are not the highest version. */
+			virtual TArray<FMetasoundFrontendVersion> FindAllInterfaceVersions(bool bInIncludeAllVersions = false) = 0;
+
+			UE_DEPRECATED(5.8, "Use FindAllInterfaceVersions instead, and refine results further if required using IInterfaceRegistry")
+			virtual TArray<FMetasoundFrontendInterface> FindAllInterfaces(bool bInIncludeAllVersions = false) { return { }; }
 #endif // WITH_EDITORONLY_DATA
 
 		protected:
diff --git a/Engine/Plugins/Runtime/Metasound/Source/MetasoundFrontend/Public/MetasoundFrontendTransform.h b/Engine/Plugins/Runtime/Metasound/Source/MetasoundFrontend/Public/MetasoundFrontendTransform.h
--- a/Engine/Plugins/Runtime/Metasound/Source/MetasoundFrontend/Public/MetasoundFrontendTransform.h
+++ b/Engine/Plugins/Runtime/Metasound/Source/MetasoundFrontend/Public/MetasoundFrontendTransform.h
@@ -8,6 +8,7 @@
 #include "MetasoundFrontendDocument.h"
 #include "MetasoundVertex.h"
 #include "Templates/Function.h"
+#include "Templates/SharedPointer.h"
 
 #define UE_API METASOUNDFRONTEND_API
 
@@ -24,13 +25,30 @@ namespace Metasound::Frontend
 		using FGetNodeDisplayNameProjection = TFunction<FText(const FNodeHandle&)>;
 		using FGetNodeDisplayNameProjectionRef = TFunctionRef<FText(const FNodeHandle&)>;
 
-		METASOUNDFRONTEND_API bool GetVersioningLoggingEnabled();
-		METASOUNDFRONTEND_API void SetVersioningLoggingEnabled(bool bIsEnabled);
-		METASOUNDFRONTEND_API void RegisterNodeDisplayNameProjection(FGetNodeDisplayNameProjection&& InNameProjection);
-		METASOUNDFRONTEND_API FGetNodeDisplayNameProjectionRef GetNodeDisplayNameProjection();
+		UE_API bool GetVersioningLoggingEnabled();
+		UE_API void SetVersioningLoggingEnabled(bool bIsEnabled);
+		UE_API void RegisterNodeDisplayNameProjection(FGetNodeDisplayNameProjection&& InNameProjection);
+		UE_API FGetNodeDisplayNameProjectionRef GetNodeDisplayNameProjection();
 #endif // WITH_EDITOR
 	}
 
+	class IBuilderTransform
+	{
+	public:
+		virtual ~IBuilderTransform() = default;
+
+		virtual bool Transform(FMetaSoundFrontendDocumentBuilder& InOutBuilder) const = 0;
+	};
+
+	class IBuilderVersionTransform : public IBuilderTransform
+	{
+	public:
+		virtual ~IBuilderVersionTransform() = default;
+
+		// Returns version update transform applies
+		virtual const FMetasoundFrontendVersion& GetVersion() const = 0;
+	};
+
 	/** Interface for transforms applied to documents. */
 	class IDocumentTransform
 	{
@@ -49,23 +67,20 @@ namespace Metasound::Frontend
 		UE_API virtual bool Transform(FMetasoundFrontendDocument& InOutDocument) const;
 	};
 
-	/** Interface for transforms applied to a graph. */
-	class IGraphTransform
+	class UE_DEPRECATED(5.8, "Use IBuilderTransform or similar instead") IGraphTransform
 	{
 	public:
 		virtual ~IGraphTransform() = default;
 
-		/** Return true if the graph was modified, false otherwise. */
 		virtual bool Transform(FMetasoundFrontendGraph& InOutGraph) const = 0;
 	};
 
 	/** Interface for transforming a node. */
-	class INodeTransform
+	class UE_DEPRECATED(5.8, "Use IBuilderTransform or similar instead") INodeTransform
 	{
 	public:
 		virtual ~INodeTransform() = default;
 
-		/** Return true if the node was modified, false otherwise. */
 		UE_API virtual bool Transform(const FGuid& InNodeID, FMetaSoundFrontendDocumentBuilder& OutBuilder) const;
 	};
 
