BugHawk fixes/polish as follow-up of Interface SearchEngine optimizations
#rnx

diff --git a/Engine/Plugins/Runtime/Metasound/Source/MetasoundEditor/Private/MetasoundEditorGraph.cpp b/Engine/Plugins/Runtime/Metasound/Source/MetasoundEditor/Private/MetasoundEditorGraph.cpp
--- a/Engine/Plugins/Runtime/Metasound/Source/MetasoundEditor/Private/MetasoundEditorGraph.cpp
+++ b/Engine/Plugins/Runtime/Metasound/Source/MetasoundEditor/Private/MetasoundEditorGraph.cpp
@@ -484,6 +484,7 @@ bool UMetasoundEditorGraphVertex::NameContainsInterfaceNamespace(FMetasoundFront
 			{
 				const bool bIsValid = IInterfaceRegistry::Get().FindInterface(InterfaceVersion, *OutInterface);
 				checkf(bIsValid, TEXT("Version found in search engine must be registered with interface registry"));
+				return bIsValid;
 			}
 			return true;
 		}
diff --git a/Engine/Plugins/Runtime/Metasound/Source/MetasoundFrontend/Private/Interfaces/MetasoundFrontendInterfaceRegistry.cpp b/Engine/Plugins/Runtime/Metasound/Source/MetasoundFrontend/Private/Interfaces/MetasoundFrontendInterfaceRegistry.cpp
--- a/Engine/Plugins/Runtime/Metasound/Source/MetasoundFrontend/Private/Interfaces/MetasoundFrontendInterfaceRegistry.cpp
+++ b/Engine/Plugins/Runtime/Metasound/Source/MetasoundFrontend/Private/Interfaces/MetasoundFrontendInterfaceRegistry.cpp
@@ -116,7 +116,7 @@ namespace Metasound::Frontend
 			if (const TUniquePtr<IInterfaceRegistryEntry>* Entry = Entries.Find(InVersion))
 			{
 				check(Entry);
-				return (*Entry)->GetInterface().Inputs.ContainsByPredicate(IsVertex);
+				return (*Entry)->GetInterface().Outputs.ContainsByPredicate(IsVertex);
 			}
 		}
 
@@ -138,9 +138,9 @@ namespace Metasound::Frontend
 	bool FInterfaceRegistry::FindInterfaceRouter(const FMetasoundFrontendVersion& InVersion, FName& OutRouter) const
 	{
 		UE::TScopeLock Lock(EntriesCriticalSection);
-		if (const TUniquePtr<IInterfaceRegistryEntry>* Entry = Entries.Find(InVersion))
+		if (const TUniquePtr<IInterfaceRegistryEntry>* Entry = Entries.Find(InVersion); Entry && Entry->IsValid())
 		{
-			OutRouter = Entry->Get()->GetRouterName();
+			OutRouter = (*Entry)->GetRouterName();
 			return true;
 		}
 
@@ -150,9 +150,9 @@ namespace Metasound::Frontend
 	TSharedPtr<IBuilderVersionTransform> FInterfaceRegistry::FindInterfaceUpdateTransform(const FMetasoundFrontendVersion& InVersion) const
 	{
 		UE::TScopeLock Lock(EntriesCriticalSection);
-		if (const TUniquePtr<IInterfaceRegistryEntry>* Entry = Entries.Find(InVersion))
+		if (const TUniquePtr<IInterfaceRegistryEntry>* Entry = Entries.Find(InVersion); Entry && Entry->IsValid())
 		{
-			return Entry->Get()->GetUpdateTransform();
+			return (*Entry)->GetUpdateTransform();
 		}
 
 		return { };
diff --git a/Engine/Plugins/Runtime/Metasound/Source/MetasoundFrontend/Private/MetasoundAssetBase.cpp b/Engine/Plugins/Runtime/Metasound/Source/MetasoundFrontend/Private/MetasoundAssetBase.cpp
--- a/Engine/Plugins/Runtime/Metasound/Source/MetasoundFrontend/Private/MetasoundAssetBase.cpp
+++ b/Engine/Plugins/Runtime/Metasound/Source/MetasoundFrontend/Private/MetasoundAssetBase.cpp
@@ -132,6 +132,8 @@ namespace Metasound
 
 			void GetUpdatePathForDocument(const FMetasoundFrontendVersion& InCurrentVersion, const FMetasoundFrontendVersion& InTargetVersion, TArray<TSharedPtr<IBuilderVersionTransform>>& OutUpgradePath)
 			{
+				OutUpgradePath.Reset();
+
 				if (InCurrentVersion.Name == InTargetVersion.Name)
 				{
 					// Get all associated registered interfaces
@@ -151,11 +153,11 @@ namespace Metasound
 					RegisteredVersions.Sort();
 
 					// Get registry entries from registry keys.
-					auto GetRegistryEntry = [](const FMetasoundFrontendVersion& InVersion)
+					auto GetUpdateTransform = [](const FMetasoundFrontendVersion& InVersion)
 					{
 						return IInterfaceRegistry::Get().FindInterfaceUpdateTransform(InVersion);
 					};
-					Algo::Transform(RegisteredVersions, OutUpgradePath, GetRegistryEntry);
+					Algo::Transform(RegisteredVersions, OutUpgradePath, GetUpdateTransform);
 				}
 			}
 
@@ -164,10 +166,12 @@ namespace Metasound
 				const FMetasoundFrontendVersionNumber* LastVersionUpdated = nullptr;
 				for (const TSharedPtr<IBuilderVersionTransform>& Transform : InUpgradePath)
 				{
-					check(Transform.IsValid());
-					if (Transform->Transform(Builder))
+					if (ensureAlwaysMsgf(Transform.IsValid(), TEXT("Transform should always be provided, even if referencing the default transform, by this point in updating document interfaces")))
 					{
-						LastVersionUpdated = &Transform->GetVersion().Number;
+						if (Transform->Transform(Builder))
+						{
+							LastVersionUpdated = &Transform->GetVersion().Number;
+						}
 					}
 				}
 
