Metasound: Fix some static const vars in headers.

diff --git a/Engine/Plugins/Runtime/Metasound/Source/MetasoundFrontend/Private/MetasoundArrayNodes.cpp b/Engine/Plugins/Runtime/Metasound/Source/MetasoundFrontend/Private/MetasoundArrayNodes.cpp
--- a/Engine/Plugins/Runtime/Metasound/Source/MetasoundFrontend/Private/MetasoundArrayNodes.cpp
+++ b/Engine/Plugins/Runtime/Metasound/Source/MetasoundFrontend/Private/MetasoundArrayNodes.cpp
@@ -7,6 +7,35 @@
 
 namespace Metasound
 {
+	namespace ArrayNodeVertexNames
+	{
+		const FLazyName InputInitialArrayName = TEXT("Array");
+#if WITH_EDITOR
+		const FText InputInitialArrayTooltip = LOCTEXT("InitialArrayTooltip", "Initial Array");
+		const FText InputInitialArrayDisplayName = LOCTEXT("InitialArrayDisplayName", "Init Array"); 
+#else
+		const FText InputInitialArrayTooltip = FText::GetEmpty();
+		const FText InputInitialArrayDisplayName = FText::GetEmpty();
+#endif
+
+		DEFINE_METASOUND_PARAM(InputArray, "Array", "Input Array.")
+		DEFINE_METASOUND_PARAM(InputLeftArray, "Left Array", "Input Left Array.")
+		DEFINE_METASOUND_PARAM(InputRightArray, "Right Array", "Input Right Array.")
+		DEFINE_METASOUND_PARAM(InputTriggerGet, "Trigger", "Trigger to get value.")
+		DEFINE_METASOUND_PARAM(InputTriggerSet, "Trigger", "Trigger to set value.")
+		DEFINE_METASOUND_PARAM(InputIndex, "Index", "Index in Array.")
+		DEFINE_METASOUND_PARAM(InputStartIndex, "Start Index", "First index to include.")
+		DEFINE_METASOUND_PARAM(InputEndIndex, "End Index", "Last index to include.")
+		DEFINE_METASOUND_PARAM(InputValue, "Value", "Value to set.")
+
+		DEFINE_METASOUND_PARAM(OutputNum, "Num", "Number of elements in the array.")
+		DEFINE_METASOUND_PARAM(OutputValue, "Element", "Value of element at array index.")
+		DEFINE_METASOUND_PARAM(OutputArrayConcat, "Array", "Array after concatenation.")
+		DEFINE_METASOUND_PARAM(OutputArraySet, "Array", "Array after setting.")
+		DEFINE_METASOUND_PARAM(OutputArraySubset, "Array", "Subset of input array.")
+		DEFINE_METASOUND_PARAM(OutputLastIndex, "Last Index", "Last index of the array.")
+	}
+
 	namespace MetasoundArrayNodesPrivate
 	{
 		FNodeClassMetadata CreateArrayNodeClassMetadata(const FName& InDataTypeName, const FName& InOperatorName, const FText& InDisplayName, const FText& InDescription, const FVertexInterface& InDefaultInterface, int32 MajorVersion, int32 MinorVersion, bool bIsDeprecated)
diff --git a/Engine/Plugins/Runtime/Metasound/Source/MetasoundFrontend/Public/MetasoundArrayNodes.h b/Engine/Plugins/Runtime/Metasound/Source/MetasoundFrontend/Public/MetasoundArrayNodes.h
--- a/Engine/Plugins/Runtime/Metasound/Source/MetasoundFrontend/Public/MetasoundArrayNodes.h
+++ b/Engine/Plugins/Runtime/Metasound/Source/MetasoundFrontend/Public/MetasoundArrayNodes.h
@@ -52,32 +52,27 @@ namespace Metasound
 
 	namespace ArrayNodeVertexNames
 	{
-		static const FLazyName InputInitialArrayName = TEXT("Array");
-#if WITH_EDITOR
-		static const FText InputInitialArrayTooltip = LOCTEXT("InitialArrayTooltip", "Initial Array");
-		static const FText InputInitialArrayDisplayName = LOCTEXT("InitialArrayDisplayName", "Init Array"); 
-#else
-		static const FText InputInitialArrayTooltip = FText::GetEmpty();
-		static const FText InputInitialArrayDisplayName = FText::GetEmpty();
-#endif
-
-		METASOUND_PARAM(InputArray, "Array", "Input Array.")
-		METASOUND_PARAM(InputLeftArray, "Left Array", "Input Left Array.")
-		METASOUND_PARAM(InputRightArray, "Right Array", "Input Right Array.")
-		METASOUND_PARAM(InputTriggerGet, "Trigger", "Trigger to get value.")
-		METASOUND_PARAM(InputTriggerSet, "Trigger", "Trigger to set value.")
-		METASOUND_PARAM(InputIndex, "Index", "Index in Array.")
-		METASOUND_PARAM(InputStartIndex, "Start Index", "First index to include.")
-		METASOUND_PARAM(InputEndIndex, "End Index", "Last index to include.")
-		METASOUND_PARAM(InputValue, "Value", "Value to set.")
-
-		METASOUND_PARAM(OutputNum, "Num", "Number of elements in the array.")
-		METASOUND_PARAM(OutputValue, "Element", "Value of element at array index.")
-		METASOUND_PARAM(OutputArrayConcat, "Array", "Array after concatenation.")
-		METASOUND_PARAM(OutputArraySet, "Array", "Array after setting.")
-		METASOUND_PARAM(OutputArraySubset, "Array", "Subset of input array.")
-		METASOUND_PARAM(OutputLastIndex, "Last Index", "Last index of the array.")
-	};
+		extern METASOUNDFRONTEND_API const FLazyName InputInitialArrayName;
+		extern METASOUNDFRONTEND_API const FText InputInitialArrayTooltip;
+		extern METASOUNDFRONTEND_API const FText InputInitialArrayDisplayName; 
+
+		DECLARE_METASOUND_PARAM(METASOUNDFRONTEND_API, InputArray)
+		DECLARE_METASOUND_PARAM(METASOUNDFRONTEND_API, InputLeftArray)
+		DECLARE_METASOUND_PARAM(METASOUNDFRONTEND_API, InputRightArray)
+		DECLARE_METASOUND_PARAM(METASOUNDFRONTEND_API, InputTriggerGet)
+		DECLARE_METASOUND_PARAM(METASOUNDFRONTEND_API, InputTriggerSet)
+		DECLARE_METASOUND_PARAM(METASOUNDFRONTEND_API, InputIndex)
+		DECLARE_METASOUND_PARAM(METASOUNDFRONTEND_API, InputStartIndex)
+		DECLARE_METASOUND_PARAM(METASOUNDFRONTEND_API, InputEndIndex)
+		DECLARE_METASOUND_PARAM(METASOUNDFRONTEND_API, InputValue)
+
+		DECLARE_METASOUND_PARAM(METASOUNDFRONTEND_API, OutputNum)
+		DECLARE_METASOUND_PARAM(METASOUNDFRONTEND_API, OutputValue)
+		DECLARE_METASOUND_PARAM(METASOUNDFRONTEND_API, OutputArrayConcat)
+		DECLARE_METASOUND_PARAM(METASOUNDFRONTEND_API, OutputArraySet)
+		DECLARE_METASOUND_PARAM(METASOUNDFRONTEND_API, OutputArraySubset)
+		DECLARE_METASOUND_PARAM(METASOUNDFRONTEND_API, OutputLastIndex)
+	}
 
 
 	/** TArrayNumOperator gets the number of elements in an Array. The operator
diff --git a/Engine/Plugins/Runtime/Metasound/Source/MetasoundGraphCore/Private/MetasoundVariableNodes.cpp b/Engine/Plugins/Runtime/Metasound/Source/MetasoundGraphCore/Private/MetasoundVariableNodes.cpp
--- a/Engine/Plugins/Runtime/Metasound/Source/MetasoundGraphCore/Private/MetasoundVariableNodes.cpp
+++ b/Engine/Plugins/Runtime/Metasound/Source/MetasoundGraphCore/Private/MetasoundVariableNodes.cpp
@@ -5,10 +5,41 @@
 #include "MetasoundNodeInterface.h"
 #include "MetasoundVertex.h"
 
+#define LOCTEXT_NAMESPACE "MetasoundGraphCore"
+
 namespace Metasound
 {
 	namespace VariableNames
 	{
+		DEFINE_METASOUND_PARAM(InputData, "Value", "Value")
+		DEFINE_METASOUND_PARAM(OutputData, "Value", "Value")
+		DEFINE_METASOUND_PARAM(InputVariable, "Variable", "Variable")
+		DEFINE_METASOUND_PARAM(OutputVariable, "Variable", "Variable")
+
+		const FDataVertexMetadata InputDataMetadata
+		{
+			  FText::GetEmpty() // description
+			, METASOUND_GET_PARAM_DISPLAYNAME(InputData) // display name
+		};
+
+		const FDataVertexMetadata OutputDataMetadata
+		{
+			  FText::GetEmpty() // description
+			, METASOUND_GET_PARAM_DISPLAYNAME(OutputData) // display name
+		};
+
+		const FDataVertexMetadata InputVariableMetadata
+		{
+			  FText::GetEmpty() // description
+			, METASOUND_GET_PARAM_DISPLAYNAME(InputVariable) // display name
+		};
+
+		const FDataVertexMetadata OutputVariableMetadata
+		{
+			  FText::GetEmpty() // description
+			, METASOUND_GET_PARAM_DISPLAYNAME(OutputVariable) // display name
+		};
+
 		FNodeClassName GetVariableNodeClassName(const FName& InDataTypeName)
 		{
 			static const FName NodeNamespace = "InitVariable";
@@ -61,3 +92,5 @@ namespace Metasound
 		}
 	}
 }
+
+#undef LOCTEXT_NAMESPACE // MetasoundGraphCore
diff --git a/Engine/Plugins/Runtime/Metasound/Source/MetasoundGraphCore/Public/MetasoundVariableNodes.h b/Engine/Plugins/Runtime/Metasound/Source/MetasoundGraphCore/Public/MetasoundVariableNodes.h
--- a/Engine/Plugins/Runtime/Metasound/Source/MetasoundGraphCore/Public/MetasoundVariableNodes.h
+++ b/Engine/Plugins/Runtime/Metasound/Source/MetasoundGraphCore/Public/MetasoundVariableNodes.h
@@ -19,34 +19,15 @@ namespace Metasound
 {
 	namespace VariableNames
 	{
-		METASOUND_PARAM(InputData, "Value", "Value")
-		METASOUND_PARAM(OutputData, "Value", "Value")
-		METASOUND_PARAM(InputVariable, "Variable", "Variable")
-		METASOUND_PARAM(OutputVariable, "Variable", "Variable")
-
-		static const FDataVertexMetadata InputDataMetadata
-		{
-			  FText::GetEmpty() // description
-			, METASOUND_GET_PARAM_DISPLAYNAME(InputData) // display name
-		};
-
-		static const FDataVertexMetadata OutputDataMetadata
-		{
-			  FText::GetEmpty() // description
-			, METASOUND_GET_PARAM_DISPLAYNAME(OutputData) // display name
-		};
-
-		static const FDataVertexMetadata InputVariableMetadata
-		{
-			  FText::GetEmpty() // description
-			, METASOUND_GET_PARAM_DISPLAYNAME(InputVariable) // display name
-		};
-
-		static const FDataVertexMetadata OutputVariableMetadata
-		{
-			  FText::GetEmpty() // description
-			, METASOUND_GET_PARAM_DISPLAYNAME(OutputVariable) // display name
-		};
+		DECLARE_METASOUND_PARAM(METASOUNDGRAPHCORE_API, InputData)
+		DECLARE_METASOUND_PARAM(METASOUNDGRAPHCORE_API, OutputData)
+		DECLARE_METASOUND_PARAM(METASOUNDGRAPHCORE_API, InputVariable)
+		DECLARE_METASOUND_PARAM(METASOUNDGRAPHCORE_API, OutputVariable)
+
+		extern METASOUNDGRAPHCORE_API const FDataVertexMetadata InputDataMetadata;
+		extern METASOUNDGRAPHCORE_API const FDataVertexMetadata OutputDataMetadata;
+		extern METASOUNDGRAPHCORE_API const FDataVertexMetadata InputVariableMetadata;
+		extern METASOUNDGRAPHCORE_API const FDataVertexMetadata OutputVariableMetadata;
 
 		/** Class name for variable node. */
 		METASOUNDGRAPHCORE_API FNodeClassName GetVariableNodeClassName(const FName& InDataTypeName);
