MetasoundFrontend: A few more global variable fixes

diff --git a/Engine/Plugins/Runtime/Metasound/Source/MetasoundFrontend/Private/MetasoundLiteralNode.cpp b/Engine/Plugins/Runtime/Metasound/Source/MetasoundFrontend/Private/MetasoundLiteralNode.cpp
--- a/Engine/Plugins/Runtime/Metasound/Source/MetasoundFrontend/Private/MetasoundLiteralNode.cpp
+++ b/Engine/Plugins/Runtime/Metasound/Source/MetasoundFrontend/Private/MetasoundLiteralNode.cpp
@@ -5,6 +5,11 @@
 #define LOCTEXT_NAMESPACE "MetasoundLiteralNode"
 namespace Metasound
 {
+	namespace LiteralNodeNames
+	{
+		DEFINE_METASOUND_PARAM(OutputValue, "Value", "Value")
+	}
+
 	FVertexInterface FLiteralNode::CreateVertexInterface(const FName& InDataTypeName, EVertexAccessType InAccessType)
 	{
 		using namespace LiteralNodeNames; 
diff --git a/Engine/Plugins/Runtime/Metasound/Source/MetasoundFrontend/Private/MetasoundReceiveNode.cpp b/Engine/Plugins/Runtime/Metasound/Source/MetasoundFrontend/Private/MetasoundReceiveNode.cpp
--- a/Engine/Plugins/Runtime/Metasound/Source/MetasoundFrontend/Private/MetasoundReceiveNode.cpp
+++ b/Engine/Plugins/Runtime/Metasound/Source/MetasoundFrontend/Private/MetasoundReceiveNode.cpp
@@ -2,10 +2,32 @@
 
 #include "MetasoundReceiveNode.h"
 
+#define LOCTEXT_NAMESPACE "MetasoundFrontend"
+
 namespace Metasound
 {
 	namespace ReceiveNodeInfo
 	{
+		DEFINE_METASOUND_PARAM(AddressInput, "Address", "Address")
+		DEFINE_METASOUND_PARAM(DefaultDataInput, "Default", "Default")
+		DEFINE_METASOUND_PARAM(Output, "Out", "Out")
+
+		const FDataVertexMetadata AddressInputMetadata
+		{
+			  FText::GetEmpty() // description
+			, METASOUND_GET_PARAM_DISPLAYNAME(AddressInput) // display name
+		};
+		const FDataVertexMetadata DefaultDataInputMetadata
+		{
+			  FText::GetEmpty() // description
+			, METASOUND_GET_PARAM_DISPLAYNAME(DefaultDataInput) // display name
+		};
+		const FDataVertexMetadata OutputMetadata
+		{
+			  FText::GetEmpty() // description
+			, METASOUND_GET_PARAM_DISPLAYNAME(Output) // display name
+		};
+
 		FNodeClassName GetClassNameForDataType(const FName& InDataTypeName)
 		{
 			return FNodeClassName { "Receive", InDataTypeName, FName() };
@@ -16,3 +38,5 @@ namespace Metasound
 		int32 GetCurrentMinorVersion() { return 0; }
 	}
 }
+
+#undef LOCTEXT_NAMESPACE
diff --git a/Engine/Plugins/Runtime/Metasound/Source/MetasoundFrontend/Public/MetasoundLiteralNode.h b/Engine/Plugins/Runtime/Metasound/Source/MetasoundFrontend/Public/MetasoundLiteralNode.h
--- a/Engine/Plugins/Runtime/Metasound/Source/MetasoundFrontend/Public/MetasoundLiteralNode.h
+++ b/Engine/Plugins/Runtime/Metasound/Source/MetasoundFrontend/Public/MetasoundLiteralNode.h
@@ -19,7 +19,7 @@ namespace Metasound
 {
 	namespace LiteralNodeNames
 	{
-		METASOUND_PARAM(OutputValue, "Value", "Value")
+		DECLARE_METASOUND_PARAM(METASOUNDFRONTEND_API, OutputValue)
 	}
 
 	template<typename DataType>
diff --git a/Engine/Plugins/Runtime/Metasound/Source/MetasoundFrontend/Public/MetasoundReceiveNode.h b/Engine/Plugins/Runtime/Metasound/Source/MetasoundFrontend/Public/MetasoundReceiveNode.h
--- a/Engine/Plugins/Runtime/Metasound/Source/MetasoundFrontend/Public/MetasoundReceiveNode.h
+++ b/Engine/Plugins/Runtime/Metasound/Source/MetasoundFrontend/Public/MetasoundReceiveNode.h
@@ -22,10 +22,13 @@ namespace Metasound
 {
 	namespace ReceiveNodeInfo
 	{
-		METASOUND_PARAM(AddressInput, "Address", "Address")
-		METASOUND_PARAM(DefaultDataInput, "Default", "Default")
-		METASOUND_PARAM(Output, "Out", "Out")
+		DECLARE_METASOUND_PARAM(METASOUNDFRONTEND_API, AddressInput)
+		DECLARE_METASOUND_PARAM(METASOUNDFRONTEND_API, DefaultDataInput)
+		DECLARE_METASOUND_PARAM(METASOUNDFRONTEND_API, Output)
 
+		extern METASOUNDFRONTEND_API const FDataVertexMetadata AddressInputMetadata;
+		extern METASOUNDFRONTEND_API const FDataVertexMetadata DefaultDataInputMetadata;
+		extern METASOUNDFRONTEND_API const FDataVertexMetadata OutputMetadata;
 
 		METASOUNDFRONTEND_API FNodeClassName GetClassNameForDataType(const FName& InDataTypeName);
 		METASOUNDFRONTEND_API int32 GetCurrentMajorVersion();
@@ -43,21 +46,6 @@ namespace Metasound
 			static FVertexInterface DeclareVertexInterface()
 			{
 				using namespace ReceiveNodeInfo;
-				static const FDataVertexMetadata AddressInputMetadata
-				{
-					  FText::GetEmpty() // description
-					, METASOUND_GET_PARAM_DISPLAYNAME(AddressInput) // display name
-				};
-				static const FDataVertexMetadata DefaultDataInputMetadata
-				{
-					  FText::GetEmpty() // description
-					, METASOUND_GET_PARAM_DISPLAYNAME(DefaultDataInput) // display name
-				};
-				static const FDataVertexMetadata OutputMetadata
-				{
-					  FText::GetEmpty() // description
-					, METASOUND_GET_PARAM_DISPLAYNAME(Output) // display name
-				};
 				return FVertexInterface(
 					FInputVertexInterface(
 						TInputDataVertex<FSendAddress>(METASOUND_GET_PARAM_NAME(AddressInput), AddressInputMetadata),
