[MetaSoundEditor] Allow interfaces to have editor-only descriptions that show as tooltips as well as show their input/output counts in the "Add Interface..." dropdown

#FailedValidations Preflight

diff --git a/Engine/Plugins/Runtime/Metasound/Source/MetasoundEditor/Private/MetasoundDetailCustomization.cpp b/Engine/Plugins/Runtime/Metasound/Source/MetasoundEditor/Private/MetasoundDetailCustomization.cpp
--- a/Engine/Plugins/Runtime/Metasound/Source/MetasoundEditor/Private/MetasoundDetailCustomization.cpp
+++ b/Engine/Plugins/Runtime/Metasound/Source/MetasoundEditor/Private/MetasoundDetailCustomization.cpp
@@ -803,12 +803,73 @@ namespace Metasound::Editor
 
 		UpdateInterfaceNames();
 
+		const UMetasoundEditorSettings* MetasoundEditorSettings = GetDefault<UMetasoundEditorSettings>();
 		SAssignNew(InterfaceComboBox, SSearchableComboBox)
 			.OptionsSource(&AddableInterfaceNames)
-			.OnGenerateWidget_Lambda([](TSharedPtr<FString> InItem)
+			.OnGenerateWidget_Lambda([&InputColor = MetasoundEditorSettings->InputNodeTitleColor, &OutputColor = MetasoundEditorSettings->OutputNodeTitleColor](TSharedPtr<FString> InItem)
 			{
-				return SNew(STextBlock)
-					.Text(FText::FromString(*InItem));
+				using namespace Metasound::Frontend;
+
+				FText* InterfaceDescription = nullptr;
+				FMetasoundFrontendVersion InterfaceVersion;
+				FMetasoundFrontendInterface InterfaceToAdd;
+				if (ensure(ISearchEngine::Get().FindHighestInterfaceVersion(FName{ *InItem.Get() }, InterfaceVersion)))
+				{
+					if (ensure(IInterfaceRegistry::Get().FindInterface(InterfaceVersion, InterfaceToAdd)))
+					{
+						InterfaceDescription = InterfaceToAdd.Metadata.Description.IsEmpty() ? nullptr : &InterfaceToAdd.Metadata.Description;
+					}
+				}
+
+				return SNew(SHorizontalBox)
+				+ SHorizontalBox::Slot()
+					.VAlign(VAlign_Center)
+					.FillWidth(1.f)
+					[
+						SNew(STextBlock)
+						.Text(FText::FromString(*InItem))
+						.ToolTipText(InterfaceDescription ? *InterfaceDescription : FText::FromStringView(*InItem))
+					]
+				+ SHorizontalBox::Slot()
+					.AutoWidth()
+					[
+						SNew(SSpacer)
+						.Size(FVector2D(8.f, 0.f))
+					]
+				+ SHorizontalBox::Slot()
+					.VAlign(VAlign_Center)
+					.AutoWidth()
+					[
+						SNew(SBorder)
+						.BorderImage(&Style::GetSlateBrushSafe("MetasoundEditor.InterfacesInOut.Background"))
+						.ToolTipText(FText::Format(LOCTEXT("InterfaceInputsOutputsTooltip", "{0} {0}|plural(one=input,other=inputs), {1} {1}|plural(one=output,other=outputs)"), TArray<FFormatArgumentValue>{ {InterfaceToAdd.Inputs.Num()}, {InterfaceToAdd.Outputs.Num()} }))
+						[
+							SNew(SHorizontalBox)
+							+ SHorizontalBox::Slot()
+								.AutoWidth()
+								[
+									SNew(STextBlock)
+									.ColorAndOpacity(InputColor)
+									.Font(FCoreStyle::GetDefaultFontStyle(FName("Bold"), 10))
+									.Text(FText::AsNumber(InterfaceToAdd.Inputs.Num()))
+								]
+							+ SHorizontalBox::Slot()
+								.AutoWidth()
+								[
+									SNew(STextBlock)
+									.Font(FCoreStyle::GetDefaultFontStyle(FName("Bold"), 10))
+									.Text(FText::AsCultureInvariant(" | "))
+								]
+							+ SHorizontalBox::Slot()
+								.AutoWidth()
+								[
+									SNew(STextBlock)
+									.ColorAndOpacity(OutputColor)
+									.Font(FCoreStyle::GetDefaultFontStyle(FName("Bold"), 10))
+									.Text(FText::AsNumber(InterfaceToAdd.Outputs.Num()))
+								]
+						] // SBorder
+					]; // SHorizontalBox
 			})
 			.OnSelectionChanged_Lambda([this](TSharedPtr<FString> NameToAdd, ESelectInfo::Type InSelectInfo)
 			{
@@ -932,7 +993,10 @@ namespace Metasound::Editor
 				return SNullWidget::NullWidget;
 			}
 
-			TSharedRef<SWidget> RemoveButtonWidget = PropertyCustomizationHelpers::MakeDeleteButton(FSimpleDelegate::CreateLambda([this, InInterfaceName, InterfaceVersion]()
+			FMetasoundFrontendInterface InterfaceToAdd;
+			const bool bFoundInterface = IInterfaceRegistry::Get().FindInterface(InterfaceVersion, InterfaceToAdd);
+
+			TSharedRef<SWidget> RemoveButtonWidget = PropertyCustomizationHelpers::MakeDeleteButton(FSimpleDelegate::CreateLambda([this, InInterfaceName, InterfaceVersion, bFoundInterface, InterfaceToAdd]()
 			{
 				using namespace Frontend;
 
@@ -949,12 +1013,11 @@ namespace Metasound::Editor
 
 				UObject& MetaSound = DocBuilder.CastDocumentObjectChecked<UObject>();
 				{
-					FMetasoundFrontendInterface InterfaceToAdd;
-					if (ensure(IInterfaceRegistry::Get().FindInterface(InterfaceVersion, InterfaceToAdd)))
+					if (ensure(bFoundInterface))
 					{
 						const FScopedTransaction Transaction(FText::Format(LOCTEXT("RemoveInterfaceTransactionFormat", "Remove MetaSound Interface '{0}'"), FText::FromString(InterfaceVersion.ToString())));
 						MetaSound.Modify();
-						FModifyInterfaceOptions Options({ MoveTemp(InterfaceToAdd) }, { });
+						FModifyInterfaceOptions Options({ InterfaceToAdd }, { });
 						Options.bSetDefaultNodeLocations = false; // Don't automatically add nodes to ed graph
 						DocBuilder.ModifyInterfaces(MoveTemp(Options));
 					}
@@ -975,6 +1038,7 @@ namespace Metasound::Editor
 				[
 					SNew(STextBlock)
 					.Text(FText::FromName(InterfaceVersion.Name))
+					.ToolTipText(bFoundInterface && !InterfaceToAdd.Metadata.Description.IsEmpty() ? InterfaceToAdd.Metadata.Description : FText::FromName(InterfaceVersion.Name))
 				]
 			+ SHorizontalBox::Slot()
 				.Padding(2.0f)
diff --git a/Engine/Plugins/Runtime/Metasound/Source/MetasoundEditor/Private/MetasoundEditorModule.cpp b/Engine/Plugins/Runtime/Metasound/Source/MetasoundEditor/Private/MetasoundEditorModule.cpp
--- a/Engine/Plugins/Runtime/Metasound/Source/MetasoundEditor/Private/MetasoundEditorModule.cpp
+++ b/Engine/Plugins/Runtime/Metasound/Source/MetasoundEditor/Private/MetasoundEditorModule.cpp
@@ -477,6 +477,7 @@ namespace Metasound::Editor
 				Set("MetasoundEditor.Audition", new IMAGE_BRUSH_SVG(TEXT("Icons/metasound_page"), Icon16));
 				Set("MetasoundEditor.Metasound.Icon", new IMAGE_BRUSH_SVG(TEXT("Icons/metasound_icon"), Icon16));
 				Set("MetasoundEditor.Speaker", new FSlateImageBrush(RootToContentDir(TEXT("/Icons/speaker_144x.png")), FVector2D(144.0f, 144.0f)));
+				Set("MetasoundEditor.InterfacesInOut.Background", new FSlateBoxBrush(FAppStyle::Get().GetContentRootDir() / TEXT("Common/RoundedSelection_16x.png"), 4.0f / 16.0f, FLinearColor(1.f, 1.f, 1.f, 0.1f)));
 
 				// Pages
 				Set("MetasoundEditor.Page.Executing.ForegroundColor", FStyleColors::AccentGreen.GetSpecifiedColor());
diff --git a/Engine/Plugins/Runtime/Metasound/Source/MetasoundFrontend/Private/MetasoundFrontendDocument.cpp b/Engine/Plugins/Runtime/Metasound/Source/MetasoundFrontend/Private/MetasoundFrontendDocument.cpp
--- a/Engine/Plugins/Runtime/Metasound/Source/MetasoundFrontend/Private/MetasoundFrontendDocument.cpp
+++ b/Engine/Plugins/Runtime/Metasound/Source/MetasoundFrontend/Private/MetasoundFrontendDocument.cpp
@@ -709,7 +709,10 @@ FMetasoundFrontendInterface::FMetasoundFrontendInterface(Audio::FParameterInterf
 	using namespace Metasound::Frontend;
 
 	Metadata.Version = { InInterface->GetName(), FMetasoundFrontendVersionNumber { InInterface->GetVersion().Major, InInterface->GetVersion().Minor } };
-
+#if WITH_EDITORONLY_DATA
+	Metadata.Description = InInterface->GetDescription();
+#endif
+	
 	// Transfer all input data from AudioExtension interface struct to FrontendInterface
 	Algo::Transform(InInterface->GetInputs(), Inputs, [this](const Audio::FParameterInterface::FInput& Input)
 	{
diff --git a/Engine/Plugins/Runtime/Metasound/Source/MetasoundFrontend/Public/MetasoundFrontendDocument.h b/Engine/Plugins/Runtime/Metasound/Source/MetasoundFrontend/Public/MetasoundFrontendDocument.h
--- a/Engine/Plugins/Runtime/Metasound/Source/MetasoundFrontend/Public/MetasoundFrontendDocument.h
+++ b/Engine/Plugins/Runtime/Metasound/Source/MetasoundFrontend/Public/MetasoundFrontendDocument.h
@@ -1366,7 +1366,7 @@ struct FMetasoundFrontendInterfaceUClassOptions
 	UE_API FMetasoundFrontendInterfaceUClassOptions(const FTopLevelAssetPath& InClassPath, bool bInIsModifiable = true, bool bInIsDefault = false);
 
 	// Path to MetaSound class interface can be added to (ex. UMetaSoundSource or UMetaSound)
-	UPROPERTY(BlueprintReadOnly, meta = (Category = Interface))
+	UPROPERTY(BlueprintReadOnly, Category = Interface)
 	FTopLevelAssetPath ClassPath;
 
 	// True if user can add or remove the given class directly to or from the inherited interface UI, false if not.
@@ -1385,13 +1385,19 @@ struct FMetasoundFrontendInterfaceMetadata
 	GENERATED_BODY()
 
 	// Name and version number of the interface
-	UPROPERTY(BlueprintReadOnly, meta = (Category = Interface))
+	UPROPERTY(BlueprintReadOnly, Category = Interface)
 	FMetasoundFrontendVersion Version;
 
 	// If specified, options used to restrict a corresponding UClass that interface may be
 	// applied to.  If unspecified, interface is assumed to be applicable to any arbitrary UClass.
-	UPROPERTY(BlueprintReadOnly, meta = (Category = Interface))
+	UPROPERTY(BlueprintReadOnly, Category = Interface)
 	TArray<FMetasoundFrontendInterfaceUClassOptions> UClassOptions;
+
+#if WITH_EDITORONLY_DATA
+	// Developer-supplied description of the interface
+	UPROPERTY(VisibleAnywhere, Category = Interface)
+	FText Description;
+#endif
 };
 
 
diff --git a/Engine/Source/Runtime/AudioExtensions/Private/IAudioParameterInterfaceRegistry.cpp b/Engine/Source/Runtime/AudioExtensions/Private/IAudioParameterInterfaceRegistry.cpp
--- a/Engine/Source/Runtime/AudioExtensions/Private/IAudioParameterInterfaceRegistry.cpp
+++ b/Engine/Source/Runtime/AudioExtensions/Private/IAudioParameterInterfaceRegistry.cpp
@@ -87,7 +87,14 @@ namespace Audio
 		return Environment;
 	}
 
-	TArray<const UClass*> FParameterInterface::FindSupportedUClasses() const
+#if WITH_EDITORONLY_DATA
+	const FText& FParameterInterface::GetDescription() const
+	{
+		return Description;
+	}
+#endif
+
+TArray<const UClass*> FParameterInterface::FindSupportedUClasses() const
 	{
 		TArray<const UClass*> SupportedUClasses;
 		for (const FClassOptions& Options : UClassOptions)
diff --git a/Engine/Source/Runtime/AudioExtensions/Public/IAudioParameterInterfaceRegistry.h b/Engine/Source/Runtime/AudioExtensions/Public/IAudioParameterInterfaceRegistry.h
--- a/Engine/Source/Runtime/AudioExtensions/Public/IAudioParameterInterfaceRegistry.h
+++ b/Engine/Source/Runtime/AudioExtensions/Public/IAudioParameterInterfaceRegistry.h
@@ -166,6 +166,11 @@ namespace Audio
 		// Returns read-only array of environment variables
 		AUDIOEXTENSIONS_API const TArray<FEnvironmentVariable>& GetEnvironment() const;
 
+#if WITH_EDITORONLY_DATA
+		// Description to be displayed in editor or tools
+		AUDIOEXTENSIONS_API const FText& GetDescription() const;
+#endif
+
 	private:
 		FName NamePrivate;
 		FVersion VersionPrivate;
@@ -175,6 +180,10 @@ namespace Audio
 		TArray<FOutput> Outputs;
 		TArray<FEnvironmentVariable> Environment;
 		TArray<FClassOptions> UClassOptions;
+		
+#if WITH_EDITOR
+		FText Description;
+#endif
 	};
 	using FParameterInterfacePtr = TSharedPtr<FParameterInterface, ESPMode::ThreadSafe>;
 
