Fix for runtime data and transactor not updating when renaming/switching data type of inputs, causing Verse manifests to be incorrect for allowlisted MetaSounds

#rnx

diff --git a/Engine/Plugins/Runtime/Metasound/Source/MetasoundEngine/Private/MetasoundBuilderSubsystem.cpp b/Engine/Plugins/Runtime/Metasound/Source/MetasoundEngine/Private/MetasoundBuilderSubsystem.cpp
--- a/Engine/Plugins/Runtime/Metasound/Source/MetasoundEngine/Private/MetasoundBuilderSubsystem.cpp
+++ b/Engine/Plugins/Runtime/Metasound/Source/MetasoundEngine/Private/MetasoundBuilderSubsystem.cpp
@@ -410,6 +410,8 @@ void UMetaSoundSourceBuilder::InitDelegates(Metasound::Frontend::FDocumentModify
 		OutDocumentDelegates.OnPresetStateChanged.AddUObject(this, &UMetaSoundSourceBuilder::OnPresetStateChanged);
 
 		OutDocumentDelegates.InterfaceDelegates.OnInputAdded.AddUObject(this, &UMetaSoundSourceBuilder::OnInputAdded);
+		OutDocumentDelegates.InterfaceDelegates.OnInputDataTypeChanged.AddUObject(this, &UMetaSoundSourceBuilder::OnInputDataTypeChanged);
+		OutDocumentDelegates.InterfaceDelegates.OnInputNameChanged.AddUObject(this, &UMetaSoundSourceBuilder::OnInputNameChanged);
 		OutDocumentDelegates.InterfaceDelegates.OnOutputAdded.AddUObject(this, &UMetaSoundSourceBuilder::OnOutputAdded);
 		OutDocumentDelegates.InterfaceDelegates.OnRemovingInput.AddUObject(this, &UMetaSoundSourceBuilder::OnRemovingInput);
 		OutDocumentDelegates.InterfaceDelegates.OnRemovingOutput.AddUObject(this, &UMetaSoundSourceBuilder::OnRemovingOutput);
@@ -535,6 +537,54 @@ void UMetaSoundSourceBuilder::OnInputAdded(int32 InputIndex)
 	});
 }
 
+void UMetaSoundSourceBuilder::OnInputDataTypeChanged(int32 InputIndex)
+{
+	OnRemovingInput(InputIndex);
+	OnInputAdded(InputIndex);
+}
+
+void UMetaSoundSourceBuilder::OnInputNameChanged(FName OldName, FName NewName)
+{
+	using namespace Metasound::DynamicGraph;
+	using namespace Metasound::Frontend;
+
+	UMetaSoundSource& Source = GetMetaSoundSource();
+	Source.RuntimeInputData.InputMap.Remove(OldName);
+
+	if (const FMetasoundFrontendClassInput* Input = GetConstBuilder().FindGraphInput(NewName))
+	{
+		constexpr bool bCreateUObjectProxies = true;
+		Source.RuntimeInputData.InputMap.Add(Input->Name, UMetaSoundSource::CreateRuntimeInput(IDataTypeRegistry::Get(), *Input, bCreateUObjectProxies, UMetaSoundSettings::GetPageOrder()));
+
+		ExecuteAuditionableTransaction([this, OldName, &Input](FDynamicOperatorTransactor& Transactor)
+		{
+			using namespace Metasound;
+			using namespace Metasound::Frontend;
+
+			for (const uint64& AudioComponentID : LiveComponentIDs)
+			{
+				if (UAudioComponent* AudioComponent = UAudioComponent::GetAudioComponentFromID(AudioComponentID))
+				{
+					if (FAudioDevice* AudioDevice = AudioComponent->GetAudioDevice())
+					{
+						AudioDevice->SendCommandToActiveSounds(AudioComponentID, [OldName, NewInputName = Input->Name](FActiveSound& ActiveSound)
+						{
+							FMetaSoundParameterTransmitter* Transmitter = static_cast<FMetaSoundParameterTransmitter*>(ActiveSound.GetTransmitter());
+							check(Transmitter);
+							Transmitter->RemoveAvailableParameter(OldName);
+							Transmitter->AddAvailableParameter(NewInputName);
+						});
+					}
+				}
+			}
+			Transactor.RemoveInputDataDestination(OldName);
+			const FLiteral NewInputLiteral = Input->FindConstDefaultChecked(Frontend::DefaultPageID).ToLiteral(Input->TypeName);
+			Transactor.AddInputDataDestination(Input->NodeID, Input->Name, NewInputLiteral, &FMetaSoundFrontendDocumentBuilder::CreateDataReference);
+			return true;
+		});
+	}
+}
+
 void UMetaSoundSourceBuilder::OnNodeAdded(int32 NodeIndex) const
 {
 	using namespace Metasound::DynamicGraph;
diff --git a/Engine/Plugins/Runtime/Metasound/Source/MetasoundEngine/Public/MetasoundBuilderSubsystem.h b/Engine/Plugins/Runtime/Metasound/Source/MetasoundEngine/Public/MetasoundBuilderSubsystem.h
--- a/Engine/Plugins/Runtime/Metasound/Source/MetasoundEngine/Public/MetasoundBuilderSubsystem.h
+++ b/Engine/Plugins/Runtime/Metasound/Source/MetasoundEngine/Public/MetasoundBuilderSubsystem.h
@@ -124,6 +124,8 @@ private:
 
 	UE_API void OnEdgeAdded(int32 EdgeIndex) const;
 	UE_API void OnInputAdded(int32 InputIndex);
+	UE_API void OnInputDataTypeChanged(int32 InputIndex);
+	UE_API void OnInputNameChanged(FName OldName, FName NewName);
 	UE_API void OnLiveComponentFinished(UAudioComponent* AudioComponent);
 	UE_API void OnNodeAdded(int32 NodeIndex) const;
 	UE_API void OnNodeConfigurationUpdated(int32 NodeIndex) const;
