Fix ensure when reloading MetaSound externally added and referenced by another asset

diff --git a/Engine/Plugins/Runtime/Metasound/Source/MetasoundFrontend/Private/MetasoundAssetBase.cpp b/Engine/Plugins/Runtime/Metasound/Source/MetasoundFrontend/Private/MetasoundAssetBase.cpp
--- a/Engine/Plugins/Runtime/Metasound/Source/MetasoundFrontend/Private/MetasoundAssetBase.cpp
+++ b/Engine/Plugins/Runtime/Metasound/Source/MetasoundFrontend/Private/MetasoundAssetBase.cpp
@@ -668,6 +668,7 @@ bool FMetasoundAssetBase::IsReferencedAsset(const FMetasoundAssetBase& InAsset)
 	using namespace Metasound::Frontend;
 
 	bool bIsReferenced = false;
+	const FString InAssetKey = InAsset.GetGraphRegistryKey().NodeKey.ToString();
 	AssetBasePrivate::DepthFirstTraversal(*this, [&](const FMetasoundAssetBase& ChildAsset)
 	{
 		TSet<const FMetasoundAssetBase*> Children;
@@ -677,8 +678,20 @@ bool FMetasoundAssetBase::IsReferencedAsset(const FMetasoundAssetBase& InAsset)
 			return Children;
 		}
 
+		// It's possible the referenced asset has been unloaded, 
+		// but is in the process of being reloaded, 
+		// so the asset may not be registered in the asset manager yet
+		// but the class key can still be checked
+		const TSet<FString>& AssetClassKeys = ChildAsset.GetReferencedAssetClassKeys();
+		if (AssetClassKeys.Contains(InAssetKey))
+		{
+			bIsReferenced = true;
+			return Children;
+		}
+
 		TArray<FMetasoundAssetBase*> ChildRefs;
 		ensureAlways(IMetaSoundAssetManager::GetChecked().TryLoadReferencedAssets(ChildAsset, ChildRefs));
+		
 		Algo::Transform(ChildRefs, Children, [](FMetasoundAssetBase* Child) { return Child; });
 		return Children;
 	});
