Improvements to sound asset digest generation:
* fixed trigger audioparameter type resolution
* kept omitted-variable comments from being stripped out
#rnx

diff --git a/Engine/Plugins/Runtime/Metasound/Source/MetasoundEngine/Private/MetasoundSource.cpp b/Engine/Plugins/Runtime/Metasound/Source/MetasoundEngine/Private/MetasoundSource.cpp
--- a/Engine/Plugins/Runtime/Metasound/Source/MetasoundEngine/Private/MetasoundSource.cpp
+++ b/Engine/Plugins/Runtime/Metasound/Source/MetasoundEngine/Private/MetasoundSource.cpp
@@ -231,21 +231,25 @@ namespace Metasound
 			FAudioParameter Params;
 			Params.ParamName = InParamName;
 			Params.TypeName = InTypeName;
+			const EMetasoundFrontendLiteralType InLiteralType = InLiteral.GetType(); 
 
-			switch (InLiteral.GetType())
+			if (Params.TypeName == TriggerName)
+			{
+				Params.ParamType = EAudioParameterType::Trigger;
+				check(InLiteralType == EMetasoundFrontendLiteralType::Boolean
+					|| InLiteralType == EMetasoundFrontendLiteralType::None);
+			}
+
+			switch (InLiteralType)
 			{
 				case EMetasoundFrontendLiteralType::Boolean:
 				{
-					if (Params.TypeName == TriggerName)
-					{
-						Params.ParamType = EAudioParameterType::Trigger;
-					}
-					else
+					if (Params.TypeName != TriggerName)
 					{
 						Params.ParamType = EAudioParameterType::Boolean;
 					}
-						
-					ensure(InLiteral.TryGet(Params.BoolParam));
+					
+					ensure(InLiteral.TryGet(Params.BoolParam));	
 				}
 				break;
 
