Fix for default values in metasound assets not being reflected verse/SG without editor restart

diff --git a/Engine/Plugins/Runtime/Metasound/Source/MetasoundEngine/Private/MetasoundBuilderSubsystem.cpp b/Engine/Plugins/Runtime/Metasound/Source/MetasoundEngine/Private/MetasoundBuilderSubsystem.cpp
--- a/Engine/Plugins/Runtime/Metasound/Source/MetasoundEngine/Private/MetasoundBuilderSubsystem.cpp
+++ b/Engine/Plugins/Runtime/Metasound/Source/MetasoundEngine/Private/MetasoundBuilderSubsystem.cpp
@@ -411,6 +411,7 @@ void UMetaSoundSourceBuilder::InitDelegates(Metasound::Frontend::FDocumentModify
 
 		OutDocumentDelegates.InterfaceDelegates.OnInputAdded.AddUObject(this, &UMetaSoundSourceBuilder::OnInputAdded);
 		OutDocumentDelegates.InterfaceDelegates.OnInputDataTypeChanged.AddUObject(this, &UMetaSoundSourceBuilder::OnInputDataTypeChanged);
+		OutDocumentDelegates.InterfaceDelegates.OnInputDefaultChanged.AddUObject(this, &UMetaSoundSourceBuilder::OnInputDefaultChanged);
 		OutDocumentDelegates.InterfaceDelegates.OnInputNameChanged.AddUObject(this, &UMetaSoundSourceBuilder::OnInputNameChanged);
 		OutDocumentDelegates.InterfaceDelegates.OnOutputAdded.AddUObject(this, &UMetaSoundSourceBuilder::OnOutputAdded);
 		OutDocumentDelegates.InterfaceDelegates.OnRemovingInput.AddUObject(this, &UMetaSoundSourceBuilder::OnRemovingInput);
@@ -543,6 +544,27 @@ void UMetaSoundSourceBuilder::OnInputDataTypeChanged(int32 InputIndex)
 	OnInputAdded(InputIndex);
 }
 
+void UMetaSoundSourceBuilder::OnInputDefaultChanged(int32 InputIndex) const
+{
+	using namespace Metasound::DynamicGraph;
+	using namespace Metasound::Frontend;
+
+	const FMetasoundFrontendDocument& Doc = Builder.GetConstDocumentChecked();
+	const FMetasoundFrontendGraphClass& GraphClass = Doc.RootGraph;
+	const FMetasoundFrontendClassInput& InputWithNewDefault = GraphClass.GetDefaultInterface().Inputs[InputIndex];
+
+	UMetaSoundSource& Source = GetMetaSoundSource();
+	Source.RuntimeInputData.InputMap.Remove(InputWithNewDefault.Name);
+
+	if (const FMetasoundFrontendClassInput* Input = GetConstBuilder().FindGraphInput(InputWithNewDefault.Name))
+	{
+		constexpr bool bCreateUObjectProxies = true;
+		Source.RuntimeInputData.InputMap.Add(Input->Name, UMetaSoundSource::CreateRuntimeInput(IDataTypeRegistry::Get(), *Input, bCreateUObjectProxies, UMetaSoundSettings::GetPageOrder()));
+	}
+
+	// default value changes to added/removed inputs with dynamic graphs is currently unsupported
+}
+
 void UMetaSoundSourceBuilder::OnInputNameChanged(FName OldName, FName NewName)
 {
 	using namespace Metasound::DynamicGraph;
diff --git a/Engine/Plugins/Runtime/Metasound/Source/MetasoundEngine/Public/MetasoundBuilderSubsystem.h b/Engine/Plugins/Runtime/Metasound/Source/MetasoundEngine/Public/MetasoundBuilderSubsystem.h
--- a/Engine/Plugins/Runtime/Metasound/Source/MetasoundEngine/Public/MetasoundBuilderSubsystem.h
+++ b/Engine/Plugins/Runtime/Metasound/Source/MetasoundEngine/Public/MetasoundBuilderSubsystem.h
@@ -125,6 +125,7 @@ private:
 	UE_API void OnEdgeAdded(int32 EdgeIndex) const;
 	UE_API void OnInputAdded(int32 InputIndex);
 	UE_API void OnInputDataTypeChanged(int32 InputIndex);
+	UE_API void OnInputDefaultChanged(int32 InputIndex) const;
 	UE_API void OnInputNameChanged(FName OldName, FName NewName);
 	UE_API void OnLiveComponentFinished(UAudioComponent* AudioComponent);
 	UE_API void OnNodeAdded(int32 NodeIndex) const;
