Protections around stale builder causing stale UI checking visibility, etc. to crash the editor

[RN]

diff --git a/Engine/Plugins/Runtime/Metasound/Source/MetasoundEditor/Private/SMetasoundGraphPin.h b/Engine/Plugins/Runtime/Metasound/Source/MetasoundEditor/Private/SMetasoundGraphPin.h
--- a/Engine/Plugins/Runtime/Metasound/Source/MetasoundEditor/Private/SMetasoundGraphPin.h
+++ b/Engine/Plugins/Runtime/Metasound/Source/MetasoundEditor/Private/SMetasoundGraphPin.h
@@ -340,14 +340,14 @@ namespace Metasound
 					.ToolTipText(LOCTEXT("ResetToClassDefaultToolTip", "Reset to class default"))
 					.ButtonStyle(FAppStyle::Get(), TEXT("NoBorder"))
 					.ContentPadding(0.0f)
-					.Visibility(TAttribute<EVisibility>::Create([this]
+					.Visibility(TAttribute<EVisibility>::CreateSPLambda(ParentPinType::AsShared(), [this]()
 					{
 						using namespace Frontend;
 						if (!ParentPinType::IsConnected())
 						{
 							UMetaSoundBuilderBase* Builder = GetBuilder();
 							UEdGraphPin* GraphPin = ParentPinType::GetPinObj();
-							if (!Builder || !GraphPin)
+							if (!Builder || !Builder->GetBuilder().IsValid() || !GraphPin)
 							{
 								return EVisibility::Collapsed;
 							}
@@ -380,7 +380,7 @@ namespace Metasound
 
 						return EVisibility::Collapsed;
 					}))
-					.OnClicked(FOnClicked::CreateLambda([this]()
+					.OnClicked(FOnClicked::CreateSPLambda(ParentPinType::AsShared(), [this]()
 					{
 						using namespace Editor;
 						using namespace Frontend;
@@ -415,9 +415,12 @@ namespace Metasound
 										if (UMetaSoundBuilderBase* Builder = GetBuilder())
 										{
 											FMetaSoundFrontendDocumentBuilder& DocBuilder = Builder->GetBuilder();
-											if (const FMetasoundFrontendVertex* Vertex = FGraphBuilder::GetPinVertex(DocBuilder, Pin))
+											if (DocBuilder.IsValid())
 											{
-												Builder->GetBuilder().RemoveNodeInputDefault(Node->GetNodeID(), Vertex->VertexID);
+												if (const FMetasoundFrontendVertex* Vertex = FGraphBuilder::GetPinVertex(DocBuilder, Pin))
+												{
+													Builder->GetBuilder().RemoveNodeInputDefault(Node->GetNodeID(), Vertex->VertexID);
+												}
 											}
 										}
 									}
