Document Config Checkpoint 1
- Adds basic serialization schema for config
- Adds checkpoint versions of inactive implementations including preset and randomizer
- Moves unique name generation to MetaSound Builder API from private MetaSound Editor Builder API
- Moves IsPreset calls to IsGraphEditable which maintains parity at this point (but will diverge with the implementation of DocConfig)
- Literal optimizations & MakeLiteral convenience functions

#rnx

diff --git a/Engine/Plugins/Experimental/TechAudioTools/Source/TechAudioToolsMetaSound/Private/ViewModels/MetaSoundViewModel.cpp b/Engine/Plugins/Experimental/TechAudioTools/Source/TechAudioToolsMetaSound/Private/ViewModels/MetaSoundViewModel.cpp
--- a/Engine/Plugins/Experimental/TechAudioTools/Source/TechAudioToolsMetaSound/Private/ViewModels/MetaSoundViewModel.cpp
+++ b/Engine/Plugins/Experimental/TechAudioTools/Source/TechAudioToolsMetaSound/Private/ViewModels/MetaSoundViewModel.cpp
@@ -1,7 +1,8 @@
-﻿// Copyright Epic Games, Inc. All Rights Reserved.
+// Copyright Epic Games, Inc. All Rights Reserved.
 
 #include "ViewModels/MetaSoundViewModel.h"
 
+#include "DocumentConfigurations/MetasoundPresetDocumentConfiguration.h"
 #include "Logging/StructuredLog.h"
 #include "MetasoundDocumentBuilderRegistry.h"
 #include "MetasoundFrontendPages.h"
@@ -148,7 +149,7 @@ void UMetaSoundViewModel::SetInputOverridesDefault(const FName& InputName, const
 
 void UMetaSoundViewModel::SetInputIsConstructorPin(const FName& InputName, const bool bIsConstructorPin) const
 {
-	if (const UMetaSoundInputViewModel* InputViewModel = *InputViewModels.Find(InputName))
+	if (const UMetaSoundInputViewModel* InputViewModel = InputViewModels.FindRef(InputName))
 	{
 		const FName DataType = InputViewModel->GetDataType();
 		if (!UMetaSoundViewModelConversionFunctions::IsConstructorType(DataType))
@@ -209,7 +210,7 @@ void UMetaSoundViewModel::SetOutputDataType(const FName& OutputName, const FName
 
 void UMetaSoundViewModel::SetOutputIsConstructorPin(const FName& OutputName, const bool bIsConstructorPin) const
 {
-	if (const UMetaSoundOutputViewModel* OutputViewModel = *OutputViewModels.Find(OutputName))
+	if (const UMetaSoundOutputViewModel* OutputViewModel = OutputViewModels.FindRef(OutputName))
 	{
 		const FName DataType = OutputViewModel->GetDataType();
 		if (!UMetaSoundViewModelConversionFunctions::IsConstructorType(DataType))
@@ -267,7 +268,7 @@ void UMetaSoundViewModel::OnInputNameChanged(const FName OldName, const FName Ne
 
 void UMetaSoundViewModel::OnInputDataTypeChanged(const FName VertexName, const FName DataType)
 {
-	if (UMetaSoundInputViewModel* InputViewModel = *InputViewModels.Find(VertexName))
+	if (UMetaSoundInputViewModel* InputViewModel = InputViewModels.FindRef(VertexName))
 	{
 		InputViewModel->OnInputDataTypeChanged(DataType);
 	}
@@ -275,7 +276,7 @@ void UMetaSoundViewModel::OnInputDataTypeChanged(const FName VertexName, const F
 
 void UMetaSoundViewModel::OnInputDefaultChanged(const FName VertexName, FMetasoundFrontendLiteral LiteralValue, const FName PageName)
 {
-	if (UMetaSoundInputViewModel* InputViewModel = *InputViewModels.Find(VertexName))
+	if (UMetaSoundInputViewModel* InputViewModel = InputViewModels.FindRef(VertexName))
 	{
 		InputViewModel->OnInputDefaultChanged(LiteralValue, PageName);
 	}
@@ -283,7 +284,7 @@ void UMetaSoundViewModel::OnInputDefaultChanged(const FName VertexName, FMetasou
 
 void UMetaSoundViewModel::OnInputInheritsDefaultChanged(const FName VertexName, const bool bInheritsDefault)
 {
-	if (UMetaSoundInputViewModel* InputViewModel = *InputViewModels.Find(VertexName))
+	if (UMetaSoundInputViewModel* InputViewModel = InputViewModels.FindRef(VertexName))
 	{
 		InputViewModel->OnInputInheritsDefaultChanged(bInheritsDefault);
 	}
@@ -291,7 +292,7 @@ void UMetaSoundViewModel::OnInputInheritsDefaultChanged(const FName VertexName,
 
 void UMetaSoundViewModel::OnInputIsConstructorPinChanged(const FName VertexName, const bool bIsConstructorPin)
 {
-	if (UMetaSoundInputViewModel* InputViewModel = *InputViewModels.Find(VertexName))
+	if (UMetaSoundInputViewModel* InputViewModel = InputViewModels.FindRef(VertexName))
 	{
 		InputViewModel->OnInputIsConstructorPinChanged(bIsConstructorPin);
 	}
@@ -337,7 +338,7 @@ void UMetaSoundViewModel::OnOutputNameChanged(const FName OldName, const FName N
 
 void UMetaSoundViewModel::OnOutputDataTypeChanged(const FName VertexName, const FName DataType)
 {
-	if (UMetaSoundOutputViewModel* OutputViewModel = *OutputViewModels.Find(VertexName))
+	if (UMetaSoundOutputViewModel* OutputViewModel = OutputViewModels.FindRef(VertexName))
 	{
 		OutputViewModel->OnOutputDataTypeChanged(DataType);
 	}
@@ -345,7 +346,7 @@ void UMetaSoundViewModel::OnOutputDataTypeChanged(const FName VertexName, const
 
 void UMetaSoundViewModel::OnOutputIsConstructorPinChanged(const FName VertexName, const bool bIsConstructorPin)
 {
-	if (UMetaSoundOutputViewModel* OutputViewModel = *OutputViewModels.Find(VertexName))
+	if (UMetaSoundOutputViewModel* OutputViewModel = OutputViewModels.FindRef(VertexName))
 	{
 		OutputViewModel->OnOutputIsConstructorPinChanged(bIsConstructorPin);
 	}
diff --git a/Engine/Plugins/Runtime/Metasound/Source/MetasoundEditor/Private/MetasoundDetailCustomization.cpp b/Engine/Plugins/Runtime/Metasound/Source/MetasoundEditor/Private/MetasoundDetailCustomization.cpp
--- a/Engine/Plugins/Runtime/Metasound/Source/MetasoundEditor/Private/MetasoundDetailCustomization.cpp
+++ b/Engine/Plugins/Runtime/Metasound/Source/MetasoundEditor/Private/MetasoundDetailCustomization.cpp
@@ -11,6 +11,7 @@
 #include "IAudioParameterTransmitter.h"
 #include "IDetailGroup.h"
 #include "Input/Events.h"
+#include "InstancedStructDetails.h"
 #include "Interfaces/MetasoundFrontendInterfaceRegistry.h"
 #include "MetasoundAssetBase.h"
 #include "MetasoundBuilderSubsystem.h"
@@ -39,6 +40,7 @@
 #include "Styling/StyleColors.h"
 #include "Templates/Casts.h"
 #include "Templates/SharedPointer.h"
+#include "UObject/UnrealType.h"
 #include "UObject/WeakObjectPtr.h"
 #include "UObject/WeakObjectPtrTemplates.h"
 #include "Widgets/Input/SCheckBox.h"
@@ -107,6 +109,11 @@ namespace Metasound::Editor
 	{
 	}
 
+	FName FMetasoundDetailCustomization::GetDocConfigPropertyPath() const
+	{
+		return BuildChildPath(DocumentPropertyName, GET_MEMBER_NAME_CHECKED(FMetasoundFrontendDocument, Configuration));
+	}
+
 	FName FMetasoundDetailCustomization::GetInterfaceVersionsPropertyPath() const
 	{
 		return BuildChildPath(DocumentPropertyName, GET_MEMBER_NAME_CHECKED(FMetasoundFrontendDocument, Interfaces));
@@ -132,7 +139,7 @@ namespace Metasound::Editor
 
 		// Only support modifying a single MetaSound at a time (Multiple
 		// MetaSound editing will be covered most likely by separate tool).
-		if (Objects.Num() > 1 || !Objects.Last().IsValid())
+		if (Objects.Num() != 1 || !Objects.Last().IsValid())
 		{
 			return;
 		}
@@ -172,8 +179,6 @@ namespace Metasound::Editor
 				const FName MajorVersionPropertyPath = BuildChildPath(VersionPropertyPath, GET_MEMBER_NAME_CHECKED(FMetasoundFrontendVersionNumber, Major));
 				const FName MinorVersionPropertyPath = BuildChildPath(VersionPropertyPath, GET_MEMBER_NAME_CHECKED(FMetasoundFrontendVersionNumber, Minor));
 
-				const FName InterfaceVersionsPropertyPath = GetInterfaceVersionsPropertyPath();
-
 				TSharedPtr<IPropertyHandle> AccessFlagsHandle = DetailLayout.GetProperty(AccessFlagsPropertyPath);
 				TSharedPtr<IPropertyHandle> AuthorHandle = DetailLayout.GetProperty(AuthorPropertyPath);
 				TSharedPtr<IPropertyHandle> CategoryHierarchyHandle = DetailLayout.GetProperty(CategoryHierarchyPropertyPath);
@@ -181,10 +186,49 @@ namespace Metasound::Editor
 				TSharedPtr<IPropertyHandle> DisplayNameHandle = DetailLayout.GetProperty(DisplayNamePropertyPath);
 				TSharedPtr<IPropertyHandle> DescHandle = DetailLayout.GetProperty(DescPropertyPath);
 				TSharedPtr<IPropertyHandle> KeywordsHandle = DetailLayout.GetProperty(KeywordsPropertyPath);
-				TSharedPtr<IPropertyHandle> InterfaceVersionsHandle = DetailLayout.GetProperty(InterfaceVersionsPropertyPath);
+				TSharedPtr<IPropertyHandle> InterfaceVersionsHandle = DetailLayout.GetProperty(GetInterfaceVersionsPropertyPath());
 				TSharedPtr<IPropertyHandle> MajorVersionHandle = DetailLayout.GetProperty(MajorVersionPropertyPath);
 				TSharedPtr<IPropertyHandle> MinorVersionHandle = DetailLayout.GetProperty(MinorVersionPropertyPath);
 
+				// Disabled as currently experimental and in local development.
+// 				{
+// 					TSharedPtr<IPropertyHandle> DocConfigHandle = DetailLayout.GetProperty(GetDocConfigPropertyPath());
+// 					if (DocConfigHandle.IsValid())
+// 					{
+// 						FSimpleDelegate OnConfigChanged = FSimpleDelegate::CreateLambda([this]()
+// 						{
+// 							if (Builder.IsValid())
+// 							{
+// 								FMetaSoundFrontendDocumentBuilder& DocBuilder = Builder->GetBuilder();
+// 								if (DocBuilder.IsValid())
+// 								{
+// 									DocBuilder.GetDocumentDelegates().OnDocumentConfigurationChanged.Broadcast({ });
+// 									DocBuilder.ConfigureDocument();
+// 								}
+// 							}
+// 						});
+// 						DocConfigHandle->SetOnPropertyValueChanged(OnConfigChanged);
+// 
+// 						const auto OnConfigPropChanged = TDelegate<void(const FPropertyChangedEvent&)>::CreateSPLambda(AsShared(), [this](const FPropertyChangedEvent& InEvent)
+// 						{
+// 							if (Builder.IsValid())
+// 							{
+// 								FMetaSoundFrontendDocumentBuilder& DocBuilder = Builder->GetBuilder();
+// 								if (DocBuilder.IsValid())
+// 								{
+// 									if (const FMetaSoundFrontendDocumentConfiguration* Config = DocBuilder.GetConstDocumentConfiguration().GetPtr())
+// 									{
+// 										const_cast<FMetaSoundFrontendDocumentConfiguration*>(Config)->OnPropertyChanged(InEvent, DocBuilder);
+// 									}
+// 								}
+// 							}
+// 						});
+// 						DocConfigHandle->SetOnChildPropertyValueChangedWithData(OnConfigPropChanged);
+// 					}
+// 
+// 					GeneralCategoryBuilder.AddProperty(DocConfigHandle);
+// 				}
+
 				// Invalid for UMetaSounds
 				TSharedPtr<IPropertyHandle> OutputFormat = DetailLayout.GetProperty(GET_MEMBER_NAME_CHECKED(UMetaSoundSource, OutputFormat));
 				if (OutputFormat.IsValid() && OutputFormat->IsValidHandle())
@@ -366,7 +410,7 @@ namespace Metasound::Editor
 
 		// Only support modifying a single MetaSound at a time (Multiple
 		// MetaSound editing will be covered most likely by separate tool).
-		if (Objects.Num() > 1)
+		if (Objects.Num() != 1 || !Objects.Last().IsValid())
 		{
 			return;
 		}
@@ -394,9 +438,9 @@ namespace Metasound::Editor
 				using namespace Engine;
 				using namespace Frontend;
 
-				if (InSelectInfo != ESelectInfo::OnNavigation)
+				if (InSelectInfo != ESelectInfo::OnNavigation && IsBuilderValid())
 				{
-					UObject& MetaSound = GetMetaSound();
+					UObject& MetaSound = GetMetaSoundChecked();
 
 					const FScopedTransaction Transaction(FText::Format(LOCTEXT("AddPageTransactionFormat", "Add MetaSound Page Graph '{0}'"), FText::FromString(*NameToAdd)));
 					MetaSound.Modify();
@@ -410,11 +454,13 @@ namespace Metasound::Editor
 					EMetaSoundBuilderResult Result = EMetaSoundBuilderResult::Failed;
 					Builder->AddGraphPage(FName(*NameToAdd.Get()), bDuplicateLastGraph, bSetAsBuildGraph, Result);
 
-					IAssetEditorInstance* AssetEditor = GEditor->GetEditorSubsystem<UAssetEditorSubsystem>()->FindEditorForAsset(&GetMetaSound(), /*bFocusIfOpen=*/false);
-					FEditor* MetaSoundEditor = static_cast<FEditor*>(AssetEditor);
-					if (MetaSoundEditor)
+					if (UAssetEditorSubsystem* AssetSubsystem = GEditor->GetEditorSubsystem<UAssetEditorSubsystem>())
 					{
-						MetaSoundEditor->RefreshDetails();
+						IAssetEditorInstance* AssetEditor = AssetSubsystem->FindEditorForAsset(&MetaSound, /*bFocusIfOpen=*/false);
+						if (FEditor* MetaSoundEditor = static_cast<FEditor*>(AssetEditor))
+						{
+							MetaSoundEditor->RefreshDetails();
+						}
 					}
 				}
 			})
@@ -443,7 +489,7 @@ namespace Metasound::Editor
 				PropertyCustomizationHelpers::MakeDeleteButton(FSimpleDelegate::CreateLambda([this]()
 				{
 					using namespace Frontend;
-					UObject& MetaSound = GetMetaSound();
+					UObject& MetaSound = GetMetaSoundChecked();
 
 					const FScopedTransaction Transaction(LOCTEXT("ResetGraphPagesTransaction", "Reset MetaSound Graph Pages"));
 					MetaSound.Modify();
@@ -481,18 +527,32 @@ namespace Metasound::Editor
 		}
 	}
 
-	UObject& FMetasoundPagesDetailCustomization::GetMetaSound() const
+	const FMetasoundFrontendDocument& FMetasoundPagesDetailCustomization::GetConstDocumentChecked() const
+	{
+		return Builder->GetBuilder().GetConstDocumentChecked();
+	}
+
+	UObject& FMetasoundPagesDetailCustomization::GetMetaSoundChecked() const
 	{
 		return Builder->GetBuilder().CastDocumentObjectChecked<UObject>();
 	}
 
+	bool FMetasoundPagesDetailCustomization::IsBuilderValid() const
+	{
+		return Builder->GetBuilder().IsValid();
+	}
+
 	void FMetasoundPagesDetailCustomization::RebuildImplemented()
 	{
 		EntryWidgets->ClearChildren();
 
+		if (!IsBuilderValid())
+		{
+			return;
+		}
+
 		TSet<FGuid> ImplementedGuids;
-		const FMetasoundFrontendDocument& Document = Builder->GetBuilder().GetConstDocumentChecked();
-		Document.RootGraph.IterateGraphPages([&ImplementedGuids](const FMetasoundFrontendGraph& Graph)
+		GetConstDocumentChecked().RootGraph.IterateGraphPages([&ImplementedGuids](const FMetasoundFrontendGraph& Graph)
 		{
 			ImplementedGuids.Add(Graph.PageID);
 		});
@@ -510,7 +570,7 @@ namespace Metasound::Editor
 					{
 						constexpr bool bOpenEditor = false; // Already focused by user action
 						UMetaSoundEditorSubsystem::GetConstChecked().SetFocusedPage(*Builder.Get(), PageID, bOpenEditor);
-						IAssetEditorInstance* AssetEditor = GEditor->GetEditorSubsystem<UAssetEditorSubsystem>()->FindEditorForAsset(&GetMetaSound(), /*bFocusIfOpen=*/false);
+						IAssetEditorInstance* AssetEditor = GEditor->GetEditorSubsystem<UAssetEditorSubsystem>()->FindEditorForAsset(&GetMetaSoundChecked(), /*bFocusIfOpen=*/false);
 						FEditor* MetaSoundEditor = static_cast<FEditor*>(AssetEditor);
 						if (MetaSoundEditor)
 						{
@@ -557,7 +617,7 @@ namespace Metasound::Editor
 					{
 						using namespace Frontend;
 						const FScopedTransaction Transaction(FText::Format(LOCTEXT("RemovePageTransactionFormat", "Remove MetaSound Page '{0}'"), FText::FromName(InName)));
-						UObject& MetaSound = GetMetaSound();
+						UObject& MetaSound = GetMetaSoundChecked();
 						MetaSound.Modify();
 
 						// Removal may modify the builder's build page ID if it is the currently set value
@@ -673,11 +733,16 @@ namespace Metasound::Editor
 		AddableItems.Reset();
 		ImplementedNames.Reset();
 
+		if (!IsBuilderValid())
+		{
+			return;
+		}
+
 		const UMetaSoundSettings* Settings = GetDefault<UMetaSoundSettings>();
 		check(Settings);
 
 		TSet<FGuid> ImplementedGuids;
-		const FMetasoundFrontendDocument& Document = Builder->GetBuilder().GetConstDocumentChecked();
+		const FMetasoundFrontendDocument& Document = GetConstDocumentChecked();
 		Document.RootGraph.IterateGraphPages([&ImplementedGuids](const FMetasoundFrontendGraph& Graph)
 		{
 			ImplementedGuids.Add(Graph.PageID);
@@ -728,7 +793,7 @@ namespace Metasound::Editor
 				if (PageSettings->Name != ParentPtr->BuildPageName)
 				{
 					ParentPtr->BuildPageName = PageSettings->Name;
-					FGraphBuilder::RegisterGraphWithFrontend(ParentPtr->GetMetaSound());
+					FGraphBuilder::RegisterGraphWithFrontend(ParentPtr->GetMetaSoundChecked());
 				}
 
 				ParentPtr->AddableItems.RemoveAll([&PageSettings](const TSharedPtr<FString>& Item) { return Item->Compare(PageSettings->Name.ToString()) == 0; });
@@ -765,7 +830,7 @@ namespace Metasound::Editor
 				if (PageSettings->Name != ParentPtr->BuildPageName)
 				{
 					ParentPtr->BuildPageName = PageSettings->Name;
-					FGraphBuilder::RegisterGraphWithFrontend(ParentPtr->GetMetaSound());
+					FGraphBuilder::RegisterGraphWithFrontend(ParentPtr->GetMetaSoundChecked());
 				}
 
 				ParentPtr->AddableItems.Add(MakeShared<FString>(PageSettings->Name.ToString()));
@@ -783,7 +848,7 @@ namespace Metasound::Editor
 
 		// Only support modifying a single MetaSound at a time (Multiple
 		// MetaSound editing will be covered most likely by separate tool).
-		if (Objects.Num() > 1)
+		if (Objects.Num() != 1 || !Objects.Last().IsValid())
 		{
 			return;
 		}
diff --git a/Engine/Plugins/Runtime/Metasound/Source/MetasoundEditor/Private/MetasoundDetailCustomization.h b/Engine/Plugins/Runtime/Metasound/Source/MetasoundEditor/Private/MetasoundDetailCustomization.h
--- a/Engine/Plugins/Runtime/Metasound/Source/MetasoundEditor/Private/MetasoundDetailCustomization.h
+++ b/Engine/Plugins/Runtime/Metasound/Source/MetasoundEditor/Private/MetasoundDetailCustomization.h
@@ -2,6 +2,7 @@
 #pragma once
 
 #include "IDetailCustomization.h"
+#include "MetasoundFrontendDocument.h"
 #include "MetasoundFrontendDocumentModifyDelegates.h"
 #include "SGraphActionMenu.h"
 #include "SSearchableComboBox.h"
@@ -47,6 +48,7 @@ namespace Metasound
 			// End of IDetailCustomization interface
 
 		private:
+			FName GetDocConfigPropertyPath() const;
 			FName GetInterfaceVersionsPropertyPath() const;
 			FName GetRootClassPropertyPath() const;
 			FName GetMetadataPropertyPath() const;
@@ -64,7 +66,9 @@ namespace Metasound
 			// End of IDetailCustomization interface
 
 		private:
-			UObject& GetMetaSound() const;
+			const FMetasoundFrontendDocument& GetConstDocumentChecked() const;
+			UObject& GetMetaSoundChecked() const;
+			bool IsBuilderValid() const;
 			void RebuildImplemented();
 			void RefreshView();
 			void UpdateItemNames();
diff --git a/Engine/Plugins/Runtime/Metasound/Source/MetasoundEditor/Private/MetasoundEditor.cpp b/Engine/Plugins/Runtime/Metasound/Source/MetasoundEditor/Private/MetasoundEditor.cpp
--- a/Engine/Plugins/Runtime/Metasound/Source/MetasoundEditor/Private/MetasoundEditor.cpp
+++ b/Engine/Plugins/Runtime/Metasound/Source/MetasoundEditor/Private/MetasoundEditor.cpp
@@ -159,6 +159,11 @@ namespace Metasound
 
 		bool PageEditorEnabled(const FMetaSoundFrontendDocumentBuilder& Builder, bool bHasProjectPageValues, bool bPresetCanEditPageValues)
 		{
+			if (!Builder.IsValid())
+			{
+				return false;
+			}
+
 			IMetasoundEditorModule& MetaSoundEditorModule = FModuleManager::GetModuleChecked<IMetasoundEditorModule>("MetaSoundEditor");
 			if (MetaSoundEditorModule.IsRestrictedMode())
 			{
@@ -583,10 +588,10 @@ namespace Metasound
 						if (GetHoveredGraph() == OwningGraph)
 						{
 							const UMetaSoundBuilderBase& Builder = OwningGraph->GetBuilderChecked();
-							const bool bIsPreset = Builder.IsPreset();
-							if (bIsPreset)
+							const bool bIsEditable = Builder.GetConstBuilder().GetIsGraphEditable();
+							if (!bIsEditable)
 							{
-								Message = FText::Format(LOCTEXT("DropTargetFailIsPreset", "'{0}': Graph is Preset"), GraphMember->GetDisplayName());
+								Message = FText::Format(LOCTEXT("DropTargetFailIsPreset", "'{0}': Graph is not editable"), GraphMember->GetDisplayName());
 							}
 							else if (UMetasoundEditorGraphInput* Input = Cast<UMetasoundEditorGraphInput>(GraphMember.Get()))
 							{
@@ -1689,12 +1694,11 @@ namespace Metasound
 				.OnActionDragged(this, &FEditor::OnActionDragged)
 				.OnActionMatchesName(this, &FEditor::HandleActionMatchesName)
 				.OnActionSelected(this, &FEditor::OnActionSelected)
-// 				.OnCategoryTextCommitted(this, &FEditor::OnCategoryNameCommitted)
 				.OnCollectAllActions(this, &FEditor::CollectAllActions)
 				.OnCollectStaticSections(this, &FEditor::CollectStaticSections)
- 				.OnContextMenuOpening(this, &FEditor::OnContextMenuOpening)
+				.OnContextMenuOpening(this, &FEditor::OnContextMenuOpening)
 				.OnCreateWidgetForAction(this, &FEditor::OnCreateWidgetForAction)
-  				.OnCanRenameSelectedAction(this, &FEditor::CanRenameOnActionNode)
+				.OnCanRenameSelectedAction(this, &FEditor::CanRenameOnActionNode)
 				.OnGetFilterText(this, &FEditor::GetFilterText)
 				.OnGetSectionTitle(this, &FEditor::OnGetSectionTitle)
 				.OnGetSectionWidget(this, &FEditor::OnGetMenuSectionWidget)
@@ -1747,7 +1751,7 @@ namespace Metasound
 				PagesDetails->SetVisibility(VisibilityAttr);
 			}
 
-			if (Builder.IsValid() && Builder->IsPreset())
+			if (Builder.IsValid() && !Builder->GetConstBuilder().GetIsGraphEditable())
 			{
 				FLevelEditorModule& LevelEditor = FModuleManager::LoadModuleChecked<FLevelEditorModule>("LevelEditor");
 				LevelEditor.OnMapChanged().AddSP(this, &FEditor::ChangeUserPresetWidgetWorld);
@@ -1797,7 +1801,7 @@ namespace Metasound
 		{
 			using namespace Metasound::Frontend;
 
-			if (!Builder.IsValid() || !Builder->IsPreset())
+			if (!Builder.IsValid() || Builder->GetConstBuilder().GetIsGraphEditable())
 			{
 				return nullptr;
 			}
@@ -3047,8 +3051,8 @@ namespace Metasound
 					FCanExecuteAction::CreateLambda([this]() { return CanPromoteToDeferredVariable(); }));
 
 				GraphEditorCommands->MapAction(FGraphEditorCommands::Get().HideNoConnectionPins,
-					FExecuteAction::CreateSP(this, &FEditor::HideUnconnectedPins));				
-				
+					FExecuteAction::CreateSP(this, &FEditor::HideUnconnectedPins));
+
 				GraphEditorCommands->MapAction(FGraphEditorCommands::Get().ShowAllPins,
 					FExecuteAction::CreateSP(this, &FEditor::ShowUnconnectedPins));
 
@@ -3213,7 +3217,7 @@ namespace Metasound
 					}
 				}
 
-				if (NextToSelect)
+				if (NextToSelect && GraphMembersMenu.IsValid())
 				{
 					if (GraphMembersMenu->SelectItemByName(NextToSelect->GetMemberName(), ESelectInfo::Direct, static_cast<int32>(NextToSelect->GetSectionID())))
 					{
@@ -3639,11 +3643,28 @@ namespace Metasound
 			return nullptr;
 		}
 
+		const FMetaSoundFrontendDocumentConfiguration* FEditor::GetConfiguration() const
+		{
+			if (const FMetaSoundFrontendDocumentBuilder* DocBuilder = GetFrontendBuilder())
+			{
+				if (DocBuilder->IsValid())
+				{
+					return DocBuilder->GetConstDocumentConfiguration().GetPtr();
+				}
+			}
+
+			return nullptr;
+		}
+
 		FMetaSoundFrontendDocumentBuilder* FEditor::GetFrontendBuilder() const
 		{
 			if (Builder.IsValid())
 			{
-				return &Builder->GetBuilder();
+				FMetaSoundFrontendDocumentBuilder& DocBuilder = Builder->GetBuilder();
+				if (DocBuilder.IsValid())
+				{
+					return &DocBuilder;
+				}
 			}
 
 			return nullptr;
@@ -3673,15 +3694,18 @@ namespace Metasound
 				}
 
 				// Renameable member nodes (inputs/outputs/variables)
-				if (UMetasoundEditorGraphMemberNode* MemberNode = Cast<UMetasoundEditorGraphMemberNode>(*SelectedIter))
+				if (GraphMembersMenu.IsValid())
 				{
-					if (const UMetasoundEditorGraphMember* Member = MemberNode->GetMember()) 
+					if (UMetasoundEditorGraphMemberNode* MemberNode = Cast<UMetasoundEditorGraphMemberNode>(*SelectedIter))
 					{
-						if (Member->CanRename())
+						if (const UMetasoundEditorGraphMember* Member = MemberNode->GetMember())
 						{
-							GraphMembersMenu->SelectItemByName(Member->GetMemberName(), ESelectInfo::Direct, static_cast<int32>(Member->GetSectionID()));
-							GraphMembersMenu->RefreshAllActions(/*bPreserveExpansion=*/ true, /*bHandleOnSelectionEvent=*/ true);
-							GraphMembersMenu->OnRequestRenameOnActionNode();
+							if (Member->CanRename())
+							{
+								GraphMembersMenu->SelectItemByName(Member->GetMemberName(), ESelectInfo::Direct, static_cast<int32>(Member->GetSectionID()));
+								GraphMembersMenu->RefreshAllActions(/*bPreserveExpansion=*/ true, /*bHandleOnSelectionEvent=*/ true);
+								GraphMembersMenu->OnRequestRenameOnActionNode();
+							}
 						}
 					}
 				}
@@ -3801,7 +3825,7 @@ namespace Metasound
 					// Duplicate the Sources Frontend Node and add a new member from it
 					if (const UMetasoundEditorGraphVariable* SourceGraphVariable = Cast<UMetasoundEditorGraphVariable>(SourceGraphMember))
 					{
-						const FName NewName = FGraphBuilder::GenerateUniqueVariableName(DocBuilder, SourceGraphVariable->GetMemberName().ToString());
+						const FName NewName = DocBuilder.GenerateUniqueVariableName(SourceGraphVariable->GetMemberName().ToString());
 						const FMetasoundFrontendVariable* NewFrontendVariable = DocBuilder.DuplicateGraphVariable(SourceGraphVariable->GetMemberName(), NewName);
 						if (ensure(NewFrontendVariable))
 						{
@@ -3813,10 +3837,9 @@ namespace Metasound
 						const FName SourceMemberName = SourceGraphVertex->GetMemberName();
 						const EMetasoundFrontendClassType ClassType = SourceGraphVertex->GetClassType();
 						
-						const FName Name = FGraphBuilder::GenerateUniqueNameByClassType(Graph.GetMetasoundChecked(), ClassType, SourceMemberName.ToString());
-
 						if (ClassType == EMetasoundFrontendClassType::Input)
 						{
+							const FName Name = DocBuilder.GenerateUniqueInputName(SourceMemberName.ToString());
 							if (const FMetasoundFrontendClassInput* NewInput = DocBuilder.DuplicateGraphInput(SourceMemberName, Name))
 							{
 								FGraphBuilder::SynchronizeGraphMembers(DocBuilder, Graph);
@@ -3825,6 +3848,7 @@ namespace Metasound
 						}
 						else if (ClassType == EMetasoundFrontendClassType::Output)
 						{
+							const FName Name = DocBuilder.GenerateUniqueOutputName(SourceMemberName.ToString());
 							if (const FMetasoundFrontendClassOutput* NewOutput = DocBuilder.DuplicateGraphOutput(SourceMemberName, Name))
 							{
 								FGraphBuilder::SynchronizeGraphMembers(DocBuilder, Graph);
@@ -4024,7 +4048,7 @@ namespace Metasound
 				else if (const UMetasoundEditorGraphVariable* SourceGraphVariable = Cast<UMetasoundEditorGraphVariable>(Member))
 				{
 					const FMetasoundEditorGraphVariableBreadcrumb Breadcrumb = SourceGraphVariable->GetBreadcrumb();
-					const FName VariableName = FGraphBuilder::GenerateUniqueVariableName(DocumentBuilder, Breadcrumb.MemberName.ToString());
+					const FName VariableName = DocumentBuilder.GenerateUniqueVariableName(Breadcrumb.MemberName.ToString());
 					const FMetasoundFrontendVariable* FrontendVariable = DocumentBuilder.AddGraphVariable(
 						VariableName,
 						SourceGraphVariable->GetDataType(),
@@ -4337,14 +4361,14 @@ PRAGMA_ENABLE_DEPRECATION_WARNINGS
 
 		void FEditor::CollectStaticSections(TArray<int32>& StaticSectionIDs)
 		{
-			const bool bIsPreset = Builder.IsValid() ? Builder->IsPreset() : false;
+			const bool bIsReadOnly = Builder.IsValid() ? !Builder->GetConstBuilder().GetIsGraphEditable() : false;
 
 			for (int32 i = 0; i < static_cast<int32>(ENodeSection::COUNT); ++i)
 			{
 				if (static_cast<ENodeSection>(i) != ENodeSection::None)
 				{
-					// Presets do not have variables
-					if (bIsPreset && static_cast<ENodeSection>(i) == ENodeSection::Variables)
+					// ReadOnly graphs do not show variables
+					if (bIsReadOnly && static_cast<ENodeSection>(i) == ENodeSection::Variables)
 					{
 						continue;
 					}
@@ -4588,9 +4612,9 @@ PRAGMA_ENABLE_DEPRECATION_WARNINGS
 			TArray<TSharedPtr<FEdGraphSchemaAction>> Actions;
 			GraphMembersMenu->GetSelectedActions(Actions);
 
-			if (Actions.IsEmpty())//Section is selected
+			if (Actions.IsEmpty()) //Section is selected
 			{
-				if (!Builder->IsPreset())
+				if (Builder->GetConstBuilder().GetIsGraphEditable())
 				{
 					MenuBuilder.BeginSection("GraphActionMenuSectionActions", LOCTEXT("SectionActionsMenuHeader", "Section Actions"));
 
@@ -5029,7 +5053,7 @@ PRAGMA_ENABLE_DEPRECATION_WARNINGS
 					MetaSound->Modify();
 					Graph.Modify();
 
-					const FName VariableName = FGraphBuilder::GenerateUniqueVariableName(DocBuilder, TEXT("Variable"));
+					const FName VariableName = DocBuilder.GenerateUniqueVariableName(TEXT("Variable"));
 					const FMetasoundFrontendVariable* FrontendVariable = DocBuilder.AddGraphVariable(VariableName, DataTypeName);
 					if (ensure(FrontendVariable))
 					{
@@ -5501,7 +5525,7 @@ PRAGMA_ENABLE_DEPRECATION_WARNINGS
 					const FMetasoundFrontendVertex* InputVertex = DocBuilder.FindNodeInput(InputVertexHandle.NodeID, InputVertexHandle.VertexID);
 					check(InputVertex);
 
-					FName Name = FGraphBuilder::GenerateUniqueNameByClassType(ParentMetasound, EMetasoundFrontendClassType::Input, Pin->GetName());
+					FName Name = DocBuilder.GenerateUniqueInputName(Pin->GetName());
 
 					EMetaSoundBuilderResult Result = EMetaSoundBuilderResult::Failed;
 					FMetasoundFrontendLiteral Literal;
@@ -5608,7 +5632,7 @@ PRAGMA_ENABLE_DEPRECATION_WARNINGS
 				FName TypeName = Pair.Key.Value;
 				UEdGraphPin* SourcePin = Pair.Value[0];
 				FMetasoundFrontendVertexHandle InputVertexHandle = FGraphBuilder::GetPinVertexHandle(DocBuilder, SourcePin);
-				FName InputName = FGraphBuilder::GenerateUniqueNameByClassType(ParentMetasound, EMetasoundFrontendClassType::Input, PinName.ToString());
+				FName InputName = DocBuilder.GenerateUniqueInputName(PinName.ToString());
 				
 				EMetaSoundBuilderResult Result = EMetaSoundBuilderResult::Failed;
 				FMetasoundFrontendLiteral Literal;
diff --git a/Engine/Plugins/Runtime/Metasound/Source/MetasoundEditor/Private/MetasoundEditor.h b/Engine/Plugins/Runtime/Metasound/Source/MetasoundEditor/Private/MetasoundEditor.h
--- a/Engine/Plugins/Runtime/Metasound/Source/MetasoundEditor/Private/MetasoundEditor.h
+++ b/Engine/Plugins/Runtime/Metasound/Source/MetasoundEditor/Private/MetasoundEditor.h
@@ -49,6 +49,7 @@ class IToolkitHost;
 class UUserWidget;
 
 struct FMetaSoundFrontendDocumentBuilder;
+struct FMetaSoundFrontendDocumentConfiguration;
 
 namespace Metasound::Editor
 {
@@ -135,6 +136,7 @@ namespace Metasound::Editor
 
 		UAudioComponent* GetAudioComponent() const;
 
+		const FMetaSoundFrontendDocumentConfiguration* GetConfiguration() const;
 		FMetaSoundFrontendDocumentBuilder* GetFrontendBuilder() const;
 
 		/** IMetasoundEditor interface */
@@ -345,7 +347,7 @@ namespace Metasound::Editor
 		bool CanCopySelectedMemberItems() const;
 
 		/** Copy selected Member items. */
-		void CopySelectedMemberItems();		
+		void CopySelectedMemberItems();
 		
 		/** Whether the currently selected Member item(s) can be cut. */
 		bool CanCutSelectedMemberItems() const;
@@ -493,7 +495,7 @@ namespace Metasound::Editor
 		/** Converts the MetaSound from a preset to a fully modifiable MetaSound. */
 		void ConvertFromPreset();
 
-		/* Whether or not page details should be visible. */
+		/* Whether or not page details should be visible/enabled. */
 		bool ShowPageGraphDetails() const;
 
 		/** Creates audition menu options */
diff --git a/Engine/Plugins/Runtime/Metasound/Source/MetasoundEditor/Private/MetasoundEditorDocumentClipboardUtils.cpp b/Engine/Plugins/Runtime/Metasound/Source/MetasoundEditor/Private/MetasoundEditorDocumentClipboardUtils.cpp
--- a/Engine/Plugins/Runtime/Metasound/Source/MetasoundEditor/Private/MetasoundEditorDocumentClipboardUtils.cpp
+++ b/Engine/Plugins/Runtime/Metasound/Source/MetasoundEditor/Private/MetasoundEditorDocumentClipboardUtils.cpp
@@ -266,7 +266,7 @@ namespace Metasound::Editor
 				else
 				{
 					const FMetasoundFrontendLiteral* Literal = Breadcrumb.DefaultLiterals.Find(Frontend::DefaultPageID);
-					const FName VariableName = FGraphBuilder::GenerateUniqueVariableName(DocBuilder, BaseName.ToString());
+					const FName VariableName = DocBuilder.GenerateUniqueVariableName(BaseName.ToString());
 					FrontendVariable = DocBuilder.AddGraphVariable(
 						VariableName,
 						Breadcrumb.DataType,
diff --git a/Engine/Plugins/Runtime/Metasound/Source/MetasoundEditor/Private/MetasoundEditorGraph.cpp b/Engine/Plugins/Runtime/Metasound/Source/MetasoundEditor/Private/MetasoundEditorGraph.cpp
--- a/Engine/Plugins/Runtime/Metasound/Source/MetasoundEditor/Private/MetasoundEditorGraph.cpp
+++ b/Engine/Plugins/Runtime/Metasound/Source/MetasoundEditor/Private/MetasoundEditorGraph.cpp
@@ -376,7 +376,7 @@ bool UMetasoundEditorGraphVertex::SetIsAdvancedDisplay(const bool IsAdvancedDisp
 		return false;
 	}
 
-	const FText TransactionLabel = FText::Format(LOCTEXT("SetGraphVertexIsAdvancedDisplayState", "Set Metasound {0} IsAdvancedDislay"), GetGraphMemberLabel());
+	const FText TransactionLabel = FText::Format(LOCTEXT("SetGraphVertexIsAdvancedDisplayState", "Set Metasound {0} IsAdvancedDisplay"), GetGraphMemberLabel());
 	const FScopedTransaction Transaction(TransactionLabel, true);
 
 	Graph->Modify();
@@ -500,7 +500,7 @@ bool UMetasoundEditorGraphVertex::NameContainsInterfaceNamespace(FMetasoundFront
 bool UMetasoundEditorGraphVertex::CanRename() const
 {
 	FMetaSoundFrontendDocumentBuilder& Builder = GetFrontendBuilderChecked();
-	return !Builder.IsPreset() && !IsInterfaceMember();
+	return Builder.GetIsGraphEditable() && !IsInterfaceMember();
 }
 
 bool UMetasoundEditorGraphVertex::CanRename(const FText& InNewText, FText& OutError) const
@@ -533,9 +533,9 @@ bool UMetasoundEditorGraphVertex::CanRename(const FText& InNewText, FText& OutEr
 	}
 
 	const FMetaSoundFrontendDocumentBuilder& Builder = GetFrontendBuilderChecked();
-	if (Builder.IsPreset())
+	if (!Builder.GetIsGraphEditable())
 	{
-		OutError = FText::Format(LOCTEXT("GraphVertexRenameInvalid_Preset", "{0} is a vertex in a preset graph and cannot be renamed."), InNewText);
+		OutError = FText::Format(LOCTEXT("GraphVertexRenameInvalid_Preset", "{0} is a vertex in a read-only graph and cannot be renamed."), InNewText);
 		return false;
 	}
 
@@ -2601,12 +2601,14 @@ void UMetasoundEditorGraph::MigrateEditorDocumentData(FMetaSoundFrontendDocument
 		}
 #endif // WITH_EDITOR
 
+PRAGMA_DISABLE_DEPRECATION_WARNINGS
 		// Presets get rebuilt anyway and may have invalid connections (newly referenced vertices)
 		// that need to be rebuilt later in asset load.
-		if (OutBuilder.IsPreset())
+		if (OutBuilder.GetConstDocumentChecked().RootGraph.PresetOptions.bIsPreset)
 		{
 			return;
 		}
+PRAGMA_ENABLE_DEPRECATION_WARNINGS
 
 		// Cache data to be used in edge swapping below, being careful to not reference the input node pointer
 		// or vertex reference within the lower inner loop as the underlying node array may be reallocated by
@@ -2709,10 +2711,12 @@ void UMetasoundEditorGraph::MigrateEditorDocumentData(FMetaSoundFrontendDocument
 		}
 	});
 
+PRAGMA_DISABLE_DEPRECATION_WARNINGS
 	// 4. Add comment nodes as builder graph comments to frontend document
 	// (No need to propagate comments for presets)
-	if (!OutBuilder.IsPreset())
+	if (!OutBuilder.GetConstDocumentChecked().RootGraph.PresetOptions.bIsPreset)
 	{
+PRAGMA_ENABLE_DEPRECATION_WARNINGS
 		TArray<UEdGraphNode_Comment*> CommentNodes;
 		GetNodesOfClass(CommentNodes);
 		for (UEdGraphNode_Comment* Node : CommentNodes)
diff --git a/Engine/Plugins/Runtime/Metasound/Source/MetasoundEditor/Private/MetasoundEditorGraphBuilder.cpp b/Engine/Plugins/Runtime/Metasound/Source/MetasoundEditor/Private/MetasoundEditorGraphBuilder.cpp
--- a/Engine/Plugins/Runtime/Metasound/Source/MetasoundEditor/Private/MetasoundEditorGraphBuilder.cpp
+++ b/Engine/Plugins/Runtime/Metasound/Source/MetasoundEditor/Private/MetasoundEditorGraphBuilder.cpp
@@ -94,20 +94,6 @@ namespace Metasound
 				return nullptr;
 			}
 
-			FName GenerateUniqueName(const TSet<FName>& InExistingNames, const FString& InBaseName)
-			{
-				int32 PostFixInt = 0;
-				FString NewName = InBaseName;
-
-				while (InExistingNames.Contains(*NewName))
-				{
-					PostFixInt++;
-					NewName = FString::Format(TEXT("{0} {1}"), { InBaseName, PostFixInt });
-				}
-
-				return FName(*NewName);
-			}
-
 			void RecurseClearDocumentModified(FMetasoundAssetBase& InAssetBase)
 			{
 				using namespace Metasound::Frontend;
@@ -628,18 +614,6 @@ namespace Metasound
 			return Categories;
 		}
 
-		FName FGraphBuilder::GenerateUniqueNameByClassType(UObject& InMetaSound, EMetasoundFrontendClassType InClassType, const FString& InBaseName)
-		{
-			using namespace Engine;
-
-			TSet<FName> ExistingNames;
-			FMetaSoundFrontendDocumentBuilder& Builder = FDocumentBuilderRegistry::GetChecked().FindOrBeginBuilding(&InMetaSound);
-			auto GetName = [&](const FMetasoundFrontendClass&, const FMetasoundFrontendNode& Node) { ExistingNames.Add(Node.Name); };
-			Builder.IterateNodesByClassType(GetName, InClassType);
-
-			return GraphBuilderPrivate::GenerateUniqueName(ExistingNames, InBaseName);
-		}
-
 		TSharedPtr<FEditor> FGraphBuilder::GetEditorForMetasound(const UObject& Metasound)
 		{
 			// TODO: FToolkitManager is deprecated. Replace with UAssetEditorSubsystem.
@@ -1010,13 +984,17 @@ PRAGMA_ENABLE_DEPRECATION_WARNINGS
 		bool FGraphBuilder::IsPreviewingMetaSound(const UObject& InMetaSound)
 		{
 			using namespace Engine;
-			if (const UAudioComponent* PreviewComponent = GEditor->GetPreviewAudioComponent())
+
+			if (GEditor)
 			{
-				if (PreviewComponent->IsPlaying())
+				if (const UAudioComponent* PreviewComponent = GEditor->GetPreviewAudioComponent())
 				{
-					if (const USoundBase* Sound = PreviewComponent->Sound)
+					if (PreviewComponent->IsPlaying())
 					{
-						return Sound->GetUniqueID() == InMetaSound.GetUniqueID();
+						if (const USoundBase* Sound = PreviewComponent->Sound)
+						{
+							return Sound->GetUniqueID() == InMetaSound.GetUniqueID();
+						}
 					}
 				}
 			}
@@ -1035,9 +1013,9 @@ PRAGMA_ENABLE_DEPRECATION_WARNINGS
 					if (ExternalNode->GetBreadcrumb().ClassName == FRerouteNodeTemplate::ClassName)
 					{
 						auto IsInput = [](UEdGraphPin* Pin) { check(Pin); return Pin->Direction == EGPD_Input; };
-						if (UEdGraphPin* RerouteInput = *ExternalNode->Pins.FindByPredicate(IsInput))
+						if (UEdGraphPin** RerouteInput = ExternalNode->Pins.FindByPredicate(IsInput); RerouteInput && *RerouteInput)
 						{
-							TArray<UEdGraphPin*>& LinkedTo = RerouteInput->LinkedTo;
+							TArray<UEdGraphPin*>& LinkedTo = (*RerouteInput)->LinkedTo;
 							if (!LinkedTo.IsEmpty())
 							{
 								UEdGraphPin* ReroutedOutput = LinkedTo.Last();
@@ -1062,9 +1040,9 @@ PRAGMA_ENABLE_DEPRECATION_WARNINGS
 					if (ExternalNode->GetBreadcrumb().ClassName == FRerouteNodeTemplate::ClassName)
 					{
 						auto IsInput = [](const UEdGraphPin* Pin) { check(Pin); return Pin->Direction == EGPD_Input; };
-						if (const UEdGraphPin* RerouteInput = *ExternalNode->Pins.FindByPredicate(IsInput))
+						if (const UEdGraphPin* const* RerouteInput = ExternalNode->Pins.FindByPredicate(IsInput))
 						{
-							const TArray<UEdGraphPin*>& LinkedTo = RerouteInput->LinkedTo;
+							const TArray<UEdGraphPin*>& LinkedTo = (*RerouteInput)->LinkedTo;
 							if (!LinkedTo.IsEmpty())
 							{
 								const UEdGraphPin* ReroutedOutput = LinkedTo.Last();
@@ -1089,9 +1067,9 @@ PRAGMA_ENABLE_DEPRECATION_WARNINGS
 					if (ExternalNode->GetBreadcrumb().ClassName == FRerouteNodeTemplate::ClassName)
 					{
 						auto IsInput = [](UEdGraphPin* Pin) { check(Pin); return Pin->Direction == EGPD_Input; };
-						if (const UEdGraphPin* RerouteInput = *ExternalNode->Pins.FindByPredicate(IsInput))
+						if (const UEdGraphPin* const* RerouteInput = ExternalNode->Pins.FindByPredicate(IsInput))
 						{
-							const TArray<UEdGraphPin*>& LinkedTo = RerouteInput->LinkedTo;
+							const TArray<UEdGraphPin*>& LinkedTo = (*RerouteInput)->LinkedTo;
 							if (!LinkedTo.IsEmpty())
 							{
 								const UEdGraphPin* ReroutedOutput = LinkedTo.Last();
@@ -1312,8 +1290,11 @@ PRAGMA_ENABLE_DEPRECATION_WARNINGS
 
 		FMetasoundFrontendClassInput FGraphBuilder::CreateUniqueClassInput(UObject& InMetaSound, const FCreateNodeVertexParams& InParams, const TArray<FMetasoundFrontendClassInputDefault>& InDefaultLiterals, const FName* InNameBase)
 		{
+			using namespace Engine;
+
 			FMetasoundFrontendClassInput ClassInput;
-			ClassInput.Name = GenerateUniqueNameByClassType(InMetaSound, EMetasoundFrontendClassType::Input, InNameBase ? InNameBase->ToString() : TEXT("Input"));
+			UMetaSoundBuilderBase& Builder = FDocumentBuilderRegistry::GetChecked().FindOrBeginBuilding(InMetaSound);
+			ClassInput.Name = Builder.GetConstBuilder().GenerateUniqueInputName(InNameBase ? InNameBase->ToString() : TEXT("Input"));
 			ClassInput.TypeName = InParams.DataType;
 			ClassInput.VertexID = FGuid::NewGuid();
 			ClassInput.NodeID = FGuid::NewGuid();
@@ -1350,8 +1331,11 @@ PRAGMA_ENABLE_DEPRECATION_WARNINGS
 
 		FMetasoundFrontendClassOutput FGraphBuilder::CreateUniqueClassOutput(UObject& InMetaSound, const FCreateNodeVertexParams& InParams, const FName* InNameBase)
 		{
+			using namespace Metasound::Engine;
+
 			FMetasoundFrontendClassOutput ClassOutput;
-			ClassOutput.Name = GenerateUniqueNameByClassType(InMetaSound, EMetasoundFrontendClassType::Output, InNameBase ? InNameBase->ToString() : TEXT("Output"));
+			UMetaSoundBuilderBase& Builder = FDocumentBuilderRegistry::GetChecked().FindOrBeginBuilding(InMetaSound);
+			ClassOutput.Name = Builder.GetConstBuilder().GenerateUniqueOutputName(InNameBase ? InNameBase->ToString() : TEXT("Output"));
 			ClassOutput.TypeName = InParams.DataType;
 			ClassOutput.VertexID = FGuid::NewGuid();
 			ClassOutput.NodeID = FGuid::NewGuid();
@@ -1362,19 +1346,6 @@ PRAGMA_ENABLE_DEPRECATION_WARNINGS
 			return ClassOutput;
 		}
 
-		FName FGraphBuilder::GenerateUniqueVariableName(const FMetaSoundFrontendDocumentBuilder& InBuilder, const FString& InBaseName)
-		{
-			using namespace Frontend;
-
-			TSet<FName> ExistingVariableNames;
-
-			// Get all the names from the existing variables on
-			// the build graph and place into the ExistingVariableNames array.
-			Algo::Transform(InBuilder.FindConstBuildGraphChecked().Variables, ExistingVariableNames, [](const FMetasoundFrontendVariable& Var) { return Var.Name; });
-
-			return GraphBuilderPrivate::GenerateUniqueName(ExistingVariableNames, InBaseName);
-		}
-
 		FMetasoundAssetBase& FGraphBuilder::GetOutermostMetaSoundChecked(UObject& InSubObject)
 		{
 			FMetasoundAssetBase* MetaSoundAsset = IMetasoundUObjectRegistry::Get().GetObjectAsAssetBase(InSubObject.GetOutermostObject());
@@ -1954,21 +1925,24 @@ PRAGMA_ENABLE_DEPRECATION_WARNINGS
 				TSet<FMetasoundAssetBase*> EditedReferencingMetaSounds;
 				if (GEditor)
 				{
-					TArray<UObject*> EditedAssets = GEditor->GetEditorSubsystem<UAssetEditorSubsystem>()->GetAllEditedAssets();
-					for (UObject* Asset : EditedAssets)
+					if (UAssetEditorSubsystem* AssetSubsystem = GEditor->GetEditorSubsystem<UAssetEditorSubsystem>())
 					{
-						if (Asset != &MetaSound)
+						TArray<UObject*> EditedAssets = AssetSubsystem->GetAllEditedAssets();
+						for (UObject* Asset : EditedAssets)
 						{
-							if (FMetasoundAssetBase* EditedMetaSound = IMetasoundUObjectRegistry::Get().GetObjectAsAssetBase(Asset))
+							if (Asset != &MetaSound)
 							{
-								if (EditedMetaSound->IsReferencedAsset(*MetaSoundAsset))
+								if (FMetasoundAssetBase* EditedMetaSound = IMetasoundUObjectRegistry::Get().GetObjectAsAssetBase(Asset))
 								{
-									// Ensure the graph has been initialized. Its possible the "edited asset" is a MetaSound
-									// but its transient EdGraph has yet to be initialized due to the active editor not necessarily
-									// being the MetaSound editor (ex. property matrix editor).
-									if (const UEdGraph* EdGraph = EditedMetaSound->GetGraph())
+									if (EditedMetaSound->IsReferencedAsset(*MetaSoundAsset))
 									{
-										EditedReferencingMetaSounds.Add(EditedMetaSound);
+										// Ensure the graph has been initialized. Its possible the "edited asset" is a MetaSound
+										// but its transient EdGraph has yet to be initialized due to the active editor not necessarily
+										// being the MetaSound editor (ex. property matrix editor).
+										if (const UEdGraph* EdGraph = EditedMetaSound->GetGraph())
+										{
+											EditedReferencingMetaSounds.Add(EditedMetaSound);
+										}
 									}
 								}
 							}
@@ -2160,7 +2134,7 @@ PRAGMA_ENABLE_DEPRECATION_WARNINGS
 					{
 						case EMetasoundFrontendClassType::External:
 						{
-							NewGraphNode = CastChecked<UMetasoundEditorGraphNode>(AddExternalNode(MetaSoundObject, NodeID, ClassMetadata, false));
+							NewGraphNode = Cast<UMetasoundEditorGraphNode>(AddExternalNode(MetaSoundObject, NodeID, ClassMetadata, false));
 						}
 						break;
 
@@ -2168,18 +2142,18 @@ PRAGMA_ENABLE_DEPRECATION_WARNINGS
 						{
 							if (Pair.Value->Metadata.GetClassName() == FInputNodeTemplate::ClassName)
 							{
-								NewGraphNode = CastChecked<UMetasoundEditorGraphNode>(AddInputNode(MetaSoundObject, NodeID, false));
+								NewGraphNode = Cast<UMetasoundEditorGraphNode>(AddInputNode(MetaSoundObject, NodeID, false));
 							}
 							else
 							{
-								NewGraphNode = CastChecked<UMetasoundEditorGraphNode>(AddTemplateNode(MetaSoundObject, NodeID, ClassMetadata, false));
+								NewGraphNode = Cast<UMetasoundEditorGraphNode>(AddTemplateNode(MetaSoundObject, NodeID, ClassMetadata, false));
 							}
 						}
 						break;
 
 						case EMetasoundFrontendClassType::Output:
 						{
-							NewGraphNode = CastChecked<UMetasoundEditorGraphNode>(AddOutputNode(MetaSoundObject, NodeID, false));
+							NewGraphNode = Cast<UMetasoundEditorGraphNode>(AddOutputNode(MetaSoundObject, NodeID, false));
 						}
 						break;
 
@@ -2188,7 +2162,7 @@ PRAGMA_ENABLE_DEPRECATION_WARNINGS
 						case EMetasoundFrontendClassType::VariableDeferredAccessor:
 						case EMetasoundFrontendClassType::Variable:
 						{
-							NewGraphNode = CastChecked<UMetasoundEditorGraphNode>(AddVariableNode(MetaSoundObject, NodeID, false));
+							NewGraphNode = Cast<UMetasoundEditorGraphNode>(AddVariableNode(MetaSoundObject, NodeID, false));
 						}
 						break;
 
diff --git a/Engine/Plugins/Runtime/Metasound/Source/MetasoundEditor/Private/MetasoundEditorGraphBuilder.h b/Engine/Plugins/Runtime/Metasound/Source/MetasoundEditor/Private/MetasoundEditorGraphBuilder.h
--- a/Engine/Plugins/Runtime/Metasound/Source/MetasoundEditor/Private/MetasoundEditorGraphBuilder.h
+++ b/Engine/Plugins/Runtime/Metasound/Source/MetasoundEditor/Private/MetasoundEditorGraphBuilder.h
@@ -143,9 +143,6 @@ namespace Metasound
 			// Adds an output node to the editor graph that corresponds to the provided output node ID.
 			static UMetasoundEditorGraphOutputNode* AddOutputNode(UObject& InMetaSound, const FGuid& InNodeID, bool bInSelectNewNode = true);
 
-			// Create a unique name for the variable.
-			static FName GenerateUniqueVariableName(const FMetaSoundFrontendDocumentBuilder& InBuilder, const FString& InBaseName);
-
 			// Convenience method for walking to the outermost object and transforming to a base MetaSound
 			static FMetasoundAssetBase& GetOutermostMetaSoundChecked(UObject& InSubObject);
 			static const FMetasoundAssetBase& GetOutermostConstMetaSoundChecked(const UObject& InSubObject);
@@ -173,9 +170,6 @@ namespace Metasound
 			// disconnect the EdGraph pins.
 			static void DisconnectPinVertex(UEdGraphPin& InPin);
 
-			// Generates a unique output name for the given MetaSound object
-			static FName GenerateUniqueNameByClassType(UObject& InMetaSound, EMetasoundFrontendClassType InClassType, const FString& InBaseName);
-
 			static UMetaSoundBuilderBase& GetBuilderFromPinChecked(const UEdGraphPin& InPin);
 
 			static TArray<FString> GetDataTypeNameCategories(const FName& InDataTypeName);
diff --git a/Engine/Plugins/Runtime/Metasound/Source/MetasoundEditor/Private/MetasoundEditorGraphNode.cpp b/Engine/Plugins/Runtime/Metasound/Source/MetasoundEditor/Private/MetasoundEditorGraphNode.cpp
--- a/Engine/Plugins/Runtime/Metasound/Source/MetasoundEditor/Private/MetasoundEditorGraphNode.cpp
+++ b/Engine/Plugins/Runtime/Metasound/Source/MetasoundEditor/Private/MetasoundEditorGraphNode.cpp
@@ -774,7 +774,7 @@ void UMetasoundEditorGraphNode::GetNodeContextMenuActions(UToolMenu* Menu, UGrap
 
 	if (Context->Node)
 	{
-		if (!GetBuilderChecked().IsPreset())
+		if (GetBuilderChecked().GetConstBuilder().GetIsGraphEditable())
 		{
 			{
 				FToolMenuSection& Section = Menu->FindOrAddSection("MetasoundGraphNodeActionsOrganization", LOCTEXT("NodeActionsOrganizationMenuHeader", "Organization"));
diff --git a/Engine/Plugins/Runtime/Metasound/Source/MetasoundEditor/Private/MetasoundEditorGraphSchema.cpp b/Engine/Plugins/Runtime/Metasound/Source/MetasoundEditor/Private/MetasoundEditorGraphSchema.cpp
--- a/Engine/Plugins/Runtime/Metasound/Source/MetasoundEditor/Private/MetasoundEditorGraphSchema.cpp
+++ b/Engine/Plugins/Runtime/Metasound/Source/MetasoundEditor/Private/MetasoundEditorGraphSchema.cpp
@@ -290,7 +290,7 @@ PRAGMA_ENABLE_DEPRECATION_WARNINGS
 
 			UMetasoundEditorGraph* MetaSoundGraph = CastChecked<UMetasoundEditorGraph>(ConnectedNode->GetGraph());
 			FMetaSoundFrontendDocumentBuilder& DocBuilder = MetaSoundGraph->GetBuilderChecked().GetBuilder();
-			const FName NodeName = FGraphBuilder::GenerateUniqueVariableName(DocBuilder, BaseName.ToString());
+			const FName NodeName = DocBuilder.GenerateUniqueVariableName(BaseName.ToString());
 
 			const FScopedTransaction Transaction(FText::Format(
 				LOCTEXT("PromoteNodeVertexToGraphVariableFormat", "Promote MetaSound Node {0} to {1}"),
@@ -1649,7 +1649,7 @@ void UMetasoundEditorGraphSchema::GetContextMenuActions(class UToolMenu* Menu, c
 		const UMetasoundEditorGraph* EdGraph = CastChecked<UMetasoundEditorGraph>(Context->Graph);
 		UMetaSoundBuilderBase& Builder = Metasound::Engine::FDocumentBuilderRegistry::GetChecked().FindOrBeginBuilding(EdGraph->GetMetasoundChecked());
 
-		if (!Builder.IsPreset())
+		if (Builder.GetConstBuilder().GetIsGraphEditable())
 		{
 			FToolMenuSection& Section = Menu->FindOrAddSection("Pin Actions");
 
@@ -2807,7 +2807,7 @@ void UMetasoundEditorGraphSchema::QueryNodeClasses(TFunctionRef<void(TUniquePtr<
 
 	// Gather all loaded classes
 	{
-		static bool bIncludeUnloadedAssets = true;
+		constexpr bool bIncludeUnloadedAssets = true;
 		const ISearchEngine::EResultVersion Options = SchemaPrivate::EnableAllVersionsMetaSoundNodeClassCreationCVar == 0
 			? ISearchEngine::EResultVersion::Highest
 			: ISearchEngine::EResultVersion::All;
diff --git a/Engine/Plugins/Runtime/Metasound/Source/MetasoundEditor/Public/MetasoundEditorGraph.h b/Engine/Plugins/Runtime/Metasound/Source/MetasoundEditor/Public/MetasoundEditorGraph.h
--- a/Engine/Plugins/Runtime/Metasound/Source/MetasoundEditor/Public/MetasoundEditorGraph.h
+++ b/Engine/Plugins/Runtime/Metasound/Source/MetasoundEditor/Public/MetasoundEditorGraph.h
@@ -107,15 +107,15 @@ public:
 		return Value;
 	}
 
-	UE_API virtual void ForceRefresh() override;
+	UE_API virtual void ForceRefresh();
 	UE_API virtual FName GetDataType() const;
-	UE_API virtual EMetasoundFrontendLiteralType GetLiteralType() const override;
+	UE_API virtual EMetasoundFrontendLiteralType GetLiteralType() const;
 	UE_API virtual void InitDefault(const FGuid& InPageID = Metasound::Frontend::DefaultPageID);
 	virtual void Initialize() { };
 	UE_API virtual void IterateDefaults(TFunctionRef<void(const FGuid&, FMetasoundFrontendLiteral)> Iter) const;
 	UE_API virtual bool RemoveDefault(const FGuid& InPageID);
 	UE_API virtual void ResetDefaults();
-	UE_API virtual void SetFromLiteral(const FMetasoundFrontendLiteral& InLiteral, const FGuid& InPageID = Metasound::Frontend::DefaultPageID) override;
+	UE_API virtual void SetFromLiteral(const FMetasoundFrontendLiteral& InLiteral, const FGuid& InPageID = Metasound::Frontend::DefaultPageID);
 	UE_API virtual bool TryFindDefault(FMetasoundFrontendLiteral& OutLiteral, const FGuid* InPageID = nullptr) const;
 
 	virtual void UpdatePreviewInstance(const Metasound::FVertexName& InParameterName, TScriptInterface<IAudioParameterControllerInterface>& InParameterInterface) const
diff --git a/Engine/Plugins/Runtime/Metasound/Source/MetasoundEngine/Private/DocumentConfigurations/MetasoundDocumentConfigurator.cpp b/Engine/Plugins/Runtime/Metasound/Source/MetasoundEngine/Private/DocumentConfigurations/MetasoundDocumentConfigurator.cpp
new file mode 100644
--- /dev/null
+++ b/Engine/Plugins/Runtime/Metasound/Source/MetasoundEngine/Private/DocumentConfigurations/MetasoundDocumentConfigurator.cpp
@@ -0,0 +1,395 @@
+// Copyright Epic Games, Inc. All Rights Reserved.
+
+#include "DocumentConfigurations/MetasoundDocumentConfigurator.h"
+
+#include "MetasoundAssetSubsystem.h"
+#include "MetasoundDocumentBuilderRegistry.h"
+#include "MetasoundFrontendDocumentBuilder.h"
+#include "Algo/Transform.h"
+
+
+#if WITH_EDITORONLY_DATA
+namespace Metasound::Engine
+{
+	FDocumentConfigurator::FDocumentConfigurator(const FDocumentConfigurator::FArgs& Args)
+#if WITH_EDITORONLY_DATA
+		: Offset(Args.Offset)
+#endif // WITH_EDITORONLY_DATA
+	{
+		checkf(Args.Builder, TEXT("Builder must be supplied via document configurator arguments"));
+		UObject& MetaSoundObject = Args.Builder->CastDocumentObjectChecked<UObject>();
+		Builder.Reset(&FDocumentBuilderRegistry::GetChecked().FindOrBeginBuilding(MetaSoundObject));
+
+		if (Args.bResetDoc)
+		{
+			ResetDocument();
+		}
+	}
+
+	FDocumentConfigurator::~FDocumentConfigurator()
+	{
+		// Inject result is ignored because 1. the configurator has already reported success if its destroying,
+		// and 2. the inject can fail if there are no inputs which is technically a valid situation.
+		EMetaSoundBuilderResult InjectResult = EMetaSoundBuilderResult::Failed;
+		constexpr bool bForceNodeCreation = false;
+		Builder->InjectInputTemplateNodes(bForceNodeCreation, InjectResult);
+	}
+
+	FDocumentConfigurator& FDocumentConfigurator::Add(FDocumentConfigurator::FInterface Interface)
+	{
+		const bool bSucceeded = Builder->GetBuilder().ModifyInterfaces(Frontend::FModifyInterfaceOptions(
+			{ }, // None removed
+			{ Interface.Version }
+		));
+
+		if (Result == EMetaSoundBuilderResult::Succeeded)
+		{
+			Result = bSucceeded ? EMetaSoundBuilderResult::Succeeded : EMetaSoundBuilderResult::Failed;
+		}
+
+		return *this;
+	}
+
+	FDocumentConfigurator& FDocumentConfigurator::Add(FDocumentConfigurator::FInput Input)
+	{
+		using namespace Frontend;
+
+		EMetaSoundBuilderResult StepResult = EMetaSoundBuilderResult::Failed;
+
+		FMetaSoundBuilderNodeOutputHandle Node = Builder->AddGraphInputNode(Input.Name, Input.DataType, MoveTemp(Input.DefaultValue), StepResult);
+		if (StepResult == EMetaSoundBuilderResult::Succeeded)
+		{
+			Builder->SetNodeLocation(Node.NodeID, Input.Location + Offset, StepResult);
+			if (Result == EMetaSoundBuilderResult::Succeeded)
+			{
+				Result = StepResult;
+			}
+		}
+
+		return *this;
+	}
+
+	FDocumentConfigurator& FDocumentConfigurator::Add(FDocumentConfigurator::FOutput Output)
+	{
+		using namespace Frontend;
+
+		EMetaSoundBuilderResult StepResult = EMetaSoundBuilderResult::Failed;
+
+		Builder->AddGraphOutputNode(Output.Name, Output.DataType, MoveTemp(Output.DefaultValue), StepResult);
+		if (StepResult == EMetaSoundBuilderResult::Failed)
+		{
+			Result = StepResult;
+		}
+
+		return *this;
+	}
+
+	FDocumentConfigurator& FDocumentConfigurator::Add(FDocumentConfigurator::FNode Node)
+	{
+		EMetaSoundBuilderResult StepResult = EMetaSoundBuilderResult::Failed;
+
+		if (NodeNameToHandle.Contains(Node.Name))
+		{
+			UE_LOG(LogMetaSound, Error, TEXT("Node with name '%s' already exists. Failed to add node instance to MetaSound graph"), *Node.Name.ToString());
+			Result = EMetaSoundBuilderResult::Failed;
+		}
+		else
+		{
+			FMetaSoundNodeHandle NodeHandle = Builder->AddNodeByClassName(Node.ClassName, StepResult);
+
+			if (StepResult == EMetaSoundBuilderResult::Succeeded)
+			{
+				NodeNameToHandle.Add(Node.Name, NodeHandle);
+				Builder->SetNodeLocation(NodeHandle, Node.Location + Offset, StepResult);
+			}
+		}
+
+		if (Result == EMetaSoundBuilderResult::Succeeded)
+		{
+			Result = StepResult;
+		}
+
+		return *this;
+	}
+
+	FDocumentConfigurator& FDocumentConfigurator::Add(FDocumentConfigurator::FEdge Edge)
+	{
+		EMetaSoundBuilderResult StepResult = EMetaSoundBuilderResult::Failed;
+
+		if (const FMetaSoundNodeHandle* OutputNodeHandle = NodeNameToHandle.Find(Edge.OutputNode))
+		{
+			if (const FMetaSoundNodeHandle* InputNodeHandle = NodeNameToHandle.Find(Edge.InputNode))
+			{
+				Builder->ConnectNodes(*OutputNodeHandle, Edge.Output, *InputNodeHandle, Edge.Input, StepResult);
+			}
+			else
+			{
+				UE_LOG(LogMetaSound, Error, TEXT("Failed to find input node with name '%s' to connect"), *Edge.InputNode.ToString());
+				StepResult = EMetaSoundBuilderResult::Failed;
+			}
+		}
+		else
+		{
+			UE_LOG(LogMetaSound, Error, TEXT("Failed to find output node with name '%s' to connect"), *Edge.OutputNode.ToString());
+			StepResult = EMetaSoundBuilderResult::Failed;
+		}
+
+		if (Result == EMetaSoundBuilderResult::Succeeded)
+		{
+			Result = StepResult;
+		}
+
+		return *this;
+	}
+
+	FDocumentConfigurator& FDocumentConfigurator::Add(FDocumentConfigurator::FEdgeFromInput Edge)
+	{
+		EMetaSoundBuilderResult StepResult = EMetaSoundBuilderResult::Failed;
+
+		if (const FMetaSoundNodeHandle* OutputNodeHandle = NodeNameToHandle.Find(Edge.Node))
+		{
+			Builder->ConnectGraphInputToNode(Edge.GraphInput, *OutputNodeHandle, Edge.NodeInput, StepResult);
+			if (Result == EMetaSoundBuilderResult::Succeeded)
+			{
+				Result = StepResult;
+			}
+		}
+		else
+		{
+			UE_LOG(LogMetaSound, Error, TEXT("Failed to find node with name '%s': Graph input to node input connection failed"), *Edge.Node.ToString());
+			StepResult = EMetaSoundBuilderResult::Failed;
+		}
+
+		if (Result == EMetaSoundBuilderResult::Succeeded)
+		{
+			Result = StepResult;
+		}
+
+		return *this;
+	}
+
+	FDocumentConfigurator& FDocumentConfigurator::Add(FDocumentConfigurator::FEdgeToOutput Edge)
+	{
+		EMetaSoundBuilderResult StepResult = EMetaSoundBuilderResult::Failed;
+
+		if (const FMetaSoundNodeHandle* OutputNodeHandle = NodeNameToHandle.Find(Edge.Node))
+		{
+			Builder->ConnectNodeToGraphOutput(*OutputNodeHandle, Edge.NodeOutput, Edge.GraphOutput, StepResult);
+			if (Result == EMetaSoundBuilderResult::Succeeded)
+			{
+				Result = StepResult;
+			}
+		}
+		else
+		{
+			UE_LOG(LogMetaSound, Error, TEXT("Failed to find node with name '%s': Graph output to node iutput connection failed"), *Edge.Node.ToString());
+			StepResult = EMetaSoundBuilderResult::Failed;
+		}
+
+		if (Result == EMetaSoundBuilderResult::Succeeded)
+		{
+			Result = StepResult;
+		}
+
+		return *this;
+	}
+
+	FDocumentConfigurator& FDocumentConfigurator::AddInternal(TArray<FDocumentConfigurator::FInterface> Interfaces)
+	{
+		TArray<FMetasoundFrontendVersion> Versions;
+		Algo::Transform(Interfaces, Versions, [](const FInterface& Interface) { return Interface.Version; });
+		const bool bSucceeded = Builder->GetBuilder().ModifyInterfaces(Frontend::FModifyInterfaceOptions(
+			{ }, // None removed
+			{ Versions }
+		));
+
+		if (Result == EMetaSoundBuilderResult::Succeeded)
+		{
+			Result = bSucceeded ? EMetaSoundBuilderResult::Succeeded : EMetaSoundBuilderResult::Failed;
+		}
+
+		return *this;
+	}
+
+	UMetaSoundBuilderBase& FDocumentConfigurator::GetBuilder()
+	{
+		return *Builder.Get();
+	}
+	
+#if WITH_EDITORONLY_DATA
+	void FDocumentConfigurator::Reflect(const FMetasoundFrontendGraphLayer& InLayer, FString& OutCode)
+	{
+		TMap<FGuid, const FMetasoundFrontendClass*> Dependencies;
+		Algo::Transform(InLayer.Dependencies, Dependencies, [](const FMetasoundFrontendClass& Dependency)
+		{
+			return TPair<FGuid, const FMetasoundFrontendClass*>(Dependency.ID, &Dependency);
+		});
+
+		TSet<FString> NodeInstanceNames;
+
+		struct FConfigNodeInfo
+		{
+			FString InstanceName;
+			const FMetasoundFrontendNode* Node = nullptr;
+		};
+		TMap<FGuid, FConfigNodeInfo> NodeIDToInfo;
+
+		FMetaSoundAssetManager& AssetManager = FMetaSoundAssetManager::GetChecked();
+
+		OutCode.Reset();
+
+		OutCode += TEXT("using FDC = FDocumentConfigurator;\n");
+		OutCode += TEXT("FDC Configurator(OutBuilder);\n");
+
+		auto CacheNodeInfo = [&NodeInstanceNames, &NodeIDToInfo](const FMetasoundFrontendClassName& ClassName, const FMetasoundFrontendNode& Node)
+		{
+			FString BaseName = ClassName.Name.ToString();
+			int32 Index = 0;
+			FString FullName = BaseName;
+			while (NodeInstanceNames.Contains(FullName))
+			{
+				FullName = FString::Printf(TEXT("%s_%i"), *BaseName, ++Index);
+			}
+			NodeInstanceNames.Add(FullName);
+			NodeIDToInfo.Add(Node.GetID(), FConfigNodeInfo { .InstanceName = FullName, .Node = &Node });
+			return FullName;
+		};
+
+		OutCode += TEXT("\n// Nodes\n");
+		OutCode += TEXT("Configurator.Add<FDC::FNode>(\n{\n");
+
+		for (const FMetasoundFrontendNode& Node : InLayer.Graph.Nodes)
+		{
+			if (const FMetasoundFrontendClass* Dependency = Dependencies.FindRef(Node.ClassID); Dependency && Dependency->Metadata.GetType() == EMetasoundFrontendClassType::External)
+			{
+				FVector2D Location = FVector2D::Zero();
+				if (!Node.Style.Display.Locations.IsEmpty())
+				{
+					Location = Node.Style.Display.Locations.Array().Last().Value;
+				}
+
+				const FMetasoundFrontendClassName& ClassName = Dependency->Metadata.GetClassName();
+				FString ClassInfo;
+				if (AssetManager.IsAssetClass(Dependency->Metadata))
+				{
+					const FTopLevelAssetPath AssetPath = AssetManager.FindAssetPath(FMetaSoundAssetKey(Dependency->Metadata));
+					check(AssetPath.IsValid());
+					ClassInfo = FString::Printf(TEXT(".AssetPath = \"%s\""), *AssetPath.ToString());
+				}
+				else
+				{
+					ClassInfo = FString::Printf(TEXT(".ClassName = \"%s\""), *ClassName.ToString());
+				}
+
+				OutCode += FString::Printf(
+					TEXT("\t{ .Name = \"%s\", %s, .Location = FVector2D(%.0ff, %.0ff) },\n"),
+					*CacheNodeInfo(ClassName, Node),
+					*ClassInfo,
+					Location.X,
+					Location.Y
+				);
+			}
+		}
+		OutCode += TEXT("});\n");
+
+		if (!InLayer.Graph.Edges.IsEmpty())
+		{
+			OutCode += TEXT("\n// Edges\n");
+			OutCode += TEXT("Configurator.Add<FDC::FEdge>(\n{\n");
+
+			for (const FMetasoundFrontendEdge& Edge : InLayer.Graph.Edges)
+			{
+				if (const FConfigNodeInfo* FromNodeInfo = NodeIDToInfo.Find(Edge.FromNodeID))
+				{
+					const FMetasoundFrontendNode* FromNode = FromNodeInfo->Node;
+					check(FromNode);
+
+					auto IsFromVertex = [&Edge](const FMetasoundFrontendVertex& Vertex) { return Vertex.VertexID == Edge.FromVertexID; };
+					const FMetasoundFrontendVertex* FromVertex = FromNode->Interface.Outputs.FindByPredicate(IsFromVertex);
+					if (!FromVertex)
+					{
+						continue;
+					}
+
+					if (const FConfigNodeInfo* ToNodeInfo = NodeIDToInfo.Find(Edge.ToNodeID))
+					{
+						const FMetasoundFrontendNode* ToNode = ToNodeInfo->Node;
+						check(ToNode);
+
+						auto IsToVertex = [&Edge](const FMetasoundFrontendVertex& Vertex) { return Vertex.VertexID == Edge.ToVertexID; };
+						const FMetasoundFrontendVertex* ToVertex = ToNode->Interface.Inputs.FindByPredicate(IsToVertex);
+						if (!ToVertex)
+						{
+							continue;
+						}
+
+						OutCode += FString::Printf(
+							TEXT("\t{ .OutputNode = \"%s\", .Output = \"%s\", .InputNode = \"%s\", .Input = \"%s\" },\n"),
+							*FromNodeInfo->InstanceName,
+							*FromVertex->Name.ToString(),
+							*ToNodeInfo->InstanceName,
+							*ToVertex->Name.ToString()
+						);
+					}
+				}
+			}
+			OutCode += TEXT("});\n");
+		}
+	}
+#endif // WITH_EDITORONLY_DATA
+
+	FDocumentConfigurator& FDocumentConfigurator::ResetDocument()
+	{
+		using namespace Frontend;
+
+		FMetaSoundFrontendDocumentBuilder& FrontendBuilder = Builder->GetBuilder();
+		TSharedRef<FDocumentModifyDelegates> ModifyDelegates = MakeShared<FDocumentModifyDelegates>(FrontendBuilder.GetDocumentDelegates());
+		FrontendBuilder.ClearDocument(ModifyDelegates);
+
+		FrontendBuilder.SetIsGraphEditable(false);
+
+		return *this;
+	}
+
+	FDocumentConfigurator& FDocumentConfigurator::Set(FDocumentConfigurator::FInputLocation Location)
+	{
+		EMetaSoundBuilderResult StepResult = EMetaSoundBuilderResult::Failed;
+		if (const FMetaSoundNodeHandle InputNode = Builder->FindGraphInputNode(Location.Name, StepResult); StepResult == EMetaSoundBuilderResult::Succeeded)
+		{
+			Builder->SetNodeLocation(InputNode, Location.Location + Offset, StepResult);
+		}
+
+		if (Result == EMetaSoundBuilderResult::Succeeded)
+		{
+			Result = StepResult;
+		}
+
+		return *this;
+	}
+
+	FDocumentConfigurator& FDocumentConfigurator::Set(FDocumentConfigurator::FNodeInputDefault Input)
+	{
+		EMetaSoundBuilderResult StepResult = EMetaSoundBuilderResult::Failed;
+		if (const FMetaSoundNodeHandle* NodeHandle = NodeNameToHandle.Find(Input.Node))
+		{
+			FMetaSoundBuilderNodeInputHandle InputHandle = Builder->FindNodeInputByName(*NodeHandle, Input.Input, StepResult);
+			if (StepResult == EMetaSoundBuilderResult::Succeeded)
+			{
+				Builder->SetNodeInputDefault(InputHandle, Input.Value, StepResult);
+			}
+		}
+
+		if (Result == EMetaSoundBuilderResult::Succeeded)
+		{
+			Result = StepResult;
+		}
+
+		return *this;
+	}
+
+	bool FDocumentConfigurator::Succeeded() const
+	{
+		return Result == EMetaSoundBuilderResult::Succeeded;
+	}
+} // namespace Metasound::Engine
+#endif // WITH_EDITORONLY_DATA
diff --git a/Engine/Plugins/Runtime/Metasound/Source/MetasoundEngine/Private/DocumentConfigurations/MetasoundRandomizerDocumentConfiguration.cpp b/Engine/Plugins/Runtime/Metasound/Source/MetasoundEngine/Private/DocumentConfigurations/MetasoundRandomizerDocumentConfiguration.cpp
new file mode 100644
--- /dev/null
+++ b/Engine/Plugins/Runtime/Metasound/Source/MetasoundEngine/Private/DocumentConfigurations/MetasoundRandomizerDocumentConfiguration.cpp
@@ -0,0 +1,285 @@
+// Copyright Epic Games, Inc. All Rights Reserved.
+
+#include "DocumentConfigurations/MetasoundRandomizerDocumentConfiguration.h"
+
+#include "Algo/MaxElement.h"
+#include "Algo/Transform.h"
+#include "Containers/ContainersFwd.h"
+#include "DocumentConfigurations/MetasoundDocumentConfigurator.h"
+#include "Interfaces/MetasoundOutputFormatInterfaces.h"
+#include "MetasoundBuilderSubsystem.h"
+#include "MetasoundDataTypeGetTraits.h"
+#include "MetasoundDocumentInterface.h"
+#include "MetasoundEngineNodesNames.h"
+#include "MetasoundFrontendDocument.h"
+#include "MetasoundFrontendDocumentBuilder.h"
+#include "MetasoundFrontendDocumentBuilderRegistry.h"
+#include "MetasoundFrontendDocumentIdGenerator.h"
+#include "MetasoundWave.h"
+
+
+#define LOCTEXT_NAMESPACE "MetaSound"
+
+namespace Metasound::Engine
+{
+	namespace RandomizerPrivate
+	{
+		const FMetasoundFrontendNode* CreateWaveInput(FMetaSoundFrontendDocumentBuilder& OutBuilder, const TArray<TObjectPtr<USoundWave>>& InSounds)
+		{
+			using namespace Frontend;
+
+			const FMetasoundFrontendDocument& Doc = OutBuilder.GetConstDocumentChecked();
+
+			FMetasoundFrontendClassInput WaveInput;
+			WaveInput.Name = "Sounds";
+			WaveInput.TypeName = GetMetasoundDataTypeName<TArray<FWaveAsset>>();
+			WaveInput.NodeID = FDocumentIDGenerator::Get().CreateNodeID(Doc);
+			WaveInput.VertexID = FDocumentIDGenerator::Get().CreateVertexID(Doc);
+			WaveInput.AccessType = EMetasoundFrontendVertexAccessType::Value;
+
+			TArray<UObject*> SoundObjects;
+			Algo::TransformIf(InSounds, SoundObjects,
+				[](const TObjectPtr<USoundWave> Sound) { return Sound != nullptr; },
+				[](const TObjectPtr<USoundWave> Sound) { return Sound.Get(); });
+
+			FMetasoundFrontendLiteral DefaultValue;
+			DefaultValue.Set(MoveTemp(SoundObjects));
+			WaveInput.InitDefault(MoveTemp(DefaultValue));
+
+			return OutBuilder.AddGraphInput(MoveTemp(WaveInput));
+		}
+
+		int32 GetMaxChannels(const TArray<TObjectPtr<USoundWave>>& InSounds)
+		{
+			auto GetNumChannels = [](const TObjectPtr<USoundWave>& Sound)
+			{
+				if (const USoundWave* SoundPtr = Sound.Get())
+				{
+					return SoundPtr->NumChannels;
+				}
+
+				return 0;
+			};
+
+			if (const TObjectPtr<USoundWave>* MaxElement = Algo::MaxElementBy(InSounds, GetNumChannels))
+			{
+				return GetNumChannels(*MaxElement);
+			}
+
+			return 0;
+		}
+
+		FMetasoundFrontendClassName GetPlayerClassName(int32 NumChannels)
+		{
+			FMetasoundFrontendClassName ClassName(EngineNodes::Namespace, "Wave Player");
+			if (NumChannels == 2)
+			{
+				ClassName.Variant = EngineNodes::StereoVariant;
+			}
+			else if (NumChannels == 4)
+			{
+				ClassName.Variant = EngineNodes::QuadVariant;
+			}
+			else if (NumChannels == 6)
+			{
+				ClassName.Variant = EngineNodes::FiveDotOneVariant;
+			}
+			else if (NumChannels == 8)
+			{
+				ClassName.Variant = EngineNodes::SevenDotOneVariant;
+			}
+			else
+			{
+				ClassName.Variant = EngineNodes::MonoVariant;
+
+				if (NumChannels != 1)
+				{
+					UE_LOG(LogMetaSound, Warning, TEXT("SoundWave format with %i NumChannels not recognized: using Mono MetaSound WavePlayer"), NumChannels);
+				}
+			}
+			return ClassName;
+		}
+
+		FMetasoundFrontendVersion GetOutputVersion(int32 NumChannels)
+		{
+			if (NumChannels == 2)
+			{
+				return OutputFormatStereoInterface::GetVersion();
+			}
+
+			if (NumChannels == 4)
+			{
+				return OutputFormatQuadInterface::GetVersion();
+			}
+
+			if (NumChannels == 6)
+			{
+				return OutputFormatFiveDotOneInterface::GetVersion();
+			}
+
+			if (NumChannels == 8)
+			{
+				return OutputFormatSevenDotOneInterface::GetVersion();
+			}
+
+			return OutputFormatMonoInterface::GetVersion();
+		}
+
+		TArray<FDocumentConfigurator::FEdgeToOutput> GetAudioOutputEdges(int32 NumChannels)
+		{
+			using FDC = FDocumentConfigurator;
+
+			if (NumChannels == 2)
+			{
+				return
+				{
+					FDC::FEdgeToOutput { .Node = "WavePlayer", .NodeOutput = "Out Left", .GraphOutput = OutputFormatStereoInterface::Outputs::LeftOut },
+					FDC::FEdgeToOutput { .Node = "WavePlayer", .NodeOutput = "Out Right", .GraphOutput = OutputFormatStereoInterface::Outputs::RightOut }
+				};
+			}
+
+			if (NumChannels == 4)
+			{
+				return
+				{
+					FDC::FEdgeToOutput { .Node = "WavePlayer", .NodeOutput = "Out Front Left", .GraphOutput = OutputFormatQuadInterface::Outputs::FrontLeftOut },
+					FDC::FEdgeToOutput { .Node = "WavePlayer", .NodeOutput = "Out Front Right", .GraphOutput = OutputFormatQuadInterface::Outputs::FrontRightOut },
+					FDC::FEdgeToOutput { .Node = "WavePlayer", .NodeOutput = "Out Side Left", .GraphOutput = OutputFormatQuadInterface::Outputs::SideLeftOut },
+					FDC::FEdgeToOutput { .Node = "WavePlayer", .NodeOutput = "Out Side Right", .GraphOutput = OutputFormatQuadInterface::Outputs::SideRightOut }
+				};
+			}
+
+			if (NumChannels == 6)
+			{
+				return
+				{
+					FDC::FEdgeToOutput { .Node = "WavePlayer", .NodeOutput = "Out Front Left", .GraphOutput = OutputFormatFiveDotOneInterface::Outputs::FrontLeftOut },
+					FDC::FEdgeToOutput { .Node = "WavePlayer", .NodeOutput = "Out Front Right", .GraphOutput = OutputFormatFiveDotOneInterface::Outputs::FrontRightOut },
+					FDC::FEdgeToOutput { .Node = "WavePlayer", .NodeOutput = "Out Front Center", .GraphOutput = OutputFormatFiveDotOneInterface::Outputs::FrontCenterOut },
+					FDC::FEdgeToOutput { .Node = "WavePlayer", .NodeOutput = "Out Low Frequency", .GraphOutput = OutputFormatFiveDotOneInterface::Outputs::LowFrequencyOut },
+					FDC::FEdgeToOutput { .Node = "WavePlayer", .NodeOutput = "Out Side Left", .GraphOutput = OutputFormatFiveDotOneInterface::Outputs::SideLeftOut },
+					FDC::FEdgeToOutput { .Node = "WavePlayer", .NodeOutput = "Out Side Right", .GraphOutput = OutputFormatFiveDotOneInterface::Outputs::SideRightOut }
+				};
+			}
+
+			if (NumChannels == 8)
+			{
+				return
+				{
+					FDC::FEdgeToOutput { .Node = "WavePlayer", .NodeOutput = "Out Front Left", .GraphOutput = OutputFormatSevenDotOneInterface::Outputs::FrontLeftOut },
+					FDC::FEdgeToOutput { .Node = "WavePlayer", .NodeOutput = "Out Front Right", .GraphOutput = OutputFormatSevenDotOneInterface::Outputs::FrontRightOut },
+					FDC::FEdgeToOutput { .Node = "WavePlayer", .NodeOutput = "Out Front Center", .GraphOutput = OutputFormatSevenDotOneInterface::Outputs::FrontCenterOut },
+					FDC::FEdgeToOutput { .Node = "WavePlayer", .NodeOutput = "Out Low Frequency", .GraphOutput = OutputFormatSevenDotOneInterface::Outputs::LowFrequencyOut },
+					FDC::FEdgeToOutput { .Node = "WavePlayer", .NodeOutput = "Out Side Left", .GraphOutput = OutputFormatSevenDotOneInterface::Outputs::SideLeftOut },
+					FDC::FEdgeToOutput { .Node = "WavePlayer", .NodeOutput = "Out Side Right", .GraphOutput = OutputFormatSevenDotOneInterface::Outputs::SideRightOut },
+					FDC::FEdgeToOutput { .Node = "WavePlayer", .NodeOutput = "Out Back Left", .GraphOutput = OutputFormatSevenDotOneInterface::Outputs::BackLeftOut },
+					FDC::FEdgeToOutput { .Node = "WavePlayer", .NodeOutput = "Out Back Right", .GraphOutput = OutputFormatSevenDotOneInterface::Outputs::BackRightOut }
+				};
+			}
+
+			if (NumChannels != 1)
+			{
+				UE_LOG(LogMetaSound, Warning, TEXT("SoundWave format with %i NumChannels not recognized: Connecting only first channel as mono."), NumChannels);
+			}
+
+			return
+			{
+				FDC::FEdgeToOutput { .Node = "WavePlayer", .NodeOutput = "Out Mono", .GraphOutput = OutputFormatMonoInterface::Outputs::MonoOut }
+			};
+		}
+	} // namespace RandomizerPrivate
+} // namespace Metasound::Engine
+
+#if WITH_EDITORONLY_DATA
+bool FMetaSoundRandomizerDocumentConfiguration::ConfigureDocument(FMetaSoundFrontendDocumentBuilder& OutBuilder, bool bLogFailure)
+{
+	using namespace Metasound;
+	using namespace Metasound::Engine;
+	using namespace Metasound::Frontend;
+
+	using FDC = FDocumentConfigurator;
+
+	// If max channels cannot be deduced, then return clear document and consider it valid as
+	// user may not be supplying waves for a reason (or may be amidst editing array of sounds).
+	const int32 MaxNumChannels = RandomizerPrivate::GetMaxChannels(Sounds);
+
+	FDC Configurator(FDC::FArgs { .Builder = &OutBuilder });
+
+	// Add interfaces first.  Done even if no sounds supplied to make sure resulting document is valid source.
+	Configurator.Add<FDC::FInterface>(
+	{
+		{ .Version = SourceInterface::GetVersion() },
+		{ .Version = RandomizerPrivate::GetOutputVersion(MaxNumChannels) }
+	});
+
+	if (bIsOneShot)
+	{
+		Configurator.Add(FDC::FInterface { .Version = SourceOneShotInterface::GetVersion() });
+	}
+
+	if (MaxNumChannels <= 0)
+	{
+		return true;
+	}
+
+	const FName WaveArrayTypeName = GetMetasoundDataTypeName<TArray<FWaveAsset>>();
+	const FMetasoundFrontendClassName PlayerClassName = RandomizerPrivate::GetPlayerClassName(MaxNumChannels);
+	const FMetasoundFrontendClassName RandomClassName = FMetasoundFrontendClassName("Array", "Random Get", WaveArrayTypeName);
+
+	Configurator
+
+	// Add Non-Interface Graph Inputs
+	.Add(FDC::FInput { .Name = "Sounds", .DataType = WaveArrayTypeName, .DefaultValue = MakeFrontendLiteral<USoundWave>(Sounds), .Location = FVector2D(-500, 250) })
+
+	// Set Interface Graph Input locations
+	.Set(FDC::FInputLocation { .Name = "UE.Source.OnPlay", .Location = FVector2D(-500, 0)})
+
+	// Add Nodes
+	.Add<FDC::FNode>(
+	{
+		{ .Name = "Pitch", .ClassName = FMetasoundFrontendClassName("UE", "RandomFloat"), .Location = FVector2D(250, 400) },
+		{ .Name = "RandomWave", .ClassName = RandomClassName, .Location = FVector2D(0, 0) },
+		{ .Name = "WavePlayer", .ClassName = PlayerClassName, .Location = FVector2D(500, 0) }
+	})
+
+	.Set(FDC::FNodeInputDefault { .Node = "Pitch", .Input = "Min", .Value = MakeFrontendLiteral(-12) })
+	.Set(FDC::FNodeInputDefault { .Node = "Pitch", .Input = "Max", .Value = MakeFrontendLiteral(12) })
+
+	// Connect Graph Inputs to Nodes
+	.Add<FDC::FEdgeFromInput>(
+	{
+		{ .GraphInput = "Sounds",			.Node = "RandomWave", .NodeInput = "In Array" },
+		{ .GraphInput = "UE.Source.OnPlay", .Node = "RandomWave", .NodeInput = "Next" }
+	})
+
+	// Connect Nodes Together
+	.Add<FDC::FEdge>(
+	{
+		{ .OutputNode = "RandomWave", .Output = "On Next",		.InputNode = "WavePlayer", .Input = "Play" },
+		{ .OutputNode = "RandomWave", .Output = "Value",		.InputNode = "WavePlayer", .Input = "Wave Asset" },
+		{ .OutputNode = "Pitch",	  .Output = "Value",		.InputNode = "WavePlayer", .Input = "Pitch Shift" },
+	})
+
+	// Connect Nodes to Graph Outputs
+	.Add<FDC::FEdgeToOutput>(RandomizerPrivate::GetAudioOutputEdges(MaxNumChannels));
+
+	if (bIsOneShot)
+	{
+		Configurator.Add(FDC::FEdgeToOutput { .Node = "WavePlayer", .NodeOutput = "On Finished", .GraphOutput = SourceOneShotInterface::Outputs::OnFinished });
+	}
+	else
+	{
+		Configurator.Set(FDC::FNodeInputDefault { .Node = "WavePlayer", .Input = "Loop", .Value = MakeFrontendLiteral(true) });
+	}
+
+	return Configurator.Succeeded();
+}
+#endif // WITH_EDITORONLY_DATA
+
+#if WITH_EDITOR
+void FMetaSoundRandomizerDocumentConfiguration::OnPropertyChanged(const FPropertyChangedEvent& InEvent, FMetaSoundFrontendDocumentBuilder& OutBuilder)
+{
+	ConfigureDocument(OutBuilder, false);
+}
+#endif // WITH_EDITOR
+#undef LOCTEXT_NAMESPACE
diff --git a/Engine/Plugins/Runtime/Metasound/Source/MetasoundEngine/Private/MetasoundBuilderBase.cpp b/Engine/Plugins/Runtime/Metasound/Source/MetasoundEngine/Private/MetasoundBuilderBase.cpp
--- a/Engine/Plugins/Runtime/Metasound/Source/MetasoundEngine/Private/MetasoundBuilderBase.cpp
+++ b/Engine/Plugins/Runtime/Metasound/Source/MetasoundEngine/Private/MetasoundBuilderBase.cpp
@@ -488,9 +488,15 @@ TArray<FMetaSoundBuilderNodeInputHandle> UMetaSoundBuilderBase::ConnectNodeOutpu
 	TArray<FMetaSoundBuilderNodeInputHandle> ConnectedVertices;
 	Algo::Transform(NewEdges, ConnectedVertices, [this](const FMetasoundFrontendEdge* NewEdge)
 	{
-		const FMetasoundFrontendVertex* Vertex = Builder.FindNodeInput(NewEdge->ToNodeID, NewEdge->ToVertexID);
-		checkf(Vertex, TEXT("Edge connection reported success but vertex not found."));
-		return FMetaSoundBuilderNodeInputHandle(NewEdge->ToNodeID, Vertex->VertexID);
+#if DO_CHECK
+		{
+			check(NewEdge);
+			const FMetasoundFrontendVertex* Vertex = Builder.FindNodeInput(NewEdge->ToNodeID, NewEdge->ToVertexID);
+			checkf(Vertex, TEXT("Edge node output to interface output connection reported success but vertex not found."));
+		}
+#endif // DO_CHECK
+
+		return FMetaSoundBuilderNodeInputHandle(NewEdge->ToNodeID, NewEdge->ToVertexID);
 	});
 
 	return ConnectedVertices;
@@ -505,9 +511,15 @@ TArray<FMetaSoundBuilderNodeOutputHandle> UMetaSoundBuilderBase::ConnectNodeInpu
 	TArray<FMetaSoundBuilderNodeOutputHandle> ConnectedVertices;
 	Algo::Transform(NewEdges, ConnectedVertices, [this](const FMetasoundFrontendEdge* NewEdge)
 	{
-		const FMetasoundFrontendVertex* Vertex = Builder.FindNodeOutput(NewEdge->FromNodeID, NewEdge->FromVertexID);
-		checkf(Vertex, TEXT("Edge connection reported success but vertex not found."));
-		return FMetaSoundBuilderNodeOutputHandle(NewEdge->ToNodeID, Vertex->VertexID);
+#if DO_CHECK
+		{
+			check(NewEdge);
+			const FMetasoundFrontendVertex* Vertex = Builder.FindNodeOutput(NewEdge->FromNodeID, NewEdge->FromVertexID);
+			checkf(Vertex, TEXT("Edge node input to interface input connection reported success but vertex not found."));
+		}
+#endif // DO_CHECK
+
+		return FMetaSoundBuilderNodeOutputHandle(NewEdge->FromNodeID, NewEdge->FromVertexID);
 	});
 
 	return ConnectedVertices;
@@ -917,13 +929,6 @@ UMetaSoundFrontendMemberMetadata* UMetaSoundBuilderBase::FindMemberMetadata(cons
 }
 #endif // WITH_EDITOR
 
-UMetaSoundBuilderDocument* UMetaSoundBuilderBase::CreateTransientDocumentObject() const
-{
-PRAGMA_DISABLE_DEPRECATION_WARNINGS
-	return &UMetaSoundBuilderDocument::Create(GetBaseMetaSoundUClass());
-PRAGMA_ENABLE_DEPRECATION_WARNINGS
-}
-
 FMetaSoundFrontendDocumentBuilder& UMetaSoundBuilderBase::GetBuilder()
 {
 	return Builder;
@@ -970,6 +975,8 @@ TArray<FName> UMetaSoundBuilderBase::GetGraphInputNames(EMetaSoundBuilderResult&
 	{
 		GraphInputNames.Emplace(Input.Name);
 	}
+
+	OutResult = EMetaSoundBuilderResult::Succeeded;
 	return GraphInputNames;
 }
 
@@ -981,6 +988,8 @@ TArray<FName> UMetaSoundBuilderBase::GetGraphOutputNames(EMetaSoundBuilderResult
 	{
 		GraphOutputNames.Emplace(Output.Name);
 	}
+
+	OutResult = EMetaSoundBuilderResult::Succeeded;
 	return GraphOutputNames;
 }
 
@@ -1051,11 +1060,6 @@ bool UMetaSoundBuilderBase::InterfaceIsDeclared(FName InterfaceName) const
 	return Builder.IsInterfaceDeclared(InterfaceName);
 }
 
-void UMetaSoundBuilderBase::InvalidateCache(bool bPrimeCache)
-{
-	Reload({ }, bPrimeCache);
-}
-
 bool UMetaSoundBuilderBase::IsPreset() const
 {
 	return Builder.IsPreset();
diff --git a/Engine/Plugins/Runtime/Metasound/Source/MetasoundEngine/Private/MetasoundBuilderSubsystem.cpp b/Engine/Plugins/Runtime/Metasound/Source/MetasoundEngine/Private/MetasoundBuilderSubsystem.cpp
--- a/Engine/Plugins/Runtime/Metasound/Source/MetasoundEngine/Private/MetasoundBuilderSubsystem.cpp
+++ b/Engine/Plugins/Runtime/Metasound/Source/MetasoundEngine/Private/MetasoundBuilderSubsystem.cpp
@@ -138,59 +138,6 @@ void UMetaSoundPatchBuilder::OnRemovingAssetReference(TScriptInterface<IMetaSoun
 	Patch.ReferencedAssetClassKeys.Remove(RegistryKey.ToString());
 }
 
-UMetaSoundBuilderBase& UMetaSoundBuilderSubsystem::AttachBuilderToAssetChecked(UObject& InObject) const
-{
-	const UClass* BaseClass = InObject.GetClass();
-	if (BaseClass == UMetaSoundSource::StaticClass())
-	{
-PRAGMA_DISABLE_DEPRECATION_WARNINGS
-		UMetaSoundSourceBuilder* NewBuilder = AttachSourceBuilderToAsset(CastChecked<UMetaSoundSource>(&InObject));
-PRAGMA_ENABLE_DEPRECATION_WARNINGS
-		return *NewBuilder;
-	}
-	else if (BaseClass == UMetaSoundPatch::StaticClass())
-	{
-PRAGMA_DISABLE_DEPRECATION_WARNINGS
-		UMetaSoundPatchBuilder* NewBuilder = AttachPatchBuilderToAsset(CastChecked<UMetaSoundPatch>(&InObject));
-PRAGMA_ENABLE_DEPRECATION_WARNINGS
-		return *NewBuilder;
-	}
-	else
-	{
-		checkf(false, TEXT("UClass '%s' is not a base MetaSound that supports attachment via the MetaSoundBuilderSubsystem"), *BaseClass->GetFullName());
-		return *NewObject<UMetaSoundPatchBuilder>();
-	}
-}
-
-UMetaSoundPatchBuilder* UMetaSoundBuilderSubsystem::AttachPatchBuilderToAsset(UMetaSoundPatch* InPatch) const
-{
-#if WITH_EDITORONLY_DATA
-	using namespace Metasound::Engine;
-
-	if (InPatch)
-	{
-		return &FDocumentBuilderRegistry::GetChecked().FindOrBeginBuilding<UMetaSoundPatchBuilder>(*InPatch);
-	}
-#endif // WITH_EDITORONLY_DATA
-
-	return nullptr;
-}
-
-UMetaSoundSourceBuilder* UMetaSoundBuilderSubsystem::AttachSourceBuilderToAsset(UMetaSoundSource* InSource) const
-{
-#if WITH_EDITORONLY_DATA
-	using namespace Metasound::Engine;
-
-	if (InSource)
-	{
-		UMetaSoundSourceBuilder& SourceBuilder = FDocumentBuilderRegistry::GetChecked().FindOrBeginBuilding<UMetaSoundSourceBuilder>(*InSource);
-		return &SourceBuilder;
-	}
-#endif // WITH_EDITORONLY_DATA
-
-	return nullptr;
-}
-
 void UMetaSoundSourceBuilder::Audition(UObject* Parent, UAudioComponent* AudioComponent, FOnCreateAuditionGeneratorHandleDelegate CreateGenerator, bool bLiveUpdatesEnabled)
 {
 	using namespace Metasound;
@@ -405,9 +352,10 @@ void UMetaSoundSourceBuilder::InitDelegates(Metasound::Frontend::FDocumentModify
 
 	if (!IsRunningCookCommandlet())
 	{
+		OutDocumentDelegates.OnDocumentConfigurationChanged.AddUObject(this, &UMetaSoundSourceBuilder::OnDocumentConfigurationChanged);
+
 		OutDocumentDelegates.PageDelegates.OnPageAdded.AddUObject(this, &UMetaSoundSourceBuilder::OnPageAdded);
 		OutDocumentDelegates.PageDelegates.OnRemovingPage.AddUObject(this, &UMetaSoundSourceBuilder::OnRemovingPage);
-		OutDocumentDelegates.OnPresetStateChanged.AddUObject(this, &UMetaSoundSourceBuilder::OnPresetStateChanged);
 
 		OutDocumentDelegates.InterfaceDelegates.OnInputAdded.AddUObject(this, &UMetaSoundSourceBuilder::OnInputAdded);
 		OutDocumentDelegates.InterfaceDelegates.OnInputDataTypeChanged.AddUObject(this, &UMetaSoundSourceBuilder::OnInputDataTypeChanged);
@@ -488,6 +436,13 @@ void UMetaSoundSourceBuilder::OnAssetReferenceAdded(TScriptInterface<IMetaSoundD
 	Source.ReferencedAssetClassKeys.Add(RegistryKey.ToString());
 }
 
+void UMetaSoundSourceBuilder::OnDocumentConfigurationChanged(const Metasound::Frontend::FDocumentConfigChangedArgs&) const
+{
+	using namespace Metasound::Frontend;
+
+	GetMetaSoundSource().InvalidateCachedRuntimeInputData();
+}
+
 void UMetaSoundSourceBuilder::OnEdgeAdded(int32 EdgeIndex) const
 {
 	using namespace Metasound::DynamicGraph;
@@ -526,7 +481,10 @@ void UMetaSoundSourceBuilder::OnInputAdded(int32 InputIndex)
 				{
 					AudioDevice->SendCommandToActiveSounds(AudioComponentID, [NewInputName = NewInput.Name](FActiveSound& ActiveSound)
 					{
-						static_cast<FMetaSoundParameterTransmitter*>(ActiveSound.GetTransmitter())->AddAvailableParameter(NewInputName);
+						if (FMetaSoundParameterTransmitter* Transmitter = static_cast<FMetaSoundParameterTransmitter*>(ActiveSound.GetTransmitter()))
+						{
+							Transmitter->AddAvailableParameter(NewInputName);
+						}
 					});
 				}
 			}
@@ -591,10 +549,11 @@ void UMetaSoundSourceBuilder::OnInputNameChanged(FName OldName, FName NewName)
 					{
 						AudioDevice->SendCommandToActiveSounds(AudioComponentID, [OldName, NewInputName = Input->Name](FActiveSound& ActiveSound)
 						{
-							FMetaSoundParameterTransmitter* Transmitter = static_cast<FMetaSoundParameterTransmitter*>(ActiveSound.GetTransmitter());
-							check(Transmitter);
-							Transmitter->RemoveAvailableParameter(OldName);
-							Transmitter->AddAvailableParameter(NewInputName);
+							if (FMetaSoundParameterTransmitter* Transmitter = static_cast<FMetaSoundParameterTransmitter*>(ActiveSound.GetTransmitter()))
+							{
+								Transmitter->RemoveAvailableParameter(OldName);
+								Transmitter->AddAvailableParameter(NewInputName);
+							}
 						});
 					}
 				}
@@ -806,13 +765,6 @@ void UMetaSoundSourceBuilder::OnPageAdded(const Metasound::Frontend::FDocumentMu
 	InitTargetPageDelegates(DocDelegates);
 }
 
-void UMetaSoundSourceBuilder::OnPresetStateChanged(const Metasound::Frontend::FDocumentPresetStateChangedArgs& InArgs) const
-{
-	using namespace Metasound::Frontend;
-
-	GetMetaSoundSource().InvalidateCachedRuntimeInputData();
-}
-
 void UMetaSoundSourceBuilder::OnRemovingPage(const Metasound::Frontend::FDocumentMutatePageArgs& InArgs)
 {
 	using namespace Metasound::Frontend;
@@ -941,9 +893,12 @@ void UMetaSoundSourceBuilder::OnRemovingInput(int32 InputIndex)
 				if (FAudioDevice* AudioDevice = AudioComponent->GetAudioDevice())
 				{
 					AudioDevice->SendCommandToActiveSounds(AudioComponentID, [InputRemoved = InputBeingRemoved.Name](FActiveSound& ActiveSound)
+					{
+						if (FMetaSoundParameterTransmitter* Transmitter = static_cast<FMetaSoundParameterTransmitter*>(ActiveSound.GetTransmitter()))
 						{
-							static_cast<FMetaSoundParameterTransmitter*>(ActiveSound.GetTransmitter())->RemoveAvailableParameter(InputRemoved);
-						});
+							Transmitter->RemoveAvailableParameter(InputRemoved);
+						}
+					});
 				}
 			}
 		}
@@ -1385,12 +1340,6 @@ FMetasoundFrontendLiteral UMetaSoundBuilderSubsystem::CreateMetaSoundLiteralFrom
 	return FMetasoundFrontendLiteral { Param };
 }
 
-bool UMetaSoundBuilderSubsystem::DetachBuilderFromAsset(const FMetasoundFrontendClassName& InClassName) const
-{
-	using namespace Metasound::Frontend;
-	return IDocumentBuilderRegistry::GetChecked().FinishBuilding(InClassName);
-}
-
 UMetaSoundBuilderBase* UMetaSoundBuilderSubsystem::FindBuilder(FName BuilderName)
 {
 	return NamedBuilders.FindRef(BuilderName);
@@ -1486,12 +1435,6 @@ UMetaSoundSourceBuilder* UMetaSoundBuilderSubsystem::FindSourceBuilder(FName Bui
 	return nullptr;
 }
 
-void UMetaSoundBuilderSubsystem::InvalidateDocumentCache(const FMetasoundFrontendClassName& InClassName) const
-{
-	using namespace Metasound::Engine;
-	FDocumentBuilderRegistry::GetChecked().ReloadBuilder(InClassName);
-}
-
 bool UMetaSoundBuilderSubsystem::IsInterfaceRegistered(FName InInterfaceName) const
 {
 	using namespace Metasound::Frontend;
diff --git a/Engine/Plugins/Runtime/Metasound/Source/MetasoundEngine/Public/DocumentConfigurations/MetasoundDocumentConfigurator.h b/Engine/Plugins/Runtime/Metasound/Source/MetasoundEngine/Public/DocumentConfigurations/MetasoundDocumentConfigurator.h
new file mode 100644
--- /dev/null
+++ b/Engine/Plugins/Runtime/Metasound/Source/MetasoundEngine/Public/DocumentConfigurations/MetasoundDocumentConfigurator.h
@@ -0,0 +1,172 @@
+// Copyright Epic Games, Inc. All Rights Reserved.
+#pragma once
+
+#include "Algo/Transform.h"
+#include "Containers/ContainersFwd.h"
+#include "MetasoundBuilderBase.h"
+#include "MetasoundBuilderSubsystem.h"
+#include "MetasoundFrontendDocument.h"
+#include "MetasoundFrontendGraphLayer.h"
+#include "MetasoundFrontendLiteral.h"
+#include "UObject/ObjectPtr.h"
+
+
+#define UE_API METASOUNDENGINE_API
+
+
+namespace Metasound::Engine
+{
+	struct FDocumentConfigurator
+	{
+		struct FArgs
+		{
+			FMetaSoundFrontendDocumentBuilder* Builder = nullptr;
+
+#if WITH_EDITORONLY_DATA
+			// Offset to use when positioning new nodes via the configuration
+			FVector2D Offset = FVector2D::Zero();
+#endif // WITH_EDITORONLY_DATA
+
+			bool bResetDoc = true;
+		};
+
+		FDocumentConfigurator(const FArgs& Args);
+		~FDocumentConfigurator();
+
+#if WITH_EDITORONLY_DATA
+		UE_EXPERIMENTAL(5.8, "Configurator code reflection is in development and may change implementation or be removed in the future")
+		UE_API static void Reflect(const FMetasoundFrontendGraphLayer& InLayer, FString& OutCode);
+#endif // WITH_EDITORONLY_DATA
+
+		struct FInterface
+		{
+			FMetasoundFrontendVersion Version;
+		};
+		UE_EXPERIMENTAL(5.8, "Configurator is in active development and subject to change")
+		UE_API FDocumentConfigurator& Add(FInterface Interface);
+
+		struct FInput
+		{
+			FName Name = "Float";
+			FName DataType = GetMetasoundDataTypeName<float>();
+			FMetasoundFrontendLiteral DefaultValue = Frontend::MakeFrontendLiteral(0.f);
+			FVector2D Location = FVector2D::Zero();
+			EMetasoundFrontendVertexAccessType AccessType = EMetasoundFrontendVertexAccessType::Reference;
+		};
+		UE_EXPERIMENTAL(5.8, "Configurator is in active development and subject to change")
+		UE_API FDocumentConfigurator& Add(FInput Input);
+
+		struct FOutput
+		{
+			FName Name = "Float";
+			FName DataType = GetMetasoundDataTypeName<float>();
+			FMetasoundFrontendLiteral DefaultValue = Frontend::MakeFrontendLiteral(0.f);
+			FVector2D Location = FVector2D::Zero();
+			EMetasoundFrontendVertexAccessType AccessType = EMetasoundFrontendVertexAccessType::Reference;
+		};
+		UE_EXPERIMENTAL(5.8, "Configurator is in active development and subject to change")
+		UE_API FDocumentConfigurator& Add(FOutput Output);
+
+		struct FNode
+		{
+			FName Name;
+			FMetasoundFrontendClassName ClassName;
+			FVector2D Location = FVector2D::Zero();
+		};
+		UE_EXPERIMENTAL(5.8, "Configurator is in active development and subject to change")
+		UE_API FDocumentConfigurator& Add(FNode Node);
+
+		struct FEdge
+		{
+			FName OutputNode;
+			FName Output;
+			FName InputNode;
+			FName Input;
+		};
+		UE_EXPERIMENTAL(5.8, "Configurator is in active development and subject to change")
+		UE_API FDocumentConfigurator& Add(FEdge Edge);
+
+		struct FEdgeFromInput
+		{
+			FName GraphInput;
+			FName Node;
+			FName NodeInput;
+		};
+		UE_EXPERIMENTAL(5.8, "Configurator is in active development and subject to change")
+		UE_API FDocumentConfigurator& Add(FEdgeFromInput Edge);
+
+		struct FEdgeToOutput
+		{
+			FName Node;
+			FName NodeOutput;
+			FName GraphOutput;
+		};
+		UE_EXPERIMENTAL(5.8, "Configurator is in active development and subject to change")
+		UE_API FDocumentConfigurator& Add(FEdgeToOutput Edge);
+
+		template <typename ConfigType>
+		FDocumentConfigurator& Add(TArray<ConfigType> Args)
+		{
+			for (ConfigType& ConfigArgs : Args)
+			{
+				Add(MoveTemp(ConfigArgs));
+			}
+
+			return *this;
+		}
+
+		template <>
+		FDocumentConfigurator& Add<FInterface>(TArray<FInterface> Interfaces)
+		{
+			return AddInternal(MoveTemp(Interfaces));
+		}
+
+		UE_EXPERIMENTAL(5.8, "Configurator is in active development and subject to change")
+		UE_API UMetaSoundBuilderBase& GetBuilder();
+
+		UE_EXPERIMENTAL(5.8, "Configurator is in active development and subject to change")
+		UE_API FDocumentConfigurator& ResetDocument();
+
+		struct FInputLocation
+		{
+			FName Name;
+			FVector2D Location = FVector2D::Zero();
+		};
+		UE_EXPERIMENTAL(5.8, "Configurator is in active development and subject to change")
+		UE_API FDocumentConfigurator& Set(FInputLocation Location);
+
+		struct FNodeInputDefault
+		{
+			FName Node;
+			FName Input;
+			FMetasoundFrontendLiteral Value;
+		};
+		UE_EXPERIMENTAL(5.8, "Configurator is in active development and subject to change")
+		UE_API FDocumentConfigurator& Set(FNodeInputDefault Input);
+
+		template <typename ConfigType>
+		FDocumentConfigurator& Set(TArray<ConfigType> Args)
+		{
+			for (ConfigType& ConfigArgs : Args)
+			{
+				Set(MoveTemp(ConfigArgs));
+			}
+
+			return *this;
+		}
+
+		UE_EXPERIMENTAL(5.8, "Configurator is in active development and subject to change")
+		UE_API bool Succeeded() const;
+
+	private:
+		FDocumentConfigurator& AddInternal(TArray<FInterface> Interfaces);
+
+		TStrongObjectPtr<UMetaSoundBuilderBase> Builder;
+		mutable EMetaSoundBuilderResult Result = EMetaSoundBuilderResult::Succeeded;
+
+		TMap<FName, FMetaSoundNodeHandle> NodeNameToHandle;
+
+		FVector2D Offset;
+	};
+} // namespace Metasound::Engine
+#undef UE_API
diff --git a/Engine/Plugins/Runtime/Metasound/Source/MetasoundEngine/Public/DocumentConfigurations/MetasoundRandomizerDocumentConfiguration.h b/Engine/Plugins/Runtime/Metasound/Source/MetasoundEngine/Public/DocumentConfigurations/MetasoundRandomizerDocumentConfiguration.h
new file mode 100644
--- /dev/null
+++ b/Engine/Plugins/Runtime/Metasound/Source/MetasoundEngine/Public/DocumentConfigurations/MetasoundRandomizerDocumentConfiguration.h
@@ -0,0 +1,42 @@
+// Copyright Epic Games, Inc. All Rights Reserved.
+#pragma once
+
+#include "MetasoundDocumentInterface.h"
+#include "MetasoundFrontendDocumentConfiguration.h"
+#include "Sound/SoundWave.h"
+#include "UObject/ObjectPtr.h"
+
+#include "MetasoundRandomizerDocumentConfiguration.generated.h"
+
+#define UE_API METASOUNDENGINE_API
+
+// Forward Declarations
+class UMetaSoundFrontendMemberMetadata;
+enum class EMetasoundFrontendClassType : uint8;
+
+
+USTRUCT(BlueprintType, DisplayName = "Randomizer")
+struct FMetaSoundRandomizerDocumentConfiguration : public FMetaSoundFrontendDocumentConfiguration
+{
+	GENERATED_BODY()
+
+	virtual ~FMetaSoundRandomizerDocumentConfiguration() = default;
+
+#if WITH_EDITORONLY_DATA
+	UE_EXPERIMENTAL(5.8, "Document Configuration is in active development and subject to change")
+	UE_API virtual bool ConfigureDocument(FMetaSoundFrontendDocumentBuilder& OutBuilder, bool bLogFailure) override;
+#endif // WITH_EDITORONLY_DATA
+
+#if WITH_EDITOR
+	UE_EXPERIMENTAL(5.8, "Document Configuration is in active development and subject to change")
+	UE_API virtual void OnPropertyChanged(const FPropertyChangedEvent& InEvent, FMetaSoundFrontendDocumentBuilder& OutBuilder) override;
+#endif // WITH_EDITOR
+
+public:
+	UPROPERTY(EditAnywhere, BlueprintReadOnly, meta = (Category = "Randomizer Options"))
+	bool bIsOneShot = false;
+
+	UPROPERTY(EditAnywhere, BlueprintReadOnly, meta = (Category = "Randomizer Options"))
+	TArray<TObjectPtr<USoundWave>> Sounds;
+};
+#undef UE_API
\ No newline at end of file
diff --git a/Engine/Plugins/Runtime/Metasound/Source/MetasoundEngine/Public/MetasoundBuilderBase.h b/Engine/Plugins/Runtime/Metasound/Source/MetasoundEngine/Public/MetasoundBuilderBase.h
--- a/Engine/Plugins/Runtime/Metasound/Source/MetasoundEngine/Public/MetasoundBuilderBase.h
+++ b/Engine/Plugins/Runtime/Metasound/Source/MetasoundEngine/Public/MetasoundBuilderBase.h
@@ -583,12 +583,6 @@ protected:
 	virtual void OnAssetReferenceAdded(TScriptInterface<IMetaSoundDocumentInterface> DocInterface) PURE_VIRTUAL(UMetaSoundBuilderBase::OnAssetReferenceAdded, );
 	virtual void OnRemovingAssetReference(TScriptInterface<IMetaSoundDocumentInterface> DocInterface) PURE_VIRTUAL(UMetaSoundBuilderBase::OnRemovingAssetReference, );
 
-	UE_DEPRECATED(5.5, "Moved to internal implementation")
-	virtual void CreateTransientBuilder() { };
-
-	UE_DEPRECATED(5.5, "Moved to 'Reload', to enforce generation of new delegates")
-	UE_API void InvalidateCache(bool bPrimeCache = false);
-
 	// Runs build, conforming the document and corresponding object data on a MetaSound UObject to that managed by this builder.
 	template <typename UClassType>
 	UClassType& BuildInternal(UObject* Parent, const FMetaSoundBuilderOptions& BuilderOptions) const
@@ -650,9 +644,6 @@ protected:
 		return *MetaSound;
 	}
 
-	UE_DEPRECATED(5.4, "Use UMetaSoundBuilderDocument::Create instead")
-	UE_API UMetaSoundBuilderDocument* CreateTransientDocumentObject() const;
-
 	// Only registers provided MetaSound's graph class and referenced graphs recursively if
 	// it has yet to be registered or if it has an attached builder reporting outstanding
 	// transactions that have yet to be registered.
diff --git a/Engine/Plugins/Runtime/Metasound/Source/MetasoundEngine/Public/MetasoundBuilderSubsystem.h b/Engine/Plugins/Runtime/Metasound/Source/MetasoundEngine/Public/MetasoundBuilderSubsystem.h
--- a/Engine/Plugins/Runtime/Metasound/Source/MetasoundEngine/Public/MetasoundBuilderSubsystem.h
+++ b/Engine/Plugins/Runtime/Metasound/Source/MetasoundEngine/Public/MetasoundBuilderSubsystem.h
@@ -122,6 +122,7 @@ private:
 
 	UE_API void InitTargetPageDelegates(Metasound::Frontend::FDocumentModifyDelegates& OutDocumentDelegates);
 
+	UE_API void OnDocumentConfigurationChanged(const Metasound::Frontend::FDocumentConfigChangedArgs& Args) const;
 	UE_API void OnEdgeAdded(int32 EdgeIndex) const;
 	UE_API void OnInputAdded(int32 InputIndex);
 	UE_API void OnInputDataTypeChanged(int32 InputIndex);
@@ -133,7 +134,6 @@ private:
 	UE_API void OnNodeInputLiteralSet(int32 NodeIndex, int32 VertexIndex, int32 LiteralIndex) const;
 	UE_API void OnOutputAdded(int32 OutputIndex) const;
 	UE_API void OnPageAdded(const Metasound::Frontend::FDocumentMutatePageArgs& InArgs);
-	UE_API void OnPresetStateChanged(const Metasound::Frontend::FDocumentPresetStateChangedArgs& InArgs) const;
 	UE_API void OnRemoveSwappingEdge(int32 SwapIndex, int32 LastIndex) const;
 	UE_API void OnRemovingInput(int32 InputIndex);
 	UE_API void OnRemoveSwappingNode(int32 SwapIndex, int32 LastIndex) const;
@@ -169,26 +169,11 @@ private:
 	TMap<FName, TObjectPtr<UMetaSoundBuilderBase>> NamedBuilders;
 
 public:	
-	UE_DEPRECATED(5.5, "Call 'ReloadBuilder' in IDocumentBuilderRegistry instead ")
-	UE_API virtual void InvalidateDocumentCache(const FMetasoundFrontendClassName& InClassName) const;
-
 	static UE_API UMetaSoundBuilderSubsystem* Get();
 	static UE_API UMetaSoundBuilderSubsystem& GetChecked();
 	static UE_API const UMetaSoundBuilderSubsystem* GetConst();
 	static UE_API const UMetaSoundBuilderSubsystem& GetConstChecked();
 
-	UE_DEPRECATED(5.5, "Use FDocumentBuilderRegistry::FindOrBeginBuilding, which is now only supported in builds loading editor-only data.")
-	UE_API UMetaSoundBuilderBase& AttachBuilderToAssetChecked(UObject& InObject) const;
-
-	UE_DEPRECATED(5.5, "Use FDocumentBuilderRegistry::FindOrBeginBuilding (when editor only data is loaded) or MetaSoundEditorSubsystem::FindOrBeginBuilding call")
-	UE_API UMetaSoundPatchBuilder* AttachPatchBuilderToAsset(UMetaSoundPatch* InPatch) const;
-
-	UE_DEPRECATED(5.5, "Use FDocumentBuilderRegistry::FindOrBeginBuilding (when editor only data is loaded) or MetaSoundEditorSubsystem::FindOrBeginBuilding call")
-	UE_API UMetaSoundSourceBuilder* AttachSourceBuilderToAsset(UMetaSoundSource* InSource) const;
-
-	UE_DEPRECATED(5.5, "Moved to IDocumentBuilderRegistry::RemoveBuilderFromAsset")
-	UE_API bool DetachBuilderFromAsset(const FMetasoundFrontendClassName& InClassName) const;
-
 	UFUNCTION(BlueprintCallable, Category = "Audio|MetaSound|Builder",  meta = (ExpandEnumAsExecs = "OutResult"))
 	UE_API UPARAM(DisplayName = "Patch Builder") UMetaSoundPatchBuilder* CreatePatchBuilder(FName BuilderName, EMetaSoundBuilderResult& OutResult);
 
diff --git a/Engine/Plugins/Runtime/Metasound/Source/MetasoundFrontend/Private/DocumentConfigurations/MetasoundPresetDocumentConfiguration.cpp b/Engine/Plugins/Runtime/Metasound/Source/MetasoundFrontend/Private/DocumentConfigurations/MetasoundPresetDocumentConfiguration.cpp
new file mode 100644
--- /dev/null
+++ b/Engine/Plugins/Runtime/Metasound/Source/MetasoundFrontend/Private/DocumentConfigurations/MetasoundPresetDocumentConfiguration.cpp
@@ -0,0 +1,94 @@
+// Copyright Epic Games, Inc. All Rights Reserved.
+#include "DocumentConfigurations/MetasoundPresetDocumentConfiguration.h"
+
+#include "MetasoundDocumentInterface.h"
+#include "MetasoundFrontendDocument.h"
+#include "MetasoundFrontendDocumentBuilder.h"
+#include "MetasoundFrontendDocumentBuilderRegistry.h"
+
+
+#define LOCTEXT_NAMESPACE "MetaSound"
+
+namespace PresetDocConfigPrivate
+{
+	bool IsInputNode(const FMetaSoundFrontendDocumentBuilder& Builder, const FGuid& NodeID)
+	{
+		if (const FMetasoundFrontendNode* Node = Builder.FindNode(NodeID))
+		{
+			if (const FMetasoundFrontendClass* Class = Builder.FindDependency(Node->ClassID))
+			{
+				return Class->Metadata.GetType() == EMetasoundFrontendClassType::Input;
+			}
+		}
+
+		return false;
+	}
+} // namespace PresetDocConfigPrivate
+
+#if WITH_EDITORONLY_DATA
+bool FMetaSoundFrontendPresetDocumentConfiguration::ConfigureDocument(FMetaSoundFrontendDocumentBuilder& OutBuilder, bool bLogFailure)
+{
+	using namespace Metasound::Frontend;
+
+	if (UObject* RefObject = Parent.GetObject())
+	{
+		if (RefObject->IsA(&OutBuilder.GetConstDocumentInterfaceChecked().GetBaseMetaSoundUClass()))
+		{
+			const FMetaSoundFrontendDocumentBuilder& ParentBuilder = IDocumentBuilderRegistry::GetChecked().FindOrBeginBuilding(Parent);
+
+			// Ensure the referenced node is registered as it must be accessible for the following call to rebuild transform
+			// In the future, this should not be necessary as all data required for rebuild should be accessed straight from
+			// supplied referenced builder (registration should only be necessary prior to preset execution).
+			ParentBuilder.GetMetasoundAsset().UpdateAndRegisterForExecution(FMetaSoundAssetRegistrationOptions { });
+
+			OutBuilder.SetIsGraphEditable(false);
+
+			// Apply root graph transform
+			FRebuildPresetRootGraph RebuildPresetRootGraph(ParentBuilder);
+			return RebuildPresetRootGraph.Transform(OutBuilder);
+		}
+		else
+		{
+			if (bLogFailure)
+			{
+				UE_LOG(LogMetaSound, Error, TEXT("The referenced node type must match the base MetaSound class when converting builder '%s' to a preset (ex. source preset must reference another source)"), *OutBuilder.GetDebugName());
+			}
+		}
+	}
+	else if (bLogFailure)
+	{
+		UE_LOG(LogMetaSound, Error, TEXT("Preset '%s' generation failed: Referenced class not set."), *OutBuilder.GetDebugName());
+	}
+
+	return false;
+}
+#endif // WITH_EDITORONLY_DATA
+
+TInstancedStruct<FMetaSoundFrontendDocumentConfigurationVertexMetadata> FMetaSoundFrontendPresetDocumentConfiguration::ConstructVertexMetadata(const FMetaSoundFrontendDocumentBuilder& Builder, const FGuid& NodeID) const
+{
+	// Presets only need ConfigVertexMetadata for inputs, so don't construct otherwise
+	if (PresetDocConfigPrivate::IsInputNode(Builder, NodeID))
+	{
+		return TInstancedStruct<FMetaSoundPresetVertexMetadata>::Make();
+	}
+
+	return { };
+}
+
+#if WITH_EDITOR
+void FMetaSoundFrontendPresetDocumentConfiguration::OnPropertyChanged(const FPropertyChangedEvent& InEvent, FMetaSoundFrontendDocumentBuilder& OutBuilder)
+{
+	ConfigureDocument(OutBuilder, false);
+}
+#endif // WITH_EDITOR
+
+void FMetaSoundPresetVertexMetadata::OnDefaultReset(const FMetaSoundFrontendDocumentBuilder& Builder)
+{
+	bInheritsDefault = true;
+}
+
+void FMetaSoundPresetVertexMetadata::OnDefaultUpdated(const FMetaSoundFrontendDocumentBuilder& Builder)
+{
+	bInheritsDefault = false;
+}
+#undef LOCTEXT_NAMESPACE
diff --git a/Engine/Plugins/Runtime/Metasound/Source/MetasoundFrontend/Private/MetasoundAssetBase.cpp b/Engine/Plugins/Runtime/Metasound/Source/MetasoundFrontend/Private/MetasoundAssetBase.cpp
--- a/Engine/Plugins/Runtime/Metasound/Source/MetasoundFrontend/Private/MetasoundAssetBase.cpp
+++ b/Engine/Plugins/Runtime/Metasound/Source/MetasoundFrontend/Private/MetasoundAssetBase.cpp
@@ -7,6 +7,7 @@
 #include "Algo/ForEach.h"
 #include "Algo/Transform.h"
 #include "Containers/Set.h"
+#include "DocumentConfigurations/MetasoundPresetDocumentConfiguration.h"
 #include "HAL/FileManager.h"
 #include "HAL/IConsoleManager.h"
 #include "IAudioParameterTransmitter.h"
@@ -210,7 +211,7 @@ namespace Metasound
 			TEXT("au.MetaSound.SampleRate"),
 			AssetBasePrivate::SampleRateOverride,
 			TEXT("Overrides the sample rate of metasounds. Negative values default to audio mixer sample rate.\n")
-			TEXT("Default: 0, Min: 8000, Max: 48000"),
+			TEXT("Default: 0, Min: 8000, Max: 96000"),
 			FConsoleVariableDelegate::CreateLambda([](IConsoleVariable* Var) { CVarMetaSoundSampleRateChanged.Broadcast(Var); }),
 			ECVF_Default);
 
@@ -379,7 +380,8 @@ void FMetasoundAssetBase::UpdateAndRegisterForExecution(const Metasound::Fronten
 		if (Owner->IsAsset())
 		{
 			DocBuilder = &IDocumentBuilderRegistry::GetChecked().FindOrBeginBuilding(Owner);
-			VersionDependencies(*DocBuilder, InRegistrationOptions.bAutoUpdateLogWarningOnDroppedConnection);
+			FAutoUpdateRootGraph AutoUpdateTransform(GetOwningAssetName(), InRegistrationOptions.bAutoUpdateLogWarningOnDroppedConnection);
+			bDidUpdate = AutoUpdateTransform.Transform(*DocBuilder);
 		}
 #else // !WITH_EDITORONLY_DATA
 		constexpr bool bDidUpdate = false;
@@ -470,7 +472,12 @@ void FMetasoundAssetBase::UpdateAndRegisterForSerialization(const Metasound::Fro
 	}
 
 	// Auto update must be done after all referenced asset classes are registered
-	bDidUpdate |= VersionDependencies(DocBuilder, /*bAutoUpdateLogWarningOnDroppedConnection=*/true);
+#if WITH_EDITORONLY_DATA
+	constexpr bool bAutoUpdateLogWarningOnDroppedConnection = true;
+	FAutoUpdateRootGraph AutoUpdateTransform(GetOwningAssetName(), bAutoUpdateLogWarningOnDroppedConnection);
+	bDidUpdate |= AutoUpdateTransform.Transform(DocBuilder);
+#endif // WITH_EDITORONLY_DATA
+
 #if WITH_EDITOR
 	if (bDidUpdate)
 	{
@@ -878,23 +885,6 @@ void FMetasoundAssetBase::RegisterAssetDependencies(const Metasound::Frontend::F
 }
 
 #if WITH_EDITORONLY_DATA
-void FMetasoundAssetBase::UpdateAndRegisterReferencesForSerialization(FName InCookPlatformName)
-{
-	using namespace Metasound::Frontend;
-	Metasound::Frontend::FMetaSoundAssetCookOptions CookOptions{ .bStripUnusedPages = !InCookPlatformName.IsNone() };
-
-	if (CookOptions.bStripUnusedPages)
-	{
-		const IDocumentBuilderRegistry& DocBuilderRegistry = IDocumentBuilderRegistry::GetChecked();
-PRAGMA_DISABLE_DEPRECATION_WARNINGS
-		CookOptions.PagesToTarget = DocBuilderRegistry.GetCookedTargetPages(InCookPlatformName);
-		CookOptions.PageOrder = DocBuilderRegistry.GetCookedPageOrder(InCookPlatformName);
-PRAGMA_ENABLE_DEPRECATION_WARNINGS
-	}
-
-	UpdateAndRegisterReferencesForSerialization(CookOptions);
-}
-
 void FMetasoundAssetBase::UpdateAndRegisterReferencesForSerialization(const Metasound::Frontend::FMetaSoundAssetCookOptions& InCookOptions)
 {
 	using namespace Metasound::Frontend;
@@ -928,20 +918,6 @@ bool FMetasoundAssetBase::TryUpdateInterfaceFromVersion(const FMetasoundFrontend
 	return false;
 }
 
-bool FMetasoundAssetBase::VersionDependencies(FMetaSoundFrontendDocumentBuilder& Builder, bool bInLogWarningsOnDroppedConnection)
-{
-	using namespace Metasound::Frontend;
-
-	bool bDocumentModified = false;
-
-#if WITH_EDITORONLY_DATA
-	FAutoUpdateRootGraph AutoUpdateTransform(GetOwningAssetName(), bInLogWarningsOnDroppedConnection);
-	bDocumentModified |= AutoUpdateTransform.Transform(Builder);
-#endif // WITH_EDITORONLY_DATA
-
-	return bDocumentModified;
-}
-
 FMetasoundFrontendVersion FMetasoundAssetBase::GetTargetInterfaceVersion(const FMetasoundFrontendVersion& InterfaceVersion) const
 {
 	using namespace Metasound::Frontend;
diff --git a/Engine/Plugins/Runtime/Metasound/Source/MetasoundFrontend/Private/MetasoundAssetManager.cpp b/Engine/Plugins/Runtime/Metasound/Source/MetasoundFrontend/Private/MetasoundAssetManager.cpp
--- a/Engine/Plugins/Runtime/Metasound/Source/MetasoundFrontend/Private/MetasoundAssetManager.cpp
+++ b/Engine/Plugins/Runtime/Metasound/Source/MetasoundFrontend/Private/MetasoundAssetManager.cpp
@@ -20,28 +20,6 @@ namespace Metasound::Frontend
 		static TUniquePtr<IMetaSoundAssetManager> Instance;
 	} // AssetManagerPrivate
 
-	namespace AssetTags
-	{
-		// Deprecated Tag Identifiers
-PRAGMA_DISABLE_DEPRECATION_WARNINGS
-		const FString ArrayDelim = TEXT(",");
-
-		const FName AssetClassID = "AssetClassID";
-
-#if WITH_EDITORONLY_DATA
-		const FName IsPreset = "bIsPreset";
-#endif // WITH_EDITORONLY_DATA
-
-		const FName RegistryVersionMajor = "RegistryVersionMajor";
-		const FName RegistryVersionMinor = "RegistryVersionMinor";
-
-#if WITH_EDITORONLY_DATA
-		const FName RegistryInputTypes = "RegistryInputTypes";
-		const FName RegistryOutputTypes = "RegistryOutputTypes";
-#endif // WITH_EDITORONLY_DATA
-PRAGMA_ENABLE_DEPRECATION_WARNINGS
-	} // namespace AssetTags
-
 	namespace AssetTagsPrivate
 	{
 		static const FLazyName AssetClassID = "AssetClassID";
diff --git a/Engine/Plugins/Runtime/Metasound/Source/MetasoundFrontend/Private/MetasoundDocumentInterface.cpp b/Engine/Plugins/Runtime/Metasound/Source/MetasoundFrontend/Private/MetasoundDocumentInterface.cpp
deleted file mode 100644
diff --git a/Engine/Plugins/Runtime/Metasound/Source/MetasoundFrontend/Private/MetasoundFrontendDocument.cpp b/Engine/Plugins/Runtime/Metasound/Source/MetasoundFrontend/Private/MetasoundFrontendDocument.cpp
--- a/Engine/Plugins/Runtime/Metasound/Source/MetasoundFrontend/Private/MetasoundFrontendDocument.cpp
+++ b/Engine/Plugins/Runtime/Metasound/Source/MetasoundFrontend/Private/MetasoundFrontendDocument.cpp
@@ -11,6 +11,7 @@
 #include "Logging/LogMacros.h"
 #include "MetasoundDocumentInterface.h"
 #include "MetasoundFrontend.h"
+#include "MetasoundFrontendDocumentBuilderRegistry.h"
 #include "MetasoundFrontendDocumentIdGenerator.h"
 #include "MetasoundFrontendDocumentVersioning.h"
 #include "MetasoundFrontendPages.h"
diff --git a/Engine/Plugins/Runtime/Metasound/Source/MetasoundFrontend/Private/MetasoundFrontendDocumentBuilder.cpp b/Engine/Plugins/Runtime/Metasound/Source/MetasoundFrontend/Private/MetasoundFrontendDocumentBuilder.cpp
--- a/Engine/Plugins/Runtime/Metasound/Source/MetasoundFrontend/Private/MetasoundFrontendDocumentBuilder.cpp
+++ b/Engine/Plugins/Runtime/Metasound/Source/MetasoundFrontend/Private/MetasoundFrontendDocumentBuilder.cpp
@@ -9,6 +9,7 @@
 #include "Algo/Sort.h"
 #include "Algo/Transform.h"
 #include "AudioParameter.h"
+#include "DocumentConfigurations/MetasoundPresetDocumentConfiguration.h"
 #include "Interfaces/MetasoundFrontendInterfaceBindingRegistry.h"
 #include "Interfaces/MetasoundFrontendInterfaceRegistry.h"
 #include "MetasoundAssetBase.h"
@@ -17,10 +18,12 @@
 #include "MetasoundFrontendController.h"
 #include "MetasoundFrontendDataTypeRegistry.h"
 #include "MetasoundFrontendDocumentCache.h"
+#include "MetasoundFrontendDocumentConfiguration.h"
 #include "MetasoundFrontendDocument.h"
 #include "MetasoundFrontendDocumentIdGenerator.h"
 #include "MetasoundFrontendDocumentModifyDelegates.h"
 #include "MetasoundFrontendDocumentVersioning.h"
+#include "MetasoundFrontendGraphLayer.h"
 #include "MetasoundFrontendNodeTemplateRegistry.h"
 #include "MetasoundFrontendNodeUpdateTransform.h"
 #include "MetasoundFrontendRegistries.h"
@@ -88,6 +91,20 @@ namespace Metasound::Frontend
 			return false;
 		}
 
+		FName GenerateUniqueName(const TSet<FName>& InExistingNames, const FString& InBaseName)
+		{
+			int32 PostFixInt = 0;
+			FString NewName = InBaseName;
+
+			while (InExistingNames.Contains(*NewName))
+			{
+				PostFixInt++;
+				NewName = FString::Format(TEXT("{0} {1}"), { InBaseName, PostFixInt });
+			}
+
+			return FName(*NewName);
+		}
+
 		bool NameContainsInterfaceNamespace(FName VertexName, FMetasoundFrontendInterface* OutInterface)
 		{
 			using namespace Metasound::Frontend;
@@ -765,14 +782,6 @@ UMetaSoundBuilderDocument::UMetaSoundBuilderDocument(const FObjectInitializer& O
 	Document.RootGraph.ID = FGuid::NewGuid();
 }
 
-UMetaSoundBuilderDocument& UMetaSoundBuilderDocument::Create(const UClass& InMetaSoundUClass)
-{
-	UMetaSoundBuilderDocument* DocObject = NewObject<UMetaSoundBuilderDocument>();
-	check(DocObject);
-	DocObject->MetaSoundUClass = &InMetaSoundUClass;
-	return *DocObject;
-}
-
 UMetaSoundBuilderDocument& UMetaSoundBuilderDocument::Create(const IMetaSoundDocumentInterface& InDocToCopy)
 {
 	UMetaSoundBuilderDocument* DocObject = NewObject<UMetaSoundBuilderDocument>();
@@ -801,6 +810,14 @@ const FMetasoundFrontendDocument& UMetaSoundBuilderDocument::GetConstDocument()
 	return Document;
 }
 
+#if WITH_EDITORONLY_DATA
+TConstStructView<FMetaSoundFrontendDocumentConfiguration> FMetaSoundFrontendDocumentBuilder::GetConstDocumentConfiguration() const
+{
+	const TInstancedStruct<FMetaSoundFrontendDocumentConfiguration>& Config = GetConstDocumentChecked().Configuration;
+	return Config;
+}
+#endif // WITH_EDITORONLY_DATA
+
 EMetasoundFrontendClassAccessFlags UMetaSoundBuilderDocument::GetDefaultAccessFlags() const
 {
 	return EMetasoundFrontendClassAccessFlags::None;
@@ -1023,7 +1040,7 @@ bool FMetaSoundFrontendDocumentBuilder::AddNamedEdges(const TSet<Metasound::Fron
 			checkf(Input, TEXT("Prior loop to check edge validity should protect against missing input vertex"));
 
 			const FMetasoundFrontendVertex* Output = FindNodeOutput(EdgeToAdd.NewEdge.FromNodeID, EdgeToAdd.NewEdge.FromVertexID, &PageID);
-			checkf(Input, TEXT("Prior loop to check edge validity should protect against missing output vertex"));
+			checkf(Output, TEXT("Prior loop to check edge validity should protect against missing output vertex"));
 
 			UE_LOG(LogMetaSound, Warning, TEXT("Connection between MetaSound output '%s' and input '%s' not added: Input already connected to '%s'."), *Output->Name.ToString(), *Input->Name.ToString(), *Output->Name.ToString());
 #endif // !NO_LOGGING
@@ -1264,6 +1281,11 @@ const FMetasoundFrontendNode* FMetaSoundFrontendDocumentBuilder::AddGraphInput(F
 
 		DocumentDelegates->InterfaceDelegates.OnInputAdded.Broadcast(NewIndex);
 #if WITH_EDITORONLY_DATA
+		if (FMetaSoundFrontendDocumentConfiguration* DocConfig = Document.Configuration.GetMutablePtr())
+		{
+			DocConfig->AddVertexMetadata(*this, NewNode->GetID());
+		}
+
 		Document.Metadata.ModifyContext.AddMemberIDModified(NewInput.NodeID);
 #endif // WITH_EDITORONLY_DATA
 
@@ -1373,7 +1395,13 @@ const FMetasoundFrontendNode* FMetaSoundFrontendDocumentBuilder::AddGraphOutput(
 		}
 
 		DocumentDelegates->InterfaceDelegates.OnOutputAdded.Broadcast(NewIndex);
+
 #if WITH_EDITORONLY_DATA
+		if (FMetaSoundFrontendDocumentConfiguration* DocConfig = Document.Configuration.GetMutablePtr())
+		{
+			DocConfig->AddVertexMetadata(*this, NewOutput.NodeID);
+		}
+
 		Document.Metadata.ModifyContext.AddMemberIDModified(NewOutput.NodeID);
 #endif // WITH_EDITORONLY_DATA
 	}
@@ -2325,6 +2353,7 @@ void FMetaSoundFrontendDocumentBuilder::ClearDocument(TSharedRef<Metasound::Fron
 
 	GraphClass.GetDefaultInterface().Inputs.Empty();
 	GraphClass.GetDefaultInterface().Outputs.Empty();
+	GraphClass.GetDefaultInterface().Environment.Empty();
 
 #if WITH_EDITOR
 	GraphClass.GetDefaultInterface().SetInputStyle({ });
@@ -2351,10 +2380,6 @@ void FMetaSoundFrontendDocumentBuilder::ClearDocument(TSharedRef<Metasound::Fron
 	});
 #endif // !WITH_EDITORONLY_DATA
 
-	GraphClass.GetDefaultInterface().Inputs.Empty();
-	GraphClass.GetDefaultInterface().Outputs.Empty();
-	GraphClass.GetDefaultInterface().Environment.Empty();
-
 	Doc.Interfaces.Empty();
 	Doc.Dependencies.Empty();
 
@@ -2480,6 +2505,40 @@ bool FMetaSoundFrontendDocumentBuilder::ConformGraphOutputNodeToClass(const FMet
 	return false;
 }
 
+#if WITH_EDITORONLY_DATA
+bool FMetaSoundFrontendDocumentBuilder::ConfigureDocument(bool bLogFailure)
+{
+	using namespace Metasound::Frontend;
+
+	if (bInDocConfigScope)
+	{
+		UE_LOG(LogMetaSound, Error, TEXT("Cannot call Configure Document: Configure already called higher in stack"));
+		return false;
+	}
+
+	bInDocConfigScope = true;
+
+	FMetasoundFrontendDocument& Document = GetDocumentChecked();
+	if (FMetaSoundFrontendDocumentConfiguration* DocConfig = Document.Configuration.GetMutablePtr())
+	{
+		if (DocConfig->ConfigureDocument(*this, bLogFailure))
+		{
+			GetDocumentInterfaceChecked().ConformObjectToDocument();
+			bInDocConfigScope = false;
+			return true;
+		}
+
+		TSharedRef<FDocumentModifyDelegates> ModifyDelegates = MakeShared<FDocumentModifyDelegates>(*DocumentDelegates);
+		ClearDocument(ModifyDelegates);
+	}
+
+	SetIsGraphEditable(true);
+
+	bInDocConfigScope = false;
+	return false;
+}
+#endif // WITH_EDITORONLY_DATA
+
 bool FMetaSoundFrontendDocumentBuilder::ContainsDependencyOfType(EMetasoundFrontendClassType ClassType) const
 {
 	return DocumentCache->ContainsDependencyOfType(ClassType);
@@ -2512,7 +2571,7 @@ bool FMetaSoundFrontendDocumentBuilder::ConvertFromPreset()
 		Style.bIsGraphEditable = true;
 #endif // WITH_EDITOR
 
-		DocumentDelegates->OnPresetStateChanged.Broadcast(FDocumentPresetStateChangedArgs { });
+		DocumentDelegates->OnDocumentConfigurationChanged.Broadcast(FDocumentConfigChangedArgs { });
 		return true;
 	}
 
@@ -2546,7 +2605,7 @@ bool FMetaSoundFrontendDocumentBuilder::ConvertToPreset(const FMetaSoundFrontend
 	if (RebuildPresetRootGraph.Transform(*this))
 	{
 		DocumentInterface->ConformObjectToDocument();
-		DocumentDelegates->OnPresetStateChanged.Broadcast(FDocumentPresetStateChangedArgs { });
+		DocumentDelegates->OnDocumentConfigurationChanged.Broadcast(FDocumentConfigChangedArgs { });
 		return true;
 	}
 
@@ -2712,26 +2771,6 @@ bool FMetaSoundFrontendDocumentBuilder::DiffAgainstRegistryInterface(const FGuid
 }
 #endif // WITH_EDITORONLY_DATA
 
-const FMetasoundFrontendNode* FMetaSoundFrontendDocumentBuilder::DuplicateGraphInput(const FMetasoundFrontendClassInput& InClassInput, const FName InName, const FGuid* InPageID)
-{
-	using namespace Metasound;
-
-	Frontend::FDocumentIDGenerator& IDGenerator = Frontend::FDocumentIDGenerator::Get();
-	const FMetasoundFrontendDocument& Doc = GetConstDocumentChecked();
-
-	const FGuid& PageID = InPageID ? *InPageID : BuildPageID;
-
-	FMetasoundFrontendClassInput ClassInput = InClassInput;
-	ClassInput.NodeID = IDGenerator.CreateNodeID(Doc);
-	ClassInput.VertexID = IDGenerator.CreateVertexID(Doc);
-#if WITH_EDITORONLY_DATA
-	ClassInput.Metadata.SetDisplayName(FText::GetEmpty());
-#endif // WITH_EDITORONLY_DATA
-	ClassInput.Name = InName;
-
-	return AddGraphInput(MoveTemp(ClassInput), &PageID);
-}
-
 const FMetasoundFrontendClassInput* FMetaSoundFrontendDocumentBuilder::DuplicateGraphInput(FName ExistingName, FName NewName)
 {
 	using namespace Metasound;
@@ -2764,26 +2803,6 @@ const FMetasoundFrontendClassInput* FMetaSoundFrontendDocumentBuilder::Duplicate
 	return FindGraphInput(NewName);
 }
 
-const FMetasoundFrontendNode* FMetaSoundFrontendDocumentBuilder::DuplicateGraphOutput(const FMetasoundFrontendClassOutput& InClassOutput, const FName InName, const FGuid* InPageID)
-{
-	using namespace Metasound::Frontend;
-
-	FDocumentIDGenerator& IDGenerator = FDocumentIDGenerator::Get();
-	const FMetasoundFrontendDocument& Doc = GetConstDocumentChecked();
-
-	const FGuid& PageID = InPageID ? *InPageID : BuildPageID;
-
-	FMetasoundFrontendClassOutput ClassOutput = InClassOutput;
-	ClassOutput.NodeID = IDGenerator.CreateNodeID(Doc);
-	ClassOutput.VertexID = IDGenerator.CreateVertexID(Doc);
-#if WITH_EDITORONLY_DATA
-	ClassOutput.Metadata.SetDisplayName(FText::GetEmpty());
-#endif // WITH_EDITORONLY_DATA
-	ClassOutput.Name = InName;
-
-	return AddGraphOutput(MoveTemp(ClassOutput), &PageID);
-}
-
 const FMetasoundFrontendClassOutput* FMetaSoundFrontendDocumentBuilder::DuplicateGraphOutput(FName ExistingName, FName NewName)
 {
 	using namespace Metasound::Frontend;
@@ -3604,6 +3623,35 @@ FMetasoundFrontendClassName FMetaSoundFrontendDocumentBuilder::GenerateNewClassN
 	return NewClassName;
 }
 
+FName FMetaSoundFrontendDocumentBuilder::GenerateUniqueInputName(const FString& InBaseName) const
+{
+	using namespace Metasound::Frontend;
+
+	TSet<FName> ExistingNames;
+	auto GetName = [&](const FMetasoundFrontendClass&, const FMetasoundFrontendNode& Node) { ExistingNames.Add(Node.Name); };
+	IterateNodesByClassType(GetName, EMetasoundFrontendClassType::Input);
+	return DocumentBuilderPrivate::GenerateUniqueName(ExistingNames, InBaseName);
+}
+
+FName FMetaSoundFrontendDocumentBuilder::GenerateUniqueOutputName(const FString& InBaseName) const
+{
+	using namespace Metasound::Frontend;
+
+	TSet<FName> ExistingNames;
+	auto GetName = [&](const FMetasoundFrontendClass&, const FMetasoundFrontendNode& Node) { ExistingNames.Add(Node.Name); };
+	IterateNodesByClassType(GetName, EMetasoundFrontendClassType::Output);
+	return DocumentBuilderPrivate::GenerateUniqueName(ExistingNames, InBaseName);
+}
+
+FName FMetaSoundFrontendDocumentBuilder::GenerateUniqueVariableName(const FString& InBaseName) const
+{
+	using namespace Metasound::Frontend;
+
+	TSet<FName> ExistingVariableNames;
+	Algo::Transform(FindConstBuildGraphChecked().Variables, ExistingVariableNames, [](const FMetasoundFrontendVariable& Var) { return Var.Name; });
+	return DocumentBuilderPrivate::GenerateUniqueName(ExistingVariableNames, InBaseName);
+}
+
 const FTopLevelAssetPath FMetaSoundFrontendDocumentBuilder::GetBuilderClassPath() const
 {
 	IMetaSoundDocumentInterface* Interface = DocumentInterface.GetInterface();
@@ -3631,13 +3679,6 @@ const FString FMetaSoundFrontendDocumentBuilder::GetDebugName() const
 	return MetaSoundObject.GetPathName();
 }
 
-const FMetasoundFrontendDocument& FMetaSoundFrontendDocumentBuilder::GetDocument() const
-{
-	const IMetaSoundDocumentInterface* Interface = DocumentInterface.GetInterface();
-	checkf(Interface, TEXT("Failed to return document; interface must always be valid while builder is operating on MetaSound UObject! Builder constructed with asset at %s"), *HintPath.ToString());
-	return Interface->GetConstDocument();
-}
-
 FMetasoundFrontendDocument& FMetaSoundFrontendDocumentBuilder::GetDocumentChecked() const
 {
 	return GetDocumentInterfaceChecked().GetDocument();
@@ -3648,13 +3689,6 @@ Metasound::Frontend::FDocumentModifyDelegates& FMetaSoundFrontendDocumentBuilder
 	return *DocumentDelegates;
 }
 
-const IMetaSoundDocumentInterface& FMetaSoundFrontendDocumentBuilder::GetDocumentInterface() const
-{
-	const IMetaSoundDocumentInterface* Interface = DocumentInterface.GetInterface();
-	checkf(Interface, TEXT("Failed to return document; interface must always be valid while builder is operating on MetaSound UObject! Builder constructed with asset at %s"), *HintPath.ToString());
-	return *Interface;
-}
-
 IMetaSoundDocumentInterface& FMetaSoundFrontendDocumentBuilder::GetDocumentInterfaceChecked() const
 {
 	IMetaSoundDocumentInterface* Interface = DocumentInterface.GetInterface();
@@ -3792,6 +3826,22 @@ const FMetasoundFrontendLiteral* FMetaSoundFrontendDocumentBuilder::GetGraphVari
 	return nullptr;
 }
 
+#if WITH_EDITORONLY_DATA
+bool FMetaSoundFrontendDocumentBuilder::GetIsGraphEditable(const FGuid* PageID) const
+{
+	bool bIsEditable = true;
+	const FMetasoundFrontendDocument& Document = GetConstDocumentChecked();
+	Document.RootGraph.IterateGraphPages([&bIsEditable, &PageID](const FMetasoundFrontendGraph& PageGraph)
+	{
+		if (!PageID || PageGraph.PageID == *PageID)
+		{
+			bIsEditable &= PageGraph.Style.bIsGraphEditable;
+		}
+	});
+	return bIsEditable;
+}
+#endif // WITH_EDITORONLY_DATA
+
 EMetasoundFrontendVertexAccessType FMetaSoundFrontendDocumentBuilder::GetNodeInputAccessType(const FGuid& InNodeID, const FGuid& InVertexID, const FGuid* InPageID) const
 {
 	using namespace Metasound::Frontend;
@@ -3890,10 +3940,8 @@ FText FMetaSoundFrontendDocumentBuilder::GetNodeInputDisplayName(const FGuid& In
 		case EMetasoundFrontendClassType::External:
 		{
 			auto MatchesName = [&VertexName](const FMetasoundFrontendClassVertex& OtherVertex) { return OtherVertex.Name == VertexName; };
-			const FMetasoundFrontendVertex* Vertex = nullptr;
-			const FMetasoundFrontendClassVertex* ClassVertex = nullptr;
-			ClassVertex = Class->GetInterfaceForNode(*Node).Inputs.FindByPredicate(MatchesName);
-			if (Vertex && ClassVertex)
+			const FMetasoundFrontendClassVertex* ClassVertex = Class->GetInterfaceForNode(*Node).Inputs.FindByPredicate(MatchesName);
+			if (ClassVertex)
 			{
 				FName Namespace, ParamName;
 				ClassVertex->SplitName(Namespace, ParamName);
@@ -3977,10 +4025,8 @@ FText FMetaSoundFrontendDocumentBuilder::GetNodeOutputDisplayName(const FGuid& I
 		case EMetasoundFrontendClassType::External:
 		{
 			auto PinMatchesClassVertex = [&VertexName](const FMetasoundFrontendClassVertex& OtherVertex) { return OtherVertex.Name == VertexName; };
-			const FMetasoundFrontendVertex* Vertex = nullptr;
-			const FMetasoundFrontendClassVertex* ClassVertex = nullptr;
-			ClassVertex = Class->GetInterfaceForNode(*Node).Inputs.FindByPredicate(PinMatchesClassVertex);
-			if (Vertex && ClassVertex)
+			const FMetasoundFrontendClassVertex* ClassVertex = Class->GetInterfaceForNode(*Node).Outputs.FindByPredicate(PinMatchesClassVertex);
+			if (ClassVertex)
 			{
 				FName Namespace, ParamName;
 				ClassVertex->SplitName(Namespace, ParamName);
@@ -4025,81 +4071,6 @@ FText FMetaSoundFrontendDocumentBuilder::GetNodeOutputDisplayName(const FGuid& I
 }
 #endif // WITH_EDITORONLY_DATA
 
-const FMetasoundFrontendLiteral* FMetaSoundFrontendDocumentBuilder::GetNodeInputClassDefault(const FGuid& InNodeID, const FGuid& InVertexID, const FGuid* InPageID) const
-{
-	using namespace Metasound;
-	using namespace Metasound::Frontend;
-
-	const FGuid& PageID = InPageID ? *InPageID : BuildPageID;
-	const Frontend::IDocumentGraphNodeCache& NodeCache = DocumentCache->GetNodeCache(PageID);
-	if (const int32* NodeIndex = NodeCache.FindNodeIndex(InNodeID))
-	{
-		const FMetasoundFrontendDocument& Document = GetConstDocumentChecked();
-		const FMetasoundFrontendNode& Node = Document.RootGraph.FindConstGraphChecked(PageID).Nodes[*NodeIndex];
-		auto IsVertexID = [&InVertexID](const FMetasoundFrontendVertex& Vertex) { return Vertex.VertexID == InVertexID; };
-		if (const FMetasoundFrontendVertex* Vertex = Node.Interface.Inputs.FindByPredicate(IsVertexID))
-		{
-			if (const FMetasoundFrontendClass* Class = DocumentCache->FindDependency(Node.ClassID))
-			{
-				const EMetasoundFrontendClassType ClassType = Class->Metadata.GetType();
-				const FMetasoundFrontendClassInterface& ClassInterface = Class->GetInterfaceForNode(Node);
-				switch (ClassType)
-				{
-					case EMetasoundFrontendClassType::Output:
-					{
-						const FMetasoundFrontendClassInput& ClassInput = ClassInterface.Inputs.Last();
-						return ClassInput.FindConstDefault(Frontend::DefaultPageID);
-					}
-
-					default:
-					{
-						auto IsClassInput = [VertexName = Vertex->Name](const FMetasoundFrontendClassInput& Input) { return Input.Name == VertexName; };
-						if (const FMetasoundFrontendClassInput* ClassInput = ClassInterface.Inputs.FindByPredicate(IsClassInput))
-						{
-							return ClassInput->FindConstDefault(Frontend::DefaultPageID);
-						}
-						static_assert(static_cast<uint32>(EMetasoundFrontendClassType::Invalid) == 10, "Potential missing case coverage for EMetasoundFrontendClassType "
-							"(default may not be sufficient for newly added class types)");
-					}
-					break;
-				}
-				static_assert(static_cast<uint32>(EMetasoundFrontendClassType::Invalid) == 10, "Potential missing case coverage for EMetasoundFrontendClassType");
-			}
-		}
-	}
-
-	return nullptr;
-}
-
-const FMetasoundFrontendLiteral* FMetaSoundFrontendDocumentBuilder::GetNodeInputDefault(const FGuid& InNodeID, const FGuid& InVertexID, const FGuid* InPageID) const
-{
-	using namespace Metasound::Frontend;
-
-	const FGuid& PageID = InPageID ? *InPageID : BuildPageID;
-	const IDocumentGraphNodeCache& NodeCache = DocumentCache->GetNodeCache(PageID);
-	if (const int32* NodeIndex = NodeCache.FindNodeIndex(InNodeID))
-	{
-		const FMetasoundFrontendGraph& Graph = GetConstDocumentChecked().RootGraph.FindConstGraphChecked(PageID);
-		const FMetasoundFrontendNode& Node = Graph.Nodes[*NodeIndex];
-
-		auto IsVertex = [&InVertexID](const FMetasoundFrontendVertex& Vertex) { return Vertex.VertexID == InVertexID; };
-		const int32 VertexIndex = Node.Interface.Inputs.IndexOfByPredicate(IsVertex);
-		if (VertexIndex != INDEX_NONE)
-		{
-			const FMetasoundFrontendVertex& NodeInput = Node.Interface.Inputs[VertexIndex];
-
-			auto IsLiteral = [&InVertexID](const FMetasoundFrontendVertexLiteral& Literal) { return Literal.VertexID == InVertexID; };
-			const int32 LiteralIndex = Node.InputLiterals.IndexOfByPredicate(IsLiteral);
-			if (LiteralIndex != INDEX_NONE)
-			{
-				return &Node.InputLiterals[LiteralIndex].Value;
-			}
-		}
-	}
-
-	return nullptr;
-}
-
 EMetasoundFrontendVertexAccessType FMetaSoundFrontendDocumentBuilder::GetNodeOutputAccessType(const FGuid& InNodeID, const FGuid& InVertexID, const FGuid* InPageID) const
 {
 	using namespace Metasound::Frontend;
@@ -4224,6 +4195,7 @@ void FMetaSoundFrontendDocumentBuilder::InitDocument(const FMetasoundFrontendDoc
 	// Initialize class using default data
 	else
 	{
+		// Initialize class using default data
 		if (Document.RootGraph.GetConstGraphPages().IsEmpty())
 		{
 			Document.RootGraph.InitDefaultGraphPage();
@@ -4234,6 +4206,7 @@ void FMetaSoundFrontendDocumentBuilder::InitDocument(const FMetasoundFrontendDoc
 		InitGraphClassMetadata(Document.RootGraph.Metadata, bResetVersion, InNewClassName);
 
 #if WITH_EDITORONLY_DATA
+		// Set default doc version Metadata
 		{
 			FMetasoundFrontendDocumentMetadata& DocMetadata = Document.Metadata;
 			DocMetadata.Version.Number = GetMaxDocumentVersion();
@@ -5257,6 +5230,13 @@ bool FMetaSoundFrontendDocumentBuilder::RemoveGraphInput(FName InputName, bool b
 
 		if (bNodesRemoved)
 		{
+#if WITH_EDITORONLY_DATA
+			if (FMetaSoundFrontendDocumentConfiguration* DocConfig = Document.Configuration.GetMutablePtr())
+			{
+				DocConfig->RemoveMetadata(NodeID);
+			}
+#endif // WITH_EDITORONLY_DATA
+
 			const int32 Index = *IndexPtr;
 			DocumentDelegates->InterfaceDelegates.OnRemovingInput.Broadcast(Index);
 
@@ -5559,11 +5539,6 @@ bool FMetaSoundFrontendDocumentBuilder::RemoveUnusedDependencies()
 	return bDidEdit;
 }
 
-bool FMetaSoundFrontendDocumentBuilder::RenameRootGraphClass(const FMetasoundFrontendClassName& InName)
-{
-	return false;
-}
-
 bool FMetaSoundFrontendDocumentBuilder::ReplaceDependency(const Metasound::Frontend::FNodeClassRegistryKey& OldNodeClassKey, const Metasound::Frontend::FNodeClassRegistryKey& NewNodeClassKey, TArray<FVertexNameAndType>* OutDisconnectedInputs, TArray<FVertexNameAndType>* OutDisconnectedOutputs)
 {
 	using namespace Metasound::Frontend;
@@ -5631,13 +5606,6 @@ bool FMetaSoundFrontendDocumentBuilder::ReplaceDependency(const Metasound::Front
 	return bSuccess;
 }
 
-void FMetaSoundFrontendDocumentBuilder::ReloadCache()
-{
-	using namespace Metasound::Frontend;
-
-	Reload(DocumentDelegates, true);
-}
-
 void FMetaSoundFrontendDocumentBuilder::ApplyNodeInstanceReplacementData(FMetasoundFrontendNode& InReplacementNode, FMetaSoundFrontendDocumentBuilder::FNodeInstanceReplacementData&& InInstanceData, TArray<FVertexNameAndType>* OutDisconnectedInputs, TArray<FVertexNameAndType>* OutDisconnectedOutputs)
 {
 	using namespace Metasound::Frontend;
@@ -5926,6 +5894,16 @@ bool FMetaSoundFrontendDocumentBuilder::SetBuildPageID(const FGuid& InBuildPageI
 
 	return false;
 }
+#endif // WITH_EDITORONLY_DATA
+
+#if WITH_EDITORONLY_DATA
+void FMetaSoundFrontendDocumentBuilder::SetDocumentConfiguration(TInstancedStruct<FMetaSoundFrontendDocumentConfiguration> InConfig)
+{
+	using namespace Metasound::Frontend;
+
+	GetDocumentChecked().Configuration = MoveTemp(InConfig);
+	DocumentDelegates->OnDocumentConfigurationChanged.Broadcast({ });
+}
 
 bool FMetaSoundFrontendDocumentBuilder::SetGraphInputAdvancedDisplay(const FName InputName, const bool InAdvancedDisplay)
 {
@@ -5946,6 +5924,18 @@ bool FMetaSoundFrontendDocumentBuilder::SetGraphInputAdvancedDisplay(const FName
 
 	return false;
 }
+
+void FMetaSoundFrontendDocumentBuilder::SetIsGraphEditable(bool bIsEditable, const FGuid* PageID) const
+{
+	FMetasoundFrontendDocument& Document = GetDocumentChecked();
+	Document.RootGraph.IterateGraphPages([bIsEditable, &PageID](FMetasoundFrontendGraph& PageGraph)
+	{
+		if (!PageID || PageGraph.PageID == *PageID)
+		{
+			PageGraph.Style.bIsGraphEditable = bIsEditable;
+		}
+	});
+}
 #endif // WITH_EDITORONLY_DATA
 
 bool FMetaSoundFrontendDocumentBuilder::SetGraphInputAccessType(FName InputName, EMetasoundFrontendVertexAccessType AccessType)
@@ -7346,18 +7336,6 @@ bool FMetaSoundFrontendDocumentBuilder::UpdateDependencyRegistryData(const TMap<
 #endif // WITH_EDITOR
 
 #if WITH_EDITORONLY_DATA
-bool FMetaSoundFrontendDocumentBuilder::VersionInterfaces()
-{
-	FMetasoundFrontendDocument& Document = GetDocumentChecked();
-	if (Document.RequiresInterfaceVersioning())
-	{
-		Document.VersionInterfaces();
-		return true;
-	}
-
-	return false;
-}
-
 FMetasoundFrontendDocument& FMetaSoundFrontendDocumentBuilder::IPropertyVersionTransform::GetDocumentUnsafe(const FMetaSoundFrontendDocumentBuilder& Builder)
 {
 	return Builder.GetDocumentChecked();
diff --git a/Engine/Plugins/Runtime/Metasound/Source/MetasoundFrontend/Private/MetasoundFrontendDocumentBuilderRegistry.cpp b/Engine/Plugins/Runtime/Metasound/Source/MetasoundFrontend/Private/MetasoundFrontendDocumentBuilderRegistry.cpp
new file mode 100644
--- /dev/null
+++ b/Engine/Plugins/Runtime/Metasound/Source/MetasoundFrontend/Private/MetasoundFrontendDocumentBuilderRegistry.cpp
@@ -0,0 +1,42 @@
+// Copyright Epic Games, Inc. All Rights Reserved.
+
+#include "MetasoundFrontendDocumentBuilderRegistry.h"
+
+#include "MetasoundFrontendDocumentIdGenerator.h"
+#include "MetasoundFrontendGraph.h"
+#include "MetasoundFrontendProxyDataCache.h"
+#include "MetasoundOperatorBuilder.h"
+#include "MetasoundOperatorBuilderSettings.h"
+
+
+namespace Metasound::Frontend
+{
+	namespace DocumentInterfacePrivate
+	{
+		TUniquePtr<IDocumentBuilderRegistry> Instance;
+	}
+
+	IDocumentBuilderRegistry* IDocumentBuilderRegistry::Get()
+	{
+		return DocumentInterfacePrivate::Instance.Get();
+	}
+
+	IDocumentBuilderRegistry& IDocumentBuilderRegistry::GetChecked()
+	{
+		check(DocumentInterfacePrivate::Instance.IsValid());
+		return *DocumentInterfacePrivate::Instance.Get();
+	}
+
+	void IDocumentBuilderRegistry::Deinitialize()
+	{
+		check(DocumentInterfacePrivate::Instance.IsValid());
+		DocumentInterfacePrivate::Instance.Reset();
+	}
+
+	void IDocumentBuilderRegistry::Initialize(TUniquePtr<IDocumentBuilderRegistry>&& InInstance)
+	{
+		check(InInstance.IsValid());
+		check(!DocumentInterfacePrivate::Instance.IsValid());
+		DocumentInterfacePrivate::Instance = MoveTemp(InInstance);
+	}
+} // namespace Metasound::Frontend
diff --git a/Engine/Plugins/Runtime/Metasound/Source/MetasoundFrontend/Private/MetasoundFrontendDocumentConfiguration.cpp b/Engine/Plugins/Runtime/Metasound/Source/MetasoundFrontend/Private/MetasoundFrontendDocumentConfiguration.cpp
new file mode 100644
--- /dev/null
+++ b/Engine/Plugins/Runtime/Metasound/Source/MetasoundFrontend/Private/MetasoundFrontendDocumentConfiguration.cpp
@@ -0,0 +1,58 @@
+// Copyright Epic Games, Inc. All Rights Reserved.
+
+#include "MetasoundFrontendDocumentConfiguration.h"
+
+#include "MetasoundDocumentInterface.h"
+#include "MetasoundFrontendDocument.h"
+#include "MetasoundFrontendDocumentBuilder.h"
+#include "MetasoundFrontendDocumentBuilderRegistry.h"
+
+
+#define LOCTEXT_NAMESPACE "MetaSound"
+
+FMetaSoundFrontendDocumentConfigurationVertexMetadata* FMetaSoundFrontendDocumentConfiguration::AddVertexMetadata(const FMetaSoundFrontendDocumentBuilder& Builder, const FGuid& NodeID)
+{
+	if (TInstancedStruct<FMetaSoundFrontendDocumentConfigurationVertexMetadata>* Entry = ConfigMetadata.Find(NodeID))
+	{
+		return Entry->GetMutablePtr();
+	}
+
+	TInstancedStruct<FMetaSoundFrontendDocumentConfigurationVertexMetadata> VertexMetadata = ConstructVertexMetadata(Builder, NodeID);
+	if (VertexMetadata.IsValid())
+	{
+		return ConfigMetadata.Add(NodeID, MoveTemp(VertexMetadata)).GetMutablePtr();
+	}
+
+	return nullptr;
+}
+
+#if WITH_EDITORONLY_DATA
+bool FMetaSoundFrontendDocumentConfiguration::ConfigureDocument(FMetaSoundFrontendDocumentBuilder& OutBuilder, bool bLogFailure)
+{
+	OutBuilder.SetIsGraphEditable(false);
+
+	return true;
+}
+#endif // WITH_EDITORONLY_DATA
+
+TInstancedStruct<FMetaSoundFrontendDocumentConfigurationVertexMetadata> FMetaSoundFrontendDocumentConfiguration::ConstructVertexMetadata(const FMetaSoundFrontendDocumentBuilder& Builder, const FGuid& NodeID) const
+{
+	return { };
+}
+
+bool FMetaSoundFrontendDocumentConfiguration::ContainsMetadata(const FGuid& NodeID) const
+{
+	return ConfigMetadata.Contains(NodeID);
+}
+
+#if WITH_EDITOR
+void FMetaSoundFrontendDocumentConfiguration::OnPropertyChanged(const FPropertyChangedEvent& InEvent, FMetaSoundFrontendDocumentBuilder& OutBuilder)
+{
+}
+#endif // WITH_EDITOR
+
+bool FMetaSoundFrontendDocumentConfiguration::RemoveMetadata(const FGuid& NodeID)
+{
+	return ConfigMetadata.Remove(NodeID) > 0;
+}
+#undef LOCTEXT_NAMESPACE
diff --git a/Engine/Plugins/Runtime/Metasound/Source/MetasoundFrontend/Private/MetasoundFrontendDocumentModifyDelegates.cpp b/Engine/Plugins/Runtime/Metasound/Source/MetasoundFrontend/Private/MetasoundFrontendDocumentModifyDelegates.cpp
--- a/Engine/Plugins/Runtime/Metasound/Source/MetasoundFrontend/Private/MetasoundFrontendDocumentModifyDelegates.cpp
+++ b/Engine/Plugins/Runtime/Metasound/Source/MetasoundFrontend/Private/MetasoundFrontendDocumentModifyDelegates.cpp
@@ -9,11 +9,16 @@
 namespace Metasound::Frontend
 {
 	FDocumentModifyDelegates::FDocumentModifyDelegates()
+PRAGMA_DISABLE_DEPRECATION_WARNINGS
+		: OnPresetStateChanged()
+PRAGMA_ENABLE_DEPRECATION_WARNINGS
 	{
 	}
 
 	FDocumentModifyDelegates::FDocumentModifyDelegates(const FMetasoundFrontendDocument& Document)
-		: FDocumentModifyDelegates()
+PRAGMA_DISABLE_DEPRECATION_WARNINGS
+		: OnPresetStateChanged()
+PRAGMA_ENABLE_DEPRECATION_WARNINGS
 	{
 		Document.RootGraph.IterateGraphPages([this](const FMetasoundFrontendGraph& Graph)
 		{
diff --git a/Engine/Plugins/Runtime/Metasound/Source/MetasoundFrontend/Private/MetasoundFrontendDocumentVersioning.cpp b/Engine/Plugins/Runtime/Metasound/Source/MetasoundFrontend/Private/MetasoundFrontendDocumentVersioning.cpp
--- a/Engine/Plugins/Runtime/Metasound/Source/MetasoundFrontend/Private/MetasoundFrontendDocumentVersioning.cpp
+++ b/Engine/Plugins/Runtime/Metasound/Source/MetasoundFrontend/Private/MetasoundFrontendDocumentVersioning.cpp
@@ -5,6 +5,7 @@
 #include "Algo/Transform.h"
 #include "Algo/Unique.h"
 #include "CoreGlobals.h"
+#include "DocumentConfigurations/MetasoundPresetDocumentConfiguration.h"
 #include "HAL/IConsoleManager.h"
 #include "Interfaces/MetasoundFrontendInterface.h"
 #include "Interfaces/MetasoundFrontendInterfaceRegistry.h"
diff --git a/Engine/Plugins/Runtime/Metasound/Source/MetasoundFrontend/Private/MetasoundFrontendGraphController.cpp b/Engine/Plugins/Runtime/Metasound/Source/MetasoundFrontend/Private/MetasoundFrontendGraphController.cpp
--- a/Engine/Plugins/Runtime/Metasound/Source/MetasoundFrontend/Private/MetasoundFrontendGraphController.cpp
+++ b/Engine/Plugins/Runtime/Metasound/Source/MetasoundFrontend/Private/MetasoundFrontendGraphController.cpp
@@ -781,7 +781,11 @@ namespace Metasound
 				Graph.Edges.Reset();
 				GraphClass->GetDefaultInterface().Inputs.Reset();
 				GraphClass->GetDefaultInterface().Outputs.Reset();
+
+PRAGMA_DISABLE_DEPRECATION_WARNINGS
 				GraphClass->PresetOptions.InputsInheritingDefault.Reset();
+PRAGMA_ENABLE_DEPRECATION_WARNINGS
+
 				OwningDocument->ClearInterfaceVersions();
 				OwningDocument->RemoveUnreferencedDependencies();
 
diff --git a/Engine/Plugins/Runtime/Metasound/Source/MetasoundFrontend/Private/MetasoundFrontendLiteral.cpp b/Engine/Plugins/Runtime/Metasound/Source/MetasoundFrontend/Private/MetasoundFrontendLiteral.cpp
--- a/Engine/Plugins/Runtime/Metasound/Source/MetasoundFrontend/Private/MetasoundFrontendLiteral.cpp
+++ b/Engine/Plugins/Runtime/Metasound/Source/MetasoundFrontend/Private/MetasoundFrontendLiteral.cpp
@@ -251,12 +251,15 @@ void FMetasoundFrontendLiteral::Set(FDefault InValue)
 	AsNumDefault = 1;
 }
 
-void FMetasoundFrontendLiteral::Set(const FDefaultArray& InValue)
+void FMetasoundFrontendLiteral::Set(FDefaultArray InValue)
 {
-	Empty();
+	AsBoolean.Empty();
+	AsInteger.Empty();
+	AsFloat.Empty();
+	AsString.Empty();
+	AsUObject.Empty();
 
 	Type = EMetasoundFrontendLiteralType::NoneArray;
-
 	AsNumDefault = FMath::Max(0, InValue.Num);
 }
 
@@ -269,13 +272,16 @@ void FMetasoundFrontendLiteral::Set(bool InValue)
 	AsBoolean.Add(InValue);
 }
 
-void FMetasoundFrontendLiteral::Set(const TArray<bool>& InValue)
+void FMetasoundFrontendLiteral::Set(TArray<bool> InValue)
 {
-	Empty();
+	AsNumDefault = 0;
+	AsInteger.Empty();
+	AsFloat.Empty();
+	AsString.Empty();
+	AsUObject.Empty();
 
 	Type = EMetasoundFrontendLiteralType::BooleanArray;
-
-	AsBoolean = InValue;
+	AsBoolean = MoveTemp(InValue);
 }
 
 void FMetasoundFrontendLiteral::Set(int32 InValue)
@@ -287,13 +293,16 @@ void FMetasoundFrontendLiteral::Set(int32 InValue)
 	AsInteger.Add(InValue);
 }
 
-void FMetasoundFrontendLiteral::Set(const TArray<int32>& InValue)
+void FMetasoundFrontendLiteral::Set(TArray<int32> InValue)
 {
-	Empty();
+	AsNumDefault = 0;
+	AsBoolean.Empty();
+	AsFloat.Empty();
+	AsString.Empty();
+	AsUObject.Empty();
 
 	Type = EMetasoundFrontendLiteralType::IntegerArray;
-
-	AsInteger = InValue;
+	AsInteger = MoveTemp(InValue);
 }
 
 void FMetasoundFrontendLiteral::Set(float InValue)
@@ -305,31 +314,40 @@ void FMetasoundFrontendLiteral::Set(float InValue)
 	AsFloat.Add(InValue);
 }
 
-void FMetasoundFrontendLiteral::Set(const TArray<float>& InValue)
+void FMetasoundFrontendLiteral::Set(TArray<float> InValue)
 {
-	Empty();
+	AsNumDefault = 0;
+	AsBoolean.Empty();
+	AsInteger.Empty();
+	AsString.Empty();
+	AsUObject.Empty();
 
 	Type = EMetasoundFrontendLiteralType::FloatArray;
-
-	AsFloat = InValue;
+	AsFloat = MoveTemp(InValue);
 }
 
-void FMetasoundFrontendLiteral::Set(const FString& InValue)
+void FMetasoundFrontendLiteral::Set(FString InValue)
 {
-	Empty();
+	AsNumDefault = 0;
+	AsBoolean.Empty();
+	AsInteger.Empty();
+	AsFloat.Empty();
+	AsUObject.Empty();
 
 	Type = EMetasoundFrontendLiteralType::String;
-
-	AsString.Add(InValue);
+	AsString = { MoveTemp(InValue) };
 }
 
-void FMetasoundFrontendLiteral::Set(const TArray<FString>& InValue)
+void FMetasoundFrontendLiteral::Set(TArray<FString> InValue)
 {
-	Empty();
+	AsNumDefault = 0;
+	AsBoolean.Empty();
+	AsInteger.Empty();
+	AsFloat.Empty();
+	AsUObject.Empty();
 
 	Type = EMetasoundFrontendLiteralType::StringArray;
-
-	AsString = InValue;
+	AsString = MoveTemp(InValue);
 }
 
 void FMetasoundFrontendLiteral::Set(UObject* InValue)
@@ -341,13 +359,16 @@ void FMetasoundFrontendLiteral::Set(UObject* InValue)
 	AsUObject.Add(InValue);
 }
 
-void FMetasoundFrontendLiteral::Set(const TArray<UObject*>& InValue)
+void FMetasoundFrontendLiteral::Set(TArray<UObject*> InValue)
 {
-	Empty();
+	AsNumDefault = 0;
+	AsBoolean.Empty();
+	AsInteger.Empty();
+	AsFloat.Empty();
+	AsString.Empty();
 
 	Type = EMetasoundFrontendLiteralType::UObjectArray;
-
-	AsUObject = InValue;
+	AsUObject = MoveTemp(InValue);
 }
 
 void FMetasoundFrontendLiteral::SetFromLiteral(const Metasound::FLiteral& InLiteral)
@@ -454,7 +475,7 @@ void FMetasoundFrontendLiteral::SetFromLiteral(const Metasound::FLiteral& InLite
 
 void FMetasoundFrontendLiteral::SetType(EMetasoundFrontendLiteralType InType)
 {
-	switch (Type)
+	switch (InType)
 	{
 		case EMetasoundFrontendLiteralType::NoneArray:
 		{
diff --git a/Engine/Plugins/Runtime/Metasound/Source/MetasoundFrontend/Private/MetasoundFrontendTransform.cpp b/Engine/Plugins/Runtime/Metasound/Source/MetasoundFrontend/Private/MetasoundFrontendTransform.cpp
--- a/Engine/Plugins/Runtime/Metasound/Source/MetasoundFrontend/Private/MetasoundFrontendTransform.cpp
+++ b/Engine/Plugins/Runtime/Metasound/Source/MetasoundFrontend/Private/MetasoundFrontendTransform.cpp
@@ -883,8 +883,9 @@ namespace Metasound
 				check(ReferencedNodeInputVertex);
 				
 				ReferencedNodeInputVertices.Add(ReferencedNodeInputVertex);
-				// If not editor, just make connection directly between input and output node
-				#if !WITH_EDITORONLY_DATA
+
+				// If not editor-only data, just make connection directly between input and output node
+#if !WITH_EDITORONLY_DATA
 				InOutBuilderToTransform.AddEdge(FMetasoundFrontendEdge
 				{
 					InputNodeID,
@@ -892,8 +893,9 @@ namespace Metasound
 					InReferencedNodeID,
 					ReferencedNodeInputVertex->VertexID
 				});
-				// If editor, add template node and connections
-				#else
+
+				// If editor-only data, add template node and connections
+#else // WITH_EDITORONLY_DATA
 				// template node takes on data type of concrete input node's output type
 				const FName DataType = InputNode->Interface.Outputs.Last().TypeName;
 				
@@ -919,7 +921,7 @@ namespace Metasound
 					InReferencedNodeID,
 					ReferencedNodeInputVertex->VertexID
 				});
-				#endif
+#endif  // WITH_EDITORONLY_DATA
 			}
 
 #if WITH_EDITOR
diff --git a/Engine/Plugins/Runtime/Metasound/Source/MetasoundFrontend/Public/DocumentConfigurations/MetasoundPresetDocumentConfiguration.h b/Engine/Plugins/Runtime/Metasound/Source/MetasoundFrontend/Public/DocumentConfigurations/MetasoundPresetDocumentConfiguration.h
new file mode 100644
--- /dev/null
+++ b/Engine/Plugins/Runtime/Metasound/Source/MetasoundFrontend/Public/DocumentConfigurations/MetasoundPresetDocumentConfiguration.h
@@ -0,0 +1,62 @@
+// Copyright Epic Games, Inc. All Rights Reserved.
+#pragma once
+
+#include "MetasoundDocumentInterface.h"
+#include "MetasoundFrontendDocumentConfiguration.h"
+#include "UObject/NameTypes.h"
+
+#include "MetasoundPresetDocumentConfiguration.generated.h"
+
+#define UE_API METASOUNDFRONTEND_API
+
+// Forward Declarations
+class UMetaSoundFrontendMemberMetadata;
+enum class EMetasoundFrontendClassType : uint8;
+
+
+USTRUCT(BlueprintType, DisplayName = "Preset")
+struct FMetaSoundFrontendPresetDocumentConfiguration : public FMetaSoundFrontendDocumentConfiguration
+{
+	GENERATED_BODY()
+
+	virtual ~FMetaSoundFrontendPresetDocumentConfiguration() = default;
+
+#if WITH_EDITORONLY_DATA
+	UE_EXPERIMENTAL(5.8, "Document Configuration is in active development and subject to change")
+	UE_API virtual bool ConfigureDocument(FMetaSoundFrontendDocumentBuilder& OutBuilder, bool bLogFailure) override;
+#endif // WITH_EDITORONLY_DATA
+
+#if WITH_EDITOR
+	UE_EXPERIMENTAL(5.8, "Document Configuration is in active development and subject to change")
+	UE_API virtual bool IsMemberEditingEnabled() const override { return true; }
+
+	UE_EXPERIMENTAL(5.8, "Document Configuration is in active development and subject to change")
+	UE_API virtual bool IsGraphEditorVisible() const override { return true; }
+
+	UE_EXPERIMENTAL(5.8, "Document Configuration is in active development and subject to change")
+	UE_API virtual void OnPropertyChanged(const FPropertyChangedEvent& InEvent, FMetaSoundFrontendDocumentBuilder& OutBuilder) override;
+#endif // WITH_EDITOR
+
+protected:
+	UE_EXPERIMENTAL(5.8, "Document Configuration is in active development and subject to change")
+	UE_API virtual TInstancedStruct<FMetaSoundFrontendDocumentConfigurationVertexMetadata> ConstructVertexMetadata(const FMetaSoundFrontendDocumentBuilder& Builder, const FGuid& NodeID) const override;
+
+public:
+	UPROPERTY(EditAnywhere, BlueprintReadOnly, Category = "Preset Options")
+	TScriptInterface<IMetaSoundDocumentInterface> Parent;
+};
+
+USTRUCT()
+struct FMetaSoundPresetVertexMetadata : public FMetaSoundFrontendDocumentConfigurationVertexMetadata
+{
+	GENERATED_BODY()
+
+	virtual ~FMetaSoundPresetVertexMetadata() = default;
+
+	UE_API virtual void OnDefaultReset(const FMetaSoundFrontendDocumentBuilder& Builder) override;
+	UE_API virtual void OnDefaultUpdated(const FMetaSoundFrontendDocumentBuilder& Builder) override;
+
+	UPROPERTY()
+	bool bInheritsDefault = true;
+};
+#undef UE_API
\ No newline at end of file
diff --git a/Engine/Plugins/Runtime/Metasound/Source/MetasoundFrontend/Public/MetasoundAssetBase.h b/Engine/Plugins/Runtime/Metasound/Source/MetasoundFrontend/Public/MetasoundAssetBase.h
--- a/Engine/Plugins/Runtime/Metasound/Source/MetasoundFrontend/Public/MetasoundAssetBase.h
+++ b/Engine/Plugins/Runtime/Metasound/Source/MetasoundFrontend/Public/MetasoundAssetBase.h
@@ -222,19 +222,12 @@ PRAGMA_ENABLE_DEPRECATION_WARNINGS
 
 private:
 #if WITH_EDITORONLY_DATA
-	UE_DEPRECATED(5.7, "Call version of function which accepts FMetaSoundAssetCookOptions.")
-	UE_API void UpdateAndRegisterReferencesForSerialization(FName CookPlatformName);
-
 	UE_API void UpdateAndRegisterReferencesForSerialization(const Metasound::Frontend::FMetaSoundAssetCookOptions& InCookOptions);
 #endif // WITH_EDITORONLY_DATA
 
 	// Checks if version is up-to-date. If so, returns true. If false, updates the interfaces within the given asset's document to the most recent version.
 	UE_API bool TryUpdateInterfaceFromVersion(const FMetasoundFrontendVersion& Version, FMetaSoundFrontendDocumentBuilder& Builder);
 
-	// Versions dependencies to most recent version where applicable. If asset is a preset, MetaSound is rebuilt to accommodate any referenced node class interface changes.
-	// Otherwise, automatically updates any nodes and respective dependent classes to accommodate changes to interfaces therein preserving edges/connections where possible.
-	UE_API bool VersionDependencies(FMetaSoundFrontendDocumentBuilder& Builder, bool bInLogWarningsOnDroppedConnection);
-
 	// Returns new interface to be versioned to from the given version. If no interface versioning is
 	// required, returns invalid interface version (no name and invalid version number).
 	FMetasoundFrontendVersion GetTargetInterfaceVersion(const FMetasoundFrontendVersion& InterfaceVersion) const;
diff --git a/Engine/Plugins/Runtime/Metasound/Source/MetasoundFrontend/Public/MetasoundAssetManager.h b/Engine/Plugins/Runtime/Metasound/Source/MetasoundFrontend/Public/MetasoundAssetManager.h
--- a/Engine/Plugins/Runtime/Metasound/Source/MetasoundFrontend/Public/MetasoundAssetManager.h
+++ b/Engine/Plugins/Runtime/Metasound/Source/MetasoundFrontend/Public/MetasoundAssetManager.h
@@ -5,6 +5,7 @@
 #include "MetasoundAssetManagerTransaction.h"
 #include "MetasoundDocumentInterface.h"
 #include "MetasoundFrontendDocument.h"
+#include "MetasoundFrontendDocumentBuilderRegistry.h"
 #include "MetasoundFrontendQuery.h"
 #include "MetasoundFrontendRegistryKey.h"
 #include "MetasoundVertex.h"
@@ -58,36 +59,6 @@ struct FMetaSoundDocumentInfo
 
 namespace Metasound::Frontend
 {
-	namespace AssetTags
-	{
-		UE_DEPRECATED(5.6, "AssetTags no longer public. Construct FMetaSoundAssetClassInfo from AssetData to parse tag fields.")
-		extern const FString UE_API ArrayDelim;
-
-#if WITH_EDITORONLY_DATA
-		UE_DEPRECATED(5.6, "AssetTags no longer public. Construct FMetaSoundAssetClassInfo from AssetData to parse tag fields.")
-		extern const FName UE_API IsPreset;
-#endif // WITH_EDITORONLY_DATA
-
-
-		UE_DEPRECATED(5.6, "AssetTags no longer public. Construct FMetaSoundAssetClassInfo from AssetData to parse tag fields.")
-		extern const FName UE_API AssetClassID;
-
-		UE_DEPRECATED(5.6, "AssetTags no longer public. Construct FMetaSoundAssetClassInfo from AssetData to parse tag fields.")
-		extern const FName UE_API RegistryVersionMajor;
-
-		UE_DEPRECATED(5.6, "AssetTags no longer public. Construct FMetaSoundAssetClassInfo from AssetData to parse tag fields.")
-		extern const FName UE_API RegistryVersionMinor;
-
-#if WITH_EDITORONLY_DATA
-		UE_DEPRECATED(5.6, "AssetTags no longer public. Construct FMetaSoundAssetClassInfo from AssetData to parse tag fields.")
-		extern const FName UE_API RegistryInputTypes;
-
-		UE_DEPRECATED(5.6, "AssetTags no longer public. Construct FMetaSoundAssetClassInfo from AssetData to parse tag fields.")
-		extern const FName UE_API RegistryOutputTypes;
-#endif // WITH_EDITORONLY_DATA
-	} // namespace AssetTags
-
-
 	struct FMetaSoundAssetRegistrationOptions
 	{
 		// If true, forces a re-register of this class (and all class dependencies
diff --git a/Engine/Plugins/Runtime/Metasound/Source/MetasoundFrontend/Public/MetasoundDocumentInterface.h b/Engine/Plugins/Runtime/Metasound/Source/MetasoundFrontend/Public/MetasoundDocumentInterface.h
--- a/Engine/Plugins/Runtime/Metasound/Source/MetasoundFrontend/Public/MetasoundDocumentInterface.h
+++ b/Engine/Plugins/Runtime/Metasound/Source/MetasoundFrontend/Public/MetasoundDocumentInterface.h
@@ -12,13 +12,12 @@
 
 #include "MetasoundDocumentInterface.generated.h"
 
-#define UE_API METASOUNDFRONTEND_API
 
+#define UE_API METASOUNDFRONTEND_API
 
 // Forward Declarations
 struct FMetaSoundFrontendDocumentBuilder;
 
-
 // UInterface for all MetaSound UClasses that implement a MetaSound document
 // as a means for accessing data via code, scripting, execution, or node
 // class generation.
@@ -67,76 +66,9 @@ private:
 	friend struct FMetaSoundFrontendDocumentBuilder;
 };
 
+// Forward declaration for building/backwards compat. Now found in separate header: MetasoundFrontendDocumentBuilderRegistry.h.
 namespace Metasound::Frontend
 {
-	class IDocumentBuilderRegistry
-	{
-	public:
-		virtual ~IDocumentBuilderRegistry() = default;
-
-#if WITH_EDITORONLY_DATA
-		// Given the provided builder, removes paged data within the associated document for a cooked build.
-		// This function removes graphs and input defaults which are not to ever be used by a given cook
-		// platform, allowing users to optimize away data and scale the amount of memory required for
-		// initial load of input UObjects and graph topology, which can also positively effect runtime
-		// performance as well, etc. Returns true if builder modified the document, false if not.
-		UE_DEPRECATED(5.7, "Use StripUnusedPages(...) instead.")
-		virtual bool CookPages(FName CookPlatformName, FMetaSoundFrontendDocumentBuilder& Builder) const = 0;
-#endif // WITH_EDITORONLY_DATA
-
-		virtual FMetaSoundFrontendDocumentBuilder* FindBuilder(TScriptInterface<IMetaSoundDocumentInterface> MetaSound) const = 0;
-		virtual FMetaSoundFrontendDocumentBuilder* FindBuilder(const FMetasoundFrontendClassName& ClassName, const FTopLevelAssetPath& AssetPath) const = 0;
-		virtual FMetaSoundFrontendDocumentBuilder* FindOutermostBuilder(const UObject& InSubObject) const = 0;
-
-#if WITH_EDITORONLY_DATA
-		// Find the existing builder for the given MetaSound, or optionally begin building by attaching a new builder.  Only available
-		// in builds with editor only data as building serialized assets (which may have template nodes, cooked builds do not) is only
-		// supported when editor data is loaded. Creating transient builders can simply be done by passing a new MetaSound asset to
-		// a FMetaSoundFrontendDocumentBuilder constructor, or this registry's implementation may supply its own create call for tracking
-		// and reuse purposes.
-		virtual FMetaSoundFrontendDocumentBuilder& FindOrBeginBuilding(TScriptInterface<IMetaSoundDocumentInterface> MetaSound) = 0;
-#endif // WITH_EDITORONLY_DATA
-
-		// Removes builder from registry, clearing any cached builder state. (Optionally) forces unregistration from the Frontend Node Class Registry
-		// (If the builder has outstanding transactions, unregistration from the Node Class Registry will occur regardless).
-		virtual bool FinishBuilding(const FMetasoundFrontendClassName& InClassName, bool bForceUnregisterNodeClass = false) const = 0;
-		virtual bool FinishBuilding(const FMetasoundFrontendClassName& InClassName, const FTopLevelAssetPath& InAssetPath, bool bForceUnregisterNodeClass = false) const = 0;
-
-		// Reloads the given builder, maintaining all modify delegate subscriptions. Returns true if builder was found and reloaded,
-		// false if not found.
-		virtual bool ReloadBuilder(const FMetasoundFrontendClassName& InClassName) const = 0;
-
-		// Given the provided document and its respective pages, provides the PageID to be used for runtime IGraph and proxy generation.
-		// Returns true if valid page was found, false if not.
-		UE_DEPRECATED(5.7, "This is no longer supported. FindPreferredPage(...)")
-		virtual FGuid ResolveTargetPageID(const FMetasoundFrontendGraphClass& InGraphClass) const = 0;
-
-		// Given the provided input and its respective paged default values, provides the default PageID to be used for runtime IGraph and proxy generation.
-		// Returns true if valid page was found, false if not.
-		UE_DEPRECATED(5.7, "This is no longer supported. FindPreferredPage(...)")
-		virtual FGuid ResolveTargetPageID(const FMetasoundFrontendClassInput& InClassInput) const = 0;
-
-		// Given the provided array of default values, provides the default PageID to be used for runtime IGraph and proxy generation.
-		// Returns true if valid page was found, false if not.
-		UE_DEPRECATED(5.7, "This is no longer supported. FindPreferredPage(...)")
-		virtual FGuid ResolveTargetPageID(const TArray<FMetasoundFrontendClassInputDefault>& Defaults) const = 0;
-
-		UE_DEPRECATED(5.7, "This is only used to support existing deprecated functionality")
-		virtual TArrayView<const FGuid> GetPageOrder() const = 0;
-
-#if WITH_EDITORONLY_DATA
-		UE_DEPRECATED(5.7, "This is only used to support existing deprecated functionality")
-		virtual TArray<FGuid> GetCookedTargetPages(FName InPlatformName) const = 0;
-
-		UE_DEPRECATED(5.7, "This is only used to support existing deprecated functionality")
-		virtual TArray<FGuid> GetCookedPageOrder(FName InPlatformName) const = 0;
-#endif // WITH_EDITORONLY_DATA
-
-		static UE_API IDocumentBuilderRegistry* Get();
-		static UE_API IDocumentBuilderRegistry& GetChecked();
-		static UE_API void Deinitialize();
-		static UE_API void Initialize(TUniquePtr<IDocumentBuilderRegistry>&& InBuilderRegistry);
-	};
+	class IDocumentBuilderRegistry;
 } // namespace Metasound::Frontend
-
 #undef UE_API
diff --git a/Engine/Plugins/Runtime/Metasound/Source/MetasoundFrontend/Public/MetasoundFrontendDocument.h b/Engine/Plugins/Runtime/Metasound/Source/MetasoundFrontend/Public/MetasoundFrontendDocument.h
--- a/Engine/Plugins/Runtime/Metasound/Source/MetasoundFrontend/Public/MetasoundFrontendDocument.h
+++ b/Engine/Plugins/Runtime/Metasound/Source/MetasoundFrontend/Public/MetasoundFrontendDocument.h
@@ -8,6 +8,7 @@
 #include "IAudioParameterInterfaceRegistry.h"
 #include "Internationalization/Text.h"
 #include "MetasoundAccessPtr.h"
+#include "MetasoundFrontendDocumentConfiguration.h"
 #include "MetasoundFrontendLiteral.h"
 #include "MetasoundOperatorData.h"
 #include "MetasoundNodeInterface.h"
@@ -2005,18 +2006,6 @@ class UMetaSoundFrontendMemberMetadata : public UObject
 	GENERATED_BODY()
 
 public:
-	UE_DEPRECATED(5.5, "Implementation moved to child editor class instead of compiled out (not required by Frontend representation")
-	virtual void ForceRefresh() { }
-
-	UE_DEPRECATED(5.5, "Default is no longer required to be stored or represented in metadata and may differ in paged or non-paged implementation")
-	FMetasoundFrontendLiteral GetDefault() const { return FMetasoundFrontendLiteral(); }
-
-	UE_DEPRECATED(5.5, "Implementation moved to child editor class instead of compiled out (not required by Frontend representation")
-	virtual EMetasoundFrontendLiteralType GetLiteralType() const { return EMetasoundFrontendLiteralType::None; }
-
-	UE_DEPRECATED(5.5, "Default is no longer required to be stored or represented in metadata and may differ in paged or non-paged implementation")
-	virtual void SetFromLiteral(const FMetasoundFrontendLiteral& InLiteral, const FGuid& InPageID = Metasound::Frontend::DefaultPageID) { }
-
 #if WITH_EDITORONLY_DATA
 	UPROPERTY()
 	FGuid MemberID;
@@ -2072,6 +2061,11 @@ public:
 	UPROPERTY()
 	TArray<FMetasoundFrontendClass> Dependencies;
 
+#if WITH_EDITORONLY_DATA
+	UPROPERTY()
+	TInstancedStruct<FMetaSoundFrontendDocumentConfiguration> Configuration;
+#endif // WITH_EDITORONLY_DATA
+
 	uint32 GetNextIdCounter() const
 	{
 		return IdCounter++;
diff --git a/Engine/Plugins/Runtime/Metasound/Source/MetasoundFrontend/Public/MetasoundFrontendDocumentBuilder.h b/Engine/Plugins/Runtime/Metasound/Source/MetasoundFrontend/Public/MetasoundFrontendDocumentBuilder.h
--- a/Engine/Plugins/Runtime/Metasound/Source/MetasoundFrontend/Public/MetasoundFrontendDocumentBuilder.h
+++ b/Engine/Plugins/Runtime/Metasound/Source/MetasoundFrontend/Public/MetasoundFrontendDocumentBuilder.h
@@ -8,13 +8,16 @@
 #include "MetasoundDocumentInterface.h"
 #include "MetasoundFrontendDocument.h"
 #include "MetasoundFrontendDocumentModifyDelegates.h"
+#include "MetasoundFrontendGraphLayer.h"
 #include "MetasoundFrontendRegistries.h"
 #include "MetasoundFrontendTransform.h"
 #include "MetasoundLiteral.h"
 #include "MetasoundOperatorSettings.h"
 #include "MetasoundVertex.h"
 #include "StructUtils/InstancedStruct.h"
+#include "StructUtils/StructView.h"
 #include "Templates/Function.h"
+#include "UObject/Class.h"
 #include "UObject/ScriptInterface.h"
 
 #include "MetasoundFrontendDocumentBuilder.generated.h"
@@ -23,6 +26,7 @@
 
 
 // Forward Declarations
+struct FMetaSoundFrontendDocumentConfiguration;
 class FMetasoundAssetBase;
 
 
@@ -123,9 +127,6 @@ class UMetaSoundBuilderDocument : public UObject, public IMetaSoundDocumentInter
 public:
 	UE_API UMetaSoundBuilderDocument(const FObjectInitializer& ObjectInitializer);
 
-	UE_DEPRECATED(5.5, "Use overload supplying MetaSound to copy (builder documents no longer supported for cases outside of cloned document registration.")
-	static UE_API UMetaSoundBuilderDocument& Create(const UClass& InBuilderClass);
-
 	// Create and return a valid builder document which copies the provided interface's document & class
 	static UE_API UMetaSoundBuilderDocument& Create(const IMetaSoundDocumentInterface& InDocToCopy);
 
@@ -152,6 +153,17 @@ public:
 private:
 	UE_API virtual FMetasoundFrontendDocument& GetDocument() override;
 
+#if WITH_EDITORONLY_DATA
+	// Convenience template function for attempting to get the underlying pointer of the set configuration.
+	// If the provided configuration type is not that which the current document is set to or no configuration
+	// is set, returns null.
+	template<typename TConfigType>
+	TConfigType* GetDocumentConfigurationAs()
+	{
+		return GetDocument().Configuration.GetMutablePtr<TConfigType>();
+	}
+#endif // WITH_EDITORONLY_DATA
+
 	UE_API virtual void OnBeginActiveBuilder() override;
 	UE_API virtual void OnFinishActiveBuilder() override;
 
@@ -266,11 +278,11 @@ PRAGMA_ENABLE_EXPERIMENTAL_WARNINGS
 	// so external delegates must be relinked if desired.
 	UE_API void ClearDocument(TSharedRef<Metasound::Frontend::FDocumentModifyDelegates> ModifyDelegates);
 
-	UE_DEPRECATED(5.5, "Use ClearDocument instead")
-	void ClearGraph() {  }
-
 #if WITH_EDITORONLY_DATA
 	UE_API bool ClearMemberMetadata(const FGuid& InMemberID);
+
+	UE_EXPERIMENTAL(5.8, "Configurator is in active development and subject to change")
+	UE_API bool ConfigureDocument(bool bLogFailure = false);
 #endif // WITH_EDITORONLY_DATA
 
 	UE_API bool ContainsDependencyOfType(EMetasoundFrontendClassType ClassType) const;
@@ -297,12 +309,6 @@ PRAGMA_ENABLE_EXPERIMENTAL_WARNINGS
 	UE_API const FMetasoundFrontendClassOutput* DuplicateGraphOutput(FName ExistingName, FName NewName);
 	UE_API const FMetasoundFrontendVariable* DuplicateGraphVariable(FName ExistingName, FName NewName, const FGuid* InPageID = nullptr);
 
-	UE_DEPRECATED(5.6, "Use the duplicate overload that supplies existing name and new name and returns input")
-	UE_API const FMetasoundFrontendNode* DuplicateGraphInput(const FMetasoundFrontendClassInput& InClassInput, const FName InName, const FGuid* InPageID = nullptr);
-
-	UE_DEPRECATED(5.6, "Use the duplicate overload that supplies existing name and new name and returns output")
-	UE_API const FMetasoundFrontendNode* DuplicateGraphOutput(const FMetasoundFrontendClassOutput& InClassOutput, const FName InName, const FGuid* InPageID = nullptr);
-
 	TConstStructView<FMetasoundFrontendClassInterface> FindClassInterfaceOverride(const FGuid& InNodeID, const FGuid* InPageID = nullptr) const;
 
 #if WITH_EDITORONLY_DATA
@@ -382,12 +388,25 @@ PRAGMA_ENABLE_EXPERIMENTAL_WARNINGS
 	UE_API int32 FindPageIndex(const FGuid& InPageID) const;
 
 	UE_API const FMetasoundFrontendDocument& GetConstDocumentChecked() const;
+
+#if WITH_EDITORONLY_DATA
+	// Returns struct view of document configuration if set.
+	UE_EXPERIMENTAL(5.8, "Configurator is in active development and subject to change")
+	UE_API TConstStructView<FMetaSoundFrontendDocumentConfiguration> GetConstDocumentConfiguration() const;
+
+	// Convenience template function for attempting to get the underlying const pointer of the set configuration.
+	// If the provided configuration type is not that which the current document is set to or no configuration
+	// is set, returns null.
+	template<typename TConfigType>
+	const TConfigType* GetConstDocumentConfigurationAs() const
+	{
+		return GetConstDocumentConfiguration().GetPtr<TConfigType>();
+	}
+#endif // WITH_EDITORONLY_DATA
+
 	UE_API const IMetaSoundDocumentInterface& GetConstDocumentInterfaceChecked() const;
 	UE_API const FString GetDebugName() const;
 
-	UE_DEPRECATED(5.5, "Use GetConstDocumentChecked() instead")
-	UE_API const FMetasoundFrontendDocument& GetDocument() const;
-
 	// The graph ID used when requests are made to mutate specific paged graph topology (ex. adding or removing nodes or edges)
 	UE_API const FGuid& GetBuildPageID() const;
 
@@ -418,10 +437,11 @@ PRAGMA_ENABLE_EXPERIMENTAL_WARNINGS
 	// asset validation errors.
 	UE_API FMetasoundFrontendClassName GenerateNewClassName();
 
-	UE_API Metasound::Frontend::FDocumentModifyDelegates& GetDocumentDelegates();
+	UE_API FName GenerateUniqueInputName(const FString& InBaseName) const;
+	UE_API FName GenerateUniqueOutputName(const FString& InBaseName) const;
+	UE_API FName GenerateUniqueVariableName(const FString& InBaseName) const;
 
-	UE_DEPRECATED(5.5, "Use GetConstDocumentInterfaceChecked instead")
-	UE_API const IMetaSoundDocumentInterface& GetDocumentInterface() const;
+	UE_API Metasound::Frontend::FDocumentModifyDelegates& GetDocumentDelegates();
 
 	// Path for document object provided at construction time. Cached on builder as a useful means of debugging
 	// and enables weak reference removal from the builder registry should the object be mid-destruction.
@@ -441,9 +461,6 @@ PRAGMA_ENABLE_EXPERIMENTAL_WARNINGS
 
 	UE_API EMetasoundFrontendVertexAccessType GetNodeInputAccessType(const FGuid& InNodeID, const FGuid& InVertexID, const FGuid* InPageID = nullptr) const;
 
-	UE_DEPRECATED(5.5, "Use FindNodeInputClass overloads instead and use GetDefaults() on result (now supports page values)")
-	UE_API const FMetasoundFrontendLiteral* GetNodeInputClassDefault(const FGuid& InNodeID, const FGuid& InVertexID, const FGuid* InPageID = nullptr) const;
-
 #if WITH_EDITORONLY_DATA
 	// Returns friendly name of the given node input vertex. Unlike a straight look-up of the vertex metadata value, this
 	// tries and determines the best display name based on the parent node's class, if no display name is given
@@ -456,9 +473,6 @@ PRAGMA_ENABLE_EXPERIMENTAL_WARNINGS
 	UE_API FText GetNodeOutputDisplayName(const FGuid& InNodeID, const FName VertexName, const FGuid* InPageID = nullptr) const;
 #endif // WITH_EDITORONLY_DATA
 
-	UE_DEPRECATED(5.5, "Use FindNodeInputDefault and returned struct Value member instead")
-	UE_API const FMetasoundFrontendLiteral* GetNodeInputDefault(const FGuid& InNodeID, const FGuid& InVertexID, const FGuid* InPageID = nullptr) const;
-
 	UE_API EMetasoundFrontendVertexAccessType GetNodeOutputAccessType(const FGuid& InNodeID, const FGuid& InVertexID, const FGuid* InPageID = nullptr) const;
 
 #if WITH_EDITORONLY_DATA
@@ -471,6 +485,11 @@ PRAGMA_ENABLE_EXPERIMENTAL_WARNINGS
 	// Returns the default value set for the variable with the given name on the given page.
 	UE_API const FMetasoundFrontendLiteral* GetGraphVariableDefault(FName InputName, const FGuid* InPageID = nullptr) const;
 
+#if WITH_EDITORONLY_DATA
+	// If PageID set, returns whether or not the given page is editable.  If no page ID is set, returns whether graph with set build ID is editable.
+	UE_API bool GetIsGraphEditable(const FGuid* InPageID = nullptr) const;
+#endif // WITH_EDITORONLY_DATA
+
 	// Initializes the builder's document, using the (optional) provided document template, (optional) class name, and (optionally) whether or not to reset the existing class version.
 	UE_API void InitDocument(const FMetasoundFrontendDocument* InDocumentTemplate = nullptr, const FMetasoundFrontendClassName* InNewClassName = nullptr, bool bResetVersion = true);
 
@@ -480,10 +499,22 @@ PRAGMA_ENABLE_EXPERIMENTAL_WARNINGS
 
 	UE_API void InitNodeLocations();
 
-	UE_DEPRECATED(5.5, "Use invalidate overload that is provided new version of modify delegates")
-	void InvalidateCache() { }
-
 	UE_API bool IsDependencyReferenced(const FGuid& InClassID) const;
+
+#if WITH_EDITORONLY_DATA
+	// Determines if the builder's document configuration is of the given type
+	template<typename TConfigType>
+	const bool IsDocumentConfiguredBy() const
+	{
+		if (const UScriptStruct* Struct = GetConstDocumentChecked().Configuration.GetScriptStruct())
+		{
+			return Struct->IsChildOf(TConfigType::StaticStruct());
+		}
+
+		return false;
+	}
+#endif // WITH_EDITORONLY_DATA
+
 	UE_API bool IsNodeInputConnected(const FGuid& InNodeID, const FGuid& InVertexID, const FGuid* InPageID = nullptr) const;
 	UE_API bool IsNodeInputConnectionUserModifiable(const FGuid& InNodeID, const FGuid& InVertexID, const FGuid* InPageID = nullptr) const;
 	UE_API bool IsNodeOutputConnected(const FGuid& InNodeID, const FGuid& InVertexID, const FGuid* InPageID = nullptr) const;
@@ -511,11 +542,6 @@ PRAGMA_ENABLE_EXPERIMENTAL_WARNINGS
 
 	UE_API bool ModifyInterfaces(Metasound::Frontend::FModifyInterfaceOptions&& InOptions);
 
-	UE_DEPRECATED(5.5,
-		"Cache invalidation may require new copy of delegates. In addition, re-priming is discouraged. "
-		"To enforce this, new recommended pattern is to construct a new builder instead")
-	UE_API void ReloadCache();
-
 	// Removes all dependencies with the given ClassID. Removes any nodes (and corresponding edges) remaining in any MetaSound paged graphs.
 	UE_API bool RemoveDependency(const FGuid& InClassID);
 
@@ -563,9 +589,6 @@ PRAGMA_ENABLE_EXPERIMENTAL_WARNINGS
 
 	UE_API bool RemoveNodeInputDefault(const FGuid& InNodeID, const FGuid& InVertexID, const FGuid* InPageID = nullptr);
 	UE_API bool RemoveUnusedDependencies();
-
-	UE_DEPRECATED(5.5, "Use GenerateNewClassName instead")
-	UE_API bool RenameRootGraphClass(const FMetasoundFrontendClassName& InName);
 	
 	using FVertexNameAndType = TTuple<FName, FName>;
 	// Replaces existing dependency with a new one. Removes all nodes using this class, and readds nodes with new class, 
@@ -594,6 +617,14 @@ PRAGMA_ENABLE_EXPERIMENTAL_WARNINGS
 	// API (exposed as a mechanism for mutating via controllers while deprecating.  Option will be removed
 	// in a future build).
 	UE_API bool SetBuildPageID(const FGuid& InBuildPageID, bool bBroadcastDelegate = true);
+#endif // WITH_EDITORONLY_DATA
+
+#if WITH_EDITORONLY_DATA
+	UE_API void SetDocumentConfiguration(TInstancedStruct<FMetaSoundFrontendDocumentConfiguration> InConfig);
+
+	// If page ID is provided, updates the graph's editable flag to the given value. If no page ID is provided,
+	// sets *all* paged graph fields to the given value.
+	UE_API void SetIsGraphEditable(bool bIsEditable, const FGuid* PageID = nullptr) const;
 
 	// Sets the given input`s IsAdvancedDisplay state. AdvancedDisplay pins are hidden in the node by default.
 	// returns true if state was changed.
@@ -738,10 +769,6 @@ PRAGMA_ENABLE_EXPERIMENTAL_WARNINGS
 	// replacing nodes with other, registry-defined concrete node class instances. Returns true if any template nodes were processed.
 	UE_API bool TransformTemplateNodes();
 
-	// Versions legacy document members that contained interface information
-	UE_DEPRECATED(5.5, "Moved to internally implemented versioning logic")
-	UE_API bool VersionInterfaces();
-
 	// Struct enabling property migration of data that must be applied prior to versioning logic
 	struct IPropertyVersionTransform
 	{
@@ -757,6 +784,8 @@ PRAGMA_ENABLE_EXPERIMENTAL_WARNINGS
 #endif // WITH_EDITORONLY_DATA
 
 private:
+	bool bInDocConfigScope = false;
+
 	using FFinalizeNodeFunctionRef = TFunctionRef<void(FMetasoundFrontendNode&, const Metasound::Frontend::FNodeRegistryKey&)>;
 
 	TInstancedStruct<FMetaSoundFrontendNodeConfiguration> CreateFrontendNodeConfiguration(const FMetasoundFrontendClassMetadata& InClassMetadata) const;
@@ -876,5 +905,4 @@ private:
 
 	FTopLevelAssetPath HintPath;
 };
-
 #undef UE_API
diff --git a/Engine/Plugins/Runtime/Metasound/Source/MetasoundFrontend/Public/MetasoundFrontendDocumentBuilderRegistry.h b/Engine/Plugins/Runtime/Metasound/Source/MetasoundFrontend/Public/MetasoundFrontendDocumentBuilderRegistry.h
new file mode 100644
--- /dev/null
+++ b/Engine/Plugins/Runtime/Metasound/Source/MetasoundFrontend/Public/MetasoundFrontendDocumentBuilderRegistry.h
@@ -0,0 +1,85 @@
+// Copyright Epic Games, Inc. All Rights Reserved.
+#pragma once
+
+#include "MetasoundBuilderInterface.h"
+#include "MetasoundDocumentInterface.h"
+#include "MetasoundFrontendDocumentModifyDelegates.h"
+#include "MetasoundOperatorInterface.h"
+#include "MetasoundOperatorSettings.h"
+#include "Templates/Function.h"
+#include "Templates/UniquePtr.h"
+#include "UObject/Interface.h"
+#include "UObject/TopLevelAssetPath.h"
+
+
+// Forward Declarations
+struct FMetasoundFrontendClassInputDefault;
+struct FMetasoundFrontendClassInput;
+struct FMetaSoundFrontendDocumentBuilder;
+struct FMetasoundFrontendGraphClass;
+struct FTopLevelAssetPath;
+
+#define UE_API METASOUNDFRONTEND_API
+
+namespace Metasound::Frontend
+{
+	class IDocumentBuilderRegistry
+	{
+	public:
+		virtual ~IDocumentBuilderRegistry() = default;
+
+#if WITH_EDITORONLY_DATA
+		UE_DEPRECATED(5.7, "Use StripUnusedPages(...) instead.")
+		virtual bool CookPages(FName CookPlatformName, FMetaSoundFrontendDocumentBuilder& Builder) const = 0;
+#endif // WITH_EDITORONLY_DATA
+
+		virtual FMetaSoundFrontendDocumentBuilder* FindBuilder(TScriptInterface<IMetaSoundDocumentInterface> MetaSound) const = 0;
+		virtual FMetaSoundFrontendDocumentBuilder* FindBuilder(const FMetasoundFrontendClassName& ClassName, const FTopLevelAssetPath& AssetPath) const = 0;
+		virtual FMetaSoundFrontendDocumentBuilder* FindOutermostBuilder(const UObject& InSubObject) const = 0;
+
+#if WITH_EDITORONLY_DATA
+		// Find the existing builder for the given MetaSound, or optionally begin building by attaching a new builder.  Only available
+		// in builds with editor only data as building serialized assets (which may have template nodes, cooked builds do not) is only
+		// supported when editor data is loaded. Creating transient builders can simply be done by passing a new MetaSound asset to
+		// a FMetaSoundFrontendDocumentBuilder constructor, or this registry's implementation may supply its own create call for tracking
+		// and reuse purposes.
+		virtual FMetaSoundFrontendDocumentBuilder& FindOrBeginBuilding(TScriptInterface<IMetaSoundDocumentInterface> MetaSound) = 0;
+#endif // WITH_EDITORONLY_DATA
+
+		// Removes builder from registry, clearing any cached builder state. (Optionally) forces unregistration from the Frontend Node Class Registry
+		// (If the builder has outstanding transactions, unregistration from the Node Class Registry will occur regardless).
+		virtual bool FinishBuilding(const FMetasoundFrontendClassName& InClassName, bool bForceUnregisterNodeClass = false) const = 0;
+		virtual bool FinishBuilding(const FMetasoundFrontendClassName& InClassName, const FTopLevelAssetPath& InAssetPath, bool bForceUnregisterNodeClass = false) const = 0;
+
+		// Reloads the given builder, maintaining all modify delegate subscriptions. Returns true if builder was found and reloaded,
+		// false if not found.
+		virtual bool ReloadBuilder(const FMetasoundFrontendClassName& InClassName) const = 0;
+
+		UE_DEPRECATED(5.7, "This is no longer supported. FindPreferredPage(...)")
+		virtual FGuid ResolveTargetPageID(const FMetasoundFrontendGraphClass& InGraphClass) const = 0;
+
+		UE_DEPRECATED(5.7, "This is no longer supported. FindPreferredPage(...)")
+		virtual FGuid ResolveTargetPageID(const FMetasoundFrontendClassInput& InClassInput) const = 0;
+
+		UE_DEPRECATED(5.7, "This is no longer supported. FindPreferredPage(...)")
+		virtual FGuid ResolveTargetPageID(const TArray<FMetasoundFrontendClassInputDefault>& Defaults) const = 0;
+
+		UE_DEPRECATED(5.7, "This is only used to support existing deprecated functionality")
+		virtual TArrayView<const FGuid> GetPageOrder() const = 0;
+
+#if WITH_EDITORONLY_DATA
+		UE_DEPRECATED(5.7, "This is only used to support existing deprecated functionality")
+		virtual TArray<FGuid> GetCookedTargetPages(FName InPlatformName) const = 0;
+
+		UE_DEPRECATED(5.7, "This is only used to support existing deprecated functionality")
+		virtual TArray<FGuid> GetCookedPageOrder(FName InPlatformName) const = 0;
+#endif // WITH_EDITORONLY_DATA
+
+		static UE_API class IDocumentBuilderRegistry* Get();
+		static UE_API class IDocumentBuilderRegistry& GetChecked();
+		static UE_API void Deinitialize();
+		static UE_API void Initialize(TUniquePtr<IDocumentBuilderRegistry>&& InBuilderRegistry);
+	};
+} // namespace Metasound::Frontend
+
+#undef UE_API
diff --git a/Engine/Plugins/Runtime/Metasound/Source/MetasoundFrontend/Public/MetasoundFrontendDocumentConfiguration.h b/Engine/Plugins/Runtime/Metasound/Source/MetasoundFrontend/Public/MetasoundFrontendDocumentConfiguration.h
new file mode 100644
--- /dev/null
+++ b/Engine/Plugins/Runtime/Metasound/Source/MetasoundFrontend/Public/MetasoundFrontendDocumentConfiguration.h
@@ -0,0 +1,156 @@
+// Copyright Epic Games, Inc. All Rights Reserved.
+#pragma once
+
+#include "Containers/Set.h"
+#include "Internationalization/Text.h"
+#include "Misc/Guid.h"
+#include "StructUtils/InstancedStruct.h"
+#include "Templates/SubclassOf.h"
+#include "UObject/NameTypes.h"
+#include "UObject/NoExportTypes.h"
+#include "UObject/ScriptInterface.h"
+
+#include "MetasoundFrontendDocumentConfiguration.generated.h"
+
+#define UE_API METASOUNDFRONTEND_API
+
+// Forward Declarations
+struct FMetaSoundFrontendDocumentBuilder;
+class UMetaSoundFrontendMemberMetadata;
+
+enum class EMetasoundFrontendClassType : uint8;
+
+
+// Vertex metadata displayed on MetaSound inputs & outputs pertaining to a given configuration.
+USTRUCT()
+struct FMetaSoundFrontendDocumentConfigurationVertexMetadata
+{
+	GENERATED_BODY()
+
+private:
+	UPROPERTY()
+	FGuid NodeID;
+
+public:
+	UE_API FMetaSoundFrontendDocumentConfigurationVertexMetadata() = default;
+	UE_API FMetaSoundFrontendDocumentConfigurationVertexMetadata(const FGuid& InNodeID)
+		: NodeID(InNodeID)
+	{
+	}
+
+	virtual ~FMetaSoundFrontendDocumentConfigurationVertexMetadata() = default;
+
+	UE_EXPERIMENTAL(5.8, "Configurator is in active development and subject to change")
+	UE_API const FGuid& GetNodeID() const { return NodeID; }
+
+	UE_EXPERIMENTAL(5.8, "Configurator is in active development and subject to change")
+	UE_API virtual void OnDefaultReset(const FMetaSoundFrontendDocumentBuilder& Builder) { }
+
+	UE_EXPERIMENTAL(5.8, "Configurator is in active development and subject to change")
+	UE_API virtual void OnDefaultUpdated(const FMetaSoundFrontendDocumentBuilder& Builder) { }
+};
+
+// Configuration of a MetaSound, used to define procedural behavior when making modifications
+// to a MetaSound prior to serialization. Configurations enable MetaSounds to be procedurally
+// generated at edit and cook-time to adhere to a specific design or topology anytime a given
+// configuration or its respective properties are changed, enabling faster and more accurate
+// iteration where fully-customized MetaSound edit behavior is not desired.
+USTRUCT()
+struct FMetaSoundFrontendDocumentConfiguration
+{
+	GENERATED_BODY()
+
+	virtual ~FMetaSoundFrontendDocumentConfiguration() = default;
+
+	// Creates and stores new ConfigMetadata with the provided NodeID. Overwrites if metadata already exists with
+	// default constructed ConfigMetadata if configuration supports Metadata.
+	UE_EXPERIMENTAL(5.8, "Configurator is in active development and subject to change")
+	UE_API FMetaSoundFrontendDocumentConfigurationVertexMetadata* AddVertexMetadata(const FMetaSoundFrontendDocumentBuilder& Builder, const FGuid& NodeID);
+
+	// Creates and stores new ConfigMetadata with the provided type and name, asserting that the type created by the overloaded configuration
+	// matches the template type. Overwrites if metadata already exists with default constructed ConfigMetadata. Returns null if configuration
+	// does not support Metadata.
+	template <typename TConfigVertexMetadata = FMetaSoundFrontendDocumentConfigurationVertexMetadata>
+	TConfigVertexMetadata* AddVertexMetadata(const FMetaSoundFrontendDocumentBuilder& Builder, const FGuid& NodeID)
+	{
+		if (TInstancedStruct<FMetaSoundFrontendDocumentConfigurationVertexMetadata>* Entry = ConfigMetadata.Find(NodeID))
+		{
+			return &Entry->GetMutable<TConfigVertexMetadata>();
+		}
+
+		TInstancedStruct<FMetaSoundFrontendDocumentConfigurationVertexMetadata> NewMetadata = ConstructVertexMetadata(Builder, NodeID);
+		if (NewMetadata.IsValid())
+		{
+			return &ConfigMetadata.Add(NodeID, MoveTemp(NewMetadata)).GetMutable<TConfigVertexMetadata>();
+		}
+
+		return nullptr;
+	}
+
+#if WITH_EDITORONLY_DATA
+	// Executes document configuration. Returns true if document configuration was successful, false if not.
+	UE_EXPERIMENTAL(5.8, "Configurator is in active development and subject to change")
+	UE_API virtual bool ConfigureDocument(FMetaSoundFrontendDocumentBuilder& OutBuilder, bool bLogFailure);
+#endif // WITH_EDITORONLY_DATA
+
+	// Returns whether or not there is existing Metadata with the given NodeID.
+	UE_EXPERIMENTAL(5.8, "Configurator is in active development and subject to change")
+	UE_API bool ContainsMetadata(const FGuid& NodeID) const;
+
+	// Helper function to returns ConfigMetadata with the provided class type and name and cast to underlying config type.
+	template <typename TConfigVertexMetadata = FMetaSoundFrontendDocumentConfigurationVertexMetadata>
+	TConfigVertexMetadata* FindMetadata(const FGuid& NodeID)
+	{
+		if (TInstancedStruct<FMetaSoundFrontendDocumentConfigurationVertexMetadata>* Entry = ConfigMetadata.Find(NodeID))
+		{
+			return Entry->GetMutablePtr<TConfigVertexMetadata>();
+		}
+
+		return nullptr;
+	}
+
+	// Helper function to returns ConfigMetadata with the provided class type and name and cast to underlying config type.
+	template <typename TConfigVertexMetadata = FMetaSoundFrontendDocumentConfigurationVertexMetadata>
+	const TConfigVertexMetadata* FindConstMetadata(const FGuid& NodeID) const
+	{
+		if (const TInstancedStruct<FMetaSoundFrontendDocumentConfigurationVertexMetadata>* Entry = ConfigMetadata.Find(NodeID))
+		{
+			return Entry->GetPtr<TConfigVertexMetadata>();
+		}
+
+		return nullptr;
+	}
+
+#if WITH_EDITOR
+	UE_EXPERIMENTAL(5.8, "Document Configuration is in active development and subject to change")
+	UE_API virtual bool IsCustomTabEnabled() const { return true; }
+
+	UE_EXPERIMENTAL(5.8, "Document Configuration is in active development and subject to change")
+	UE_API virtual bool IsGraphEditorVisible() const { return false; }
+
+	UE_EXPERIMENTAL(5.8, "Document Configuration is in active development and subject to change")
+	UE_API virtual bool IsInterfaceEditingEnabled() const { return false; }
+
+	UE_EXPERIMENTAL(5.8, "Document Configuration is in active development and subject to change")
+	UE_API virtual bool IsMemberEditingEnabled() const { return false; }
+
+	UE_EXPERIMENTAL(5.8, "Document Configuration is in active development and subject to change")
+	UE_API virtual bool IsPageEditingEnabled() const { return false; }
+
+	UE_EXPERIMENTAL(5.8, "Document Configuration is in active development and subject to change")
+	UE_API virtual void OnPropertyChanged(const FPropertyChangedEvent& InEvent, FMetaSoundFrontendDocumentBuilder& OutBuilder);
+#endif // WITH_EDITOR
+
+	// Removes ConfigMetadata with the provided NodeID. Returns if Metadata was found and removed or not.
+	UE_EXPERIMENTAL(5.8, "Document Configuration is in active development and subject to change")
+	UE_API bool RemoveMetadata(const FGuid& NodeID);
+
+protected:
+	UE_EXPERIMENTAL(5.8, "Configurator is in active development and subject to change")
+	UE_API virtual TInstancedStruct<FMetaSoundFrontendDocumentConfigurationVertexMetadata> ConstructVertexMetadata(const FMetaSoundFrontendDocumentBuilder& Builder, const FGuid& NodeID) const;
+
+private:
+	UPROPERTY()
+	TMap<FGuid, TInstancedStruct<FMetaSoundFrontendDocumentConfigurationVertexMetadata>> ConfigMetadata;
+};
+#undef UE_API
diff --git a/Engine/Plugins/Runtime/Metasound/Source/MetasoundFrontend/Public/MetasoundFrontendDocumentModifyDelegates.h b/Engine/Plugins/Runtime/Metasound/Source/MetasoundFrontend/Public/MetasoundFrontendDocumentModifyDelegates.h
--- a/Engine/Plugins/Runtime/Metasound/Source/MetasoundFrontend/Public/MetasoundFrontendDocumentModifyDelegates.h
+++ b/Engine/Plugins/Runtime/Metasound/Source/MetasoundFrontend/Public/MetasoundFrontendDocumentModifyDelegates.h
@@ -28,10 +28,14 @@ namespace Metasound::Frontend
 		FGuid PageID;
 	};
 
-	struct FDocumentPresetStateChangedArgs
-	{
-	};
 
+	struct FDocumentConfigChangedArgs {};
+	DECLARE_MULTICAST_DELEGATE_OneParam(FOnDocumentConfigurationChanged, const FDocumentConfigChangedArgs& /* Args */);
+
+	// Soft deprecated in 5.8 in favor of FOnDocumentConfigurationChanged.  Once delegate is removed, it can be officially deprecated.
+	struct FDocumentPresetStateChangedArgs {};
+
+	// Soft deprecated in 5.8 as usage is no longer valid (migrated to OnDocumentConfigurationChanged)
 	DECLARE_MULTICAST_DELEGATE_OneParam(FOnDocumentPresetStateChanged, const FDocumentPresetStateChangedArgs& /* Args */);
 
 	DECLARE_MULTICAST_DELEGATE_OneParam(FOnDocumentPageAdded, const FDocumentMutatePageArgs& /* Args */);
@@ -108,15 +112,20 @@ namespace Metasound::Frontend
 		UE_API FDocumentModifyDelegates();
 		UE_API FDocumentModifyDelegates(const FMetasoundFrontendDocument& Document);
 
+PRAGMA_DISABLE_DEPRECATION_WARNINGS
 		~FDocumentModifyDelegates() = default;
+PRAGMA_ENABLE_DEPRECATION_WARNINGS
 
 		FOnMetaSoundFrontendDocumentMutateMetadata OnDocumentMetadataChanged;
 		FOnMetaSoundFrontendDocumentMutateArray OnDependencyAdded;
 		FOnMetaSoundFrontendDocumentRemoveSwappingArray OnRemoveSwappingDependency;
 		FOnMetaSoundFrontendDocumentRenameClass OnRenamingDependencyClass;
 
+		// Soft deprecated in 5.8: "Use OnDocumentConfigurationChanged instead"
 		FOnDocumentPresetStateChanged OnPresetStateChanged;
 
+		FOnDocumentConfigurationChanged OnDocumentConfigurationChanged;
+
 		FPageModifyDelegates PageDelegates;
 		FInterfaceModifyDelegates InterfaceDelegates;
 
diff --git a/Engine/Plugins/Runtime/Metasound/Source/MetasoundFrontend/Public/MetasoundFrontendGraphLayer.h b/Engine/Plugins/Runtime/Metasound/Source/MetasoundFrontend/Public/MetasoundFrontendGraphLayer.h
new file mode 100644
--- /dev/null
+++ b/Engine/Plugins/Runtime/Metasound/Source/MetasoundFrontend/Public/MetasoundFrontendGraphLayer.h
@@ -0,0 +1,35 @@
+// Copyright Epic Games, Inc. All Rights Reserved.
+#pragma once
+
+#include "Algo/Transform.h"
+#include "Containers/Array.h"
+#include "MetasoundFrontendDocument.h"
+#include "UObject/NoExportTypes.h"
+
+
+#include "MetasoundFrontendGraphLayer.generated.h"
+
+
+#define UE_API METASOUNDFRONTEND_API
+
+struct UE_EXPERIMENTAL(5.7, "GraphLayers are currently being tested for Config generation, but may expand functionality and change in the future.") FMetasoundFrontendGraphLayer;
+
+USTRUCT()
+struct FMetasoundFrontendGraphLayer
+{
+	GENERATED_BODY()
+
+	UPROPERTY()
+	TArray<FMetasoundFrontendClass> Dependencies;
+
+	UPROPERTY()
+	FMetasoundFrontendGraph Graph;
+
+	UPROPERTY()
+	TArray<FMetasoundFrontendClassInput> Inputs;
+
+	UPROPERTY()
+	TArray<FMetasoundFrontendClassOutput> Outputs;
+};
+
+#undef UE_API
diff --git a/Engine/Plugins/Runtime/Metasound/Source/MetasoundFrontend/Public/MetasoundFrontendLiteral.h b/Engine/Plugins/Runtime/Metasound/Source/MetasoundFrontend/Public/MetasoundFrontendLiteral.h
--- a/Engine/Plugins/Runtime/Metasound/Source/MetasoundFrontend/Public/MetasoundFrontendLiteral.h
+++ b/Engine/Plugins/Runtime/Metasound/Source/MetasoundFrontend/Public/MetasoundFrontendLiteral.h
@@ -2,6 +2,7 @@
 
 #pragma once
 
+#include "Algo/Transform.h"
 #include "AudioParameterControllerInterface.h"
 #include "MetasoundLiteral.h"
 
@@ -74,6 +75,7 @@ struct FMetasoundFrontendLiteral
 	struct FDefaultArray { int32 Num = 0; };
 
 	FMetasoundFrontendLiteral() = default;
+
 	UE_API FMetasoundFrontendLiteral(const FAudioParameter& InParameter);
 
 private:
@@ -125,17 +127,17 @@ public:
 	UE_API void SetType(EMetasoundFrontendLiteralType InType);
 
 	UE_API void Set(FDefault InValue);
-	UE_API void Set(const FDefaultArray& InValue);
+	UE_API void Set(FDefaultArray InValue);
 	UE_API void Set(bool InValue);
-	UE_API void Set(const TArray<bool>& InValue);
+	UE_API void Set(TArray<bool> InValue);
 	UE_API void Set(int32 InValue);
-	UE_API void Set(const TArray<int32>& InValue);
+	UE_API void Set(TArray<int32> InValue);
 	UE_API void Set(float InValue);
-	UE_API void Set(const TArray<float>& InValue);
-	UE_API void Set(const FString& InValue);
-	UE_API void Set(const TArray<FString>& InValue);
+	UE_API void Set(TArray<float> InValue);
+	UE_API void Set(FString InValue);
+	UE_API void Set(TArray<FString> InValue);
 	UE_API void Set(UObject* InValue);
-	UE_API void Set(const TArray<UObject*>& InValue);
+	UE_API void Set(TArray<UObject*> InValue);
 
 	UE_API void SetFromLiteral(const Metasound::FLiteral& InLiteral);
 
@@ -169,22 +171,38 @@ public:
 		return A.IsEqual(B);
 	}
 
-
 private:
 	// Remove all values.
 	UE_API void Empty();
 };
 
-namespace Metasound
+namespace Metasound::Frontend
 {
-	namespace Frontend
+	template <typename ObjectType = UObject>
+	FMetasoundFrontendLiteral MakeFrontendLiteral(TConstArrayView<TObjectPtr<ObjectType>> ObjectPtrs)
 	{
-		// Convenience function to convert Metasound::ELiteralType to EMetasoundFrontendLiteralType.
-		METASOUNDFRONTEND_API EMetasoundFrontendLiteralType GetMetasoundFrontendLiteralType(ELiteralType InLiteralType);
+		TArray<UObject*> Objects;
+		Algo::Transform(ObjectPtrs, Objects,
+			[](const TObjectPtr<ObjectType> Obj) { return Obj.Get(); });
 
-		// Convenience function to convert EMetasoundFrontendLiteralType to Metasound::ELiteralType.
-		METASOUNDFRONTEND_API ELiteralType GetMetasoundLiteralType(EMetasoundFrontendLiteralType InLiteralType);
+		FMetasoundFrontendLiteral Value;
+		Value.Set(MoveTemp(Objects));
+		return Value;
 	}
+
+	template <typename LiteralType>
+	FMetasoundFrontendLiteral MakeFrontendLiteral(LiteralType InitValue)
+	{
+		FMetasoundFrontendLiteral Value;
+		Value.Set(MoveTemp(InitValue));
+		return Value;
+	}
+
+	// Convenience function to convert Metasound::ELiteralType to EMetasoundFrontendLiteralType.
+	UE_API EMetasoundFrontendLiteralType GetMetasoundFrontendLiteralType(ELiteralType InLiteralType);
+
+	// Convenience function to convert EMetasoundFrontendLiteralType to Metasound::ELiteralType.
+	UE_API ELiteralType GetMetasoundLiteralType(EMetasoundFrontendLiteralType InLiteralType);
 }
 
 #undef UE_API
diff --git a/Engine/Plugins/Runtime/Metasound/Source/MetasoundFrontend/Public/MetasoundFrontendTransform.h b/Engine/Plugins/Runtime/Metasound/Source/MetasoundFrontend/Public/MetasoundFrontendTransform.h
--- a/Engine/Plugins/Runtime/Metasound/Source/MetasoundFrontend/Public/MetasoundFrontendTransform.h
+++ b/Engine/Plugins/Runtime/Metasound/Source/MetasoundFrontend/Public/MetasoundFrontendTransform.h
@@ -15,6 +15,7 @@
 
 // Forward Declarations
 struct FMetaSoundFrontendDocumentBuilder;
+struct FMetaSoundFrontendPresetDocumentConfiguration;
 
 
 namespace Metasound::Frontend
