[Backout] - CL47159054 - CIS Compile Warning
#rnx
Original CL Desc
-----------------------------------------------------------------
Document Config Checkpoint 1
- Adds basic serialization schema for config
- Adds checkpoint versions of inactive implementations including preset and randomizer
- Moves unique name generation to MetaSound Builder API from private MetaSound Editor Builder API
- Moves IsPreset calls to IsGraphEditable which maintains parity at this point (but will diverge with the implementation of DocConfig)
- Literal optimizations & MakeLiteral convenience functions

#rnx

diff --git a/Engine/Plugins/Experimental/TechAudioTools/Source/TechAudioToolsMetaSound/Private/ViewModels/MetaSoundViewModel.cpp b/Engine/Plugins/Experimental/TechAudioTools/Source/TechAudioToolsMetaSound/Private/ViewModels/MetaSoundViewModel.cpp
--- a/Engine/Plugins/Experimental/TechAudioTools/Source/TechAudioToolsMetaSound/Private/ViewModels/MetaSoundViewModel.cpp
+++ b/Engine/Plugins/Experimental/TechAudioTools/Source/TechAudioToolsMetaSound/Private/ViewModels/MetaSoundViewModel.cpp
@@ -1,8 +1,7 @@
-// Copyright Epic Games, Inc. All Rights Reserved.
+﻿// Copyright Epic Games, Inc. All Rights Reserved.
 
 #include "ViewModels/MetaSoundViewModel.h"
 
-#include "DocumentConfigurations/MetasoundPresetDocumentConfiguration.h"
 #include "Logging/StructuredLog.h"
 #include "MetasoundDocumentBuilderRegistry.h"
 #include "MetasoundFrontendPages.h"
@@ -149,7 +148,7 @@ void UMetaSoundViewModel::SetInputOverridesDefault(const FName& InputName, const
 
 void UMetaSoundViewModel::SetInputIsConstructorPin(const FName& InputName, const bool bIsConstructorPin) const
 {
-	if (const UMetaSoundInputViewModel* InputViewModel = InputViewModels.FindRef(InputName))
+	if (const UMetaSoundInputViewModel* InputViewModel = *InputViewModels.Find(InputName))
 	{
 		const FName DataType = InputViewModel->GetDataType();
 		if (!UMetaSoundViewModelConversionFunctions::IsConstructorType(DataType))
@@ -210,7 +209,7 @@ void UMetaSoundViewModel::SetOutputDataType(const FName& OutputName, const FName
 
 void UMetaSoundViewModel::SetOutputIsConstructorPin(const FName& OutputName, const bool bIsConstructorPin) const
 {
-	if (const UMetaSoundOutputViewModel* OutputViewModel = OutputViewModels.FindRef(OutputName))
+	if (const UMetaSoundOutputViewModel* OutputViewModel = *OutputViewModels.Find(OutputName))
 	{
 		const FName DataType = OutputViewModel->GetDataType();
 		if (!UMetaSoundViewModelConversionFunctions::IsConstructorType(DataType))
@@ -268,7 +267,7 @@ void UMetaSoundViewModel::OnInputNameChanged(const FName OldName, const FName Ne
 
 void UMetaSoundViewModel::OnInputDataTypeChanged(const FName VertexName, const FName DataType)
 {
-	if (UMetaSoundInputViewModel* InputViewModel = InputViewModels.FindRef(VertexName))
+	if (UMetaSoundInputViewModel* InputViewModel = *InputViewModels.Find(VertexName))
 	{
 		InputViewModel->OnInputDataTypeChanged(DataType);
 	}
@@ -276,7 +275,7 @@ void UMetaSoundViewModel::OnInputDataTypeChanged(const FName VertexName, const F
 
 void UMetaSoundViewModel::OnInputDefaultChanged(const FName VertexName, FMetasoundFrontendLiteral LiteralValue, const FName PageName)
 {
-	if (UMetaSoundInputViewModel* InputViewModel = InputViewModels.FindRef(VertexName))
+	if (UMetaSoundInputViewModel* InputViewModel = *InputViewModels.Find(VertexName))
 	{
 		InputViewModel->OnInputDefaultChanged(LiteralValue, PageName);
 	}
@@ -284,7 +283,7 @@ void UMetaSoundViewModel::OnInputDefaultChanged(const FName VertexName, FMetasou
 
 void UMetaSoundViewModel::OnInputInheritsDefaultChanged(const FName VertexName, const bool bInheritsDefault)
 {
-	if (UMetaSoundInputViewModel* InputViewModel = InputViewModels.FindRef(VertexName))
+	if (UMetaSoundInputViewModel* InputViewModel = *InputViewModels.Find(VertexName))
 	{
 		InputViewModel->OnInputInheritsDefaultChanged(bInheritsDefault);
 	}
@@ -292,7 +291,7 @@ void UMetaSoundViewModel::OnInputInheritsDefaultChanged(const FName VertexName,
 
 void UMetaSoundViewModel::OnInputIsConstructorPinChanged(const FName VertexName, const bool bIsConstructorPin)
 {
-	if (UMetaSoundInputViewModel* InputViewModel = InputViewModels.FindRef(VertexName))
+	if (UMetaSoundInputViewModel* InputViewModel = *InputViewModels.Find(VertexName))
 	{
 		InputViewModel->OnInputIsConstructorPinChanged(bIsConstructorPin);
 	}
@@ -338,7 +337,7 @@ void UMetaSoundViewModel::OnOutputNameChanged(const FName OldName, const FName N
 
 void UMetaSoundViewModel::OnOutputDataTypeChanged(const FName VertexName, const FName DataType)
 {
-	if (UMetaSoundOutputViewModel* OutputViewModel = OutputViewModels.FindRef(VertexName))
+	if (UMetaSoundOutputViewModel* OutputViewModel = *OutputViewModels.Find(VertexName))
 	{
 		OutputViewModel->OnOutputDataTypeChanged(DataType);
 	}
@@ -346,7 +345,7 @@ void UMetaSoundViewModel::OnOutputDataTypeChanged(const FName VertexName, const
 
 void UMetaSoundViewModel::OnOutputIsConstructorPinChanged(const FName VertexName, const bool bIsConstructorPin)
 {
-	if (UMetaSoundOutputViewModel* OutputViewModel = OutputViewModels.FindRef(VertexName))
+	if (UMetaSoundOutputViewModel* OutputViewModel = *OutputViewModels.Find(VertexName))
 	{
 		OutputViewModel->OnOutputIsConstructorPinChanged(bIsConstructorPin);
 	}
diff --git a/Engine/Plugins/Runtime/Metasound/Source/MetasoundEditor/Private/MetasoundDetailCustomization.cpp b/Engine/Plugins/Runtime/Metasound/Source/MetasoundEditor/Private/MetasoundDetailCustomization.cpp
--- a/Engine/Plugins/Runtime/Metasound/Source/MetasoundEditor/Private/MetasoundDetailCustomization.cpp
+++ b/Engine/Plugins/Runtime/Metasound/Source/MetasoundEditor/Private/MetasoundDetailCustomization.cpp
@@ -11,7 +11,6 @@
 #include "IAudioParameterTransmitter.h"
 #include "IDetailGroup.h"
 #include "Input/Events.h"
-#include "InstancedStructDetails.h"
 #include "Interfaces/MetasoundFrontendInterfaceRegistry.h"
 #include "MetasoundAssetBase.h"
 #include "MetasoundBuilderSubsystem.h"
@@ -40,7 +39,6 @@
 #include "Styling/StyleColors.h"
 #include "Templates/Casts.h"
 #include "Templates/SharedPointer.h"
-#include "UObject/UnrealType.h"
 #include "UObject/WeakObjectPtr.h"
 #include "UObject/WeakObjectPtrTemplates.h"
 #include "Widgets/Input/SCheckBox.h"
@@ -109,11 +107,6 @@ namespace Metasound::Editor
 	{
 	}
 
-	FName FMetasoundDetailCustomization::GetDocConfigPropertyPath() const
-	{
-		return BuildChildPath(DocumentPropertyName, GET_MEMBER_NAME_CHECKED(FMetasoundFrontendDocument, Configuration));
-	}
-
 	FName FMetasoundDetailCustomization::GetInterfaceVersionsPropertyPath() const
 	{
 		return BuildChildPath(DocumentPropertyName, GET_MEMBER_NAME_CHECKED(FMetasoundFrontendDocument, Interfaces));
@@ -139,7 +132,7 @@ namespace Metasound::Editor
 
 		// Only support modifying a single MetaSound at a time (Multiple
 		// MetaSound editing will be covered most likely by separate tool).
-		if (Objects.Num() != 1 || !Objects.Last().IsValid())
+		if (Objects.Num() > 1 || !Objects.Last().IsValid())
 		{
 			return;
 		}
@@ -179,6 +172,8 @@ namespace Metasound::Editor
 				const FName MajorVersionPropertyPath = BuildChildPath(VersionPropertyPath, GET_MEMBER_NAME_CHECKED(FMetasoundFrontendVersionNumber, Major));
 				const FName MinorVersionPropertyPath = BuildChildPath(VersionPropertyPath, GET_MEMBER_NAME_CHECKED(FMetasoundFrontendVersionNumber, Minor));
 
+				const FName InterfaceVersionsPropertyPath = GetInterfaceVersionsPropertyPath();
+
 				TSharedPtr<IPropertyHandle> AccessFlagsHandle = DetailLayout.GetProperty(AccessFlagsPropertyPath);
 				TSharedPtr<IPropertyHandle> AuthorHandle = DetailLayout.GetProperty(AuthorPropertyPath);
 				TSharedPtr<IPropertyHandle> CategoryHierarchyHandle = DetailLayout.GetProperty(CategoryHierarchyPropertyPath);
@@ -186,49 +181,10 @@ namespace Metasound::Editor
 				TSharedPtr<IPropertyHandle> DisplayNameHandle = DetailLayout.GetProperty(DisplayNamePropertyPath);
 				TSharedPtr<IPropertyHandle> DescHandle = DetailLayout.GetProperty(DescPropertyPath);
 				TSharedPtr<IPropertyHandle> KeywordsHandle = DetailLayout.GetProperty(KeywordsPropertyPath);
-				TSharedPtr<IPropertyHandle> InterfaceVersionsHandle = DetailLayout.GetProperty(GetInterfaceVersionsPropertyPath());
+				TSharedPtr<IPropertyHandle> InterfaceVersionsHandle = DetailLayout.GetProperty(InterfaceVersionsPropertyPath);
 				TSharedPtr<IPropertyHandle> MajorVersionHandle = DetailLayout.GetProperty(MajorVersionPropertyPath);
 				TSharedPtr<IPropertyHandle> MinorVersionHandle = DetailLayout.GetProperty(MinorVersionPropertyPath);
 
-				// Disabled as currently experimental and in local development.
-// 				{
-// 					TSharedPtr<IPropertyHandle> DocConfigHandle = DetailLayout.GetProperty(GetDocConfigPropertyPath());
-// 					if (DocConfigHandle.IsValid())
-// 					{
-// 						FSimpleDelegate OnConfigChanged = FSimpleDelegate::CreateLambda([this]()
-// 						{
-// 							if (Builder.IsValid())
-// 							{
-// 								FMetaSoundFrontendDocumentBuilder& DocBuilder = Builder->GetBuilder();
-// 								if (DocBuilder.IsValid())
-// 								{
-// 									DocBuilder.GetDocumentDelegates().OnDocumentConfigurationChanged.Broadcast({ });
-// 									DocBuilder.ConfigureDocument();
-// 								}
-// 							}
-// 						});
-// 						DocConfigHandle->SetOnPropertyValueChanged(OnConfigChanged);
-// 
-// 						const auto OnConfigPropChanged = TDelegate<void(const FPropertyChangedEvent&)>::CreateSPLambda(AsShared(), [this](const FPropertyChangedEvent& InEvent)
-// 						{
-// 							if (Builder.IsValid())
-// 							{
-// 								FMetaSoundFrontendDocumentBuilder& DocBuilder = Builder->GetBuilder();
-// 								if (DocBuilder.IsValid())
-// 								{
-// 									if (const FMetaSoundFrontendDocumentConfiguration* Config = DocBuilder.GetConstDocumentConfiguration().GetPtr())
-// 									{
-// 										const_cast<FMetaSoundFrontendDocumentConfiguration*>(Config)->OnPropertyChanged(InEvent, DocBuilder);
-// 									}
-// 								}
-// 							}
-// 						});
-// 						DocConfigHandle->SetOnChildPropertyValueChangedWithData(OnConfigPropChanged);
-// 					}
-// 
-// 					GeneralCategoryBuilder.AddProperty(DocConfigHandle);
-// 				}
-
 				// Invalid for UMetaSounds
 				TSharedPtr<IPropertyHandle> OutputFormat = DetailLayout.GetProperty(GET_MEMBER_NAME_CHECKED(UMetaSoundSource, OutputFormat));
 				if (OutputFormat.IsValid() && OutputFormat->IsValidHandle())
@@ -410,7 +366,7 @@ namespace Metasound::Editor
 
 		// Only support modifying a single MetaSound at a time (Multiple
 		// MetaSound editing will be covered most likely by separate tool).
-		if (Objects.Num() != 1 || !Objects.Last().IsValid())
+		if (Objects.Num() > 1)
 		{
 			return;
 		}
@@ -438,9 +394,9 @@ namespace Metasound::Editor
 				using namespace Engine;
 				using namespace Frontend;
 
-				if (InSelectInfo != ESelectInfo::OnNavigation && IsBuilderValid())
+				if (InSelectInfo != ESelectInfo::OnNavigation)
 				{
-					UObject& MetaSound = GetMetaSoundChecked();
+					UObject& MetaSound = GetMetaSound();
 
 					const FScopedTransaction Transaction(FText::Format(LOCTEXT("AddPageTransactionFormat", "Add MetaSound Page Graph '{0}'"), FText::FromString(*NameToAdd)));
 					MetaSound.Modify();
@@ -454,13 +410,11 @@ namespace Metasound::Editor
 					EMetaSoundBuilderResult Result = EMetaSoundBuilderResult::Failed;
 					Builder->AddGraphPage(FName(*NameToAdd.Get()), bDuplicateLastGraph, bSetAsBuildGraph, Result);
 
-					if (UAssetEditorSubsystem* AssetSubsystem = GEditor->GetEditorSubsystem<UAssetEditorSubsystem>())
+					IAssetEditorInstance* AssetEditor = GEditor->GetEditorSubsystem<UAssetEditorSubsystem>()->FindEditorForAsset(&GetMetaSound(), /*bFocusIfOpen=*/false);
+					FEditor* MetaSoundEditor = static_cast<FEditor*>(AssetEditor);
+					if (MetaSoundEditor)
 					{
-						IAssetEditorInstance* AssetEditor = AssetSubsystem->FindEditorForAsset(&MetaSound, /*bFocusIfOpen=*/false);
-						if (FEditor* MetaSoundEditor = static_cast<FEditor*>(AssetEditor))
-						{
-							MetaSoundEditor->RefreshDetails();
-						}
+						MetaSoundEditor->RefreshDetails();
 					}
 				}
 			})
@@ -489,7 +443,7 @@ namespace Metasound::Editor
 				PropertyCustomizationHelpers::MakeDeleteButton(FSimpleDelegate::CreateLambda([this]()
 				{
 					using namespace Frontend;
-					UObject& MetaSound = GetMetaSoundChecked();
+					UObject& MetaSound = GetMetaSound();
 
 					const FScopedTransaction Transaction(LOCTEXT("ResetGraphPagesTransaction", "Reset MetaSound Graph Pages"));
 					MetaSound.Modify();
@@ -527,32 +481,18 @@ namespace Metasound::Editor
 		}
 	}
 
-	const FMetasoundFrontendDocument& FMetasoundPagesDetailCustomization::GetConstDocumentChecked() const
-	{
-		return Builder->GetBuilder().GetConstDocumentChecked();
-	}
-
-	UObject& FMetasoundPagesDetailCustomization::GetMetaSoundChecked() const
+	UObject& FMetasoundPagesDetailCustomization::GetMetaSound() const
 	{
 		return Builder->GetBuilder().CastDocumentObjectChecked<UObject>();
 	}
 
-	bool FMetasoundPagesDetailCustomization::IsBuilderValid() const
-	{
-		return Builder->GetBuilder().IsValid();
-	}
-
 	void FMetasoundPagesDetailCustomization::RebuildImplemented()
 	{
 		EntryWidgets->ClearChildren();
 
-		if (!IsBuilderValid())
-		{
-			return;
-		}
-
 		TSet<FGuid> ImplementedGuids;
-		GetConstDocumentChecked().RootGraph.IterateGraphPages([&ImplementedGuids](const FMetasoundFrontendGraph& Graph)
+		const FMetasoundFrontendDocument& Document = Builder->GetBuilder().GetConstDocumentChecked();
+		Document.RootGraph.IterateGraphPages([&ImplementedGuids](const FMetasoundFrontendGraph& Graph)
 		{
 			ImplementedGuids.Add(Graph.PageID);
 		});
@@ -570,7 +510,7 @@ namespace Metasound::Editor
 					{
 						constexpr bool bOpenEditor = false; // Already focused by user action
 						UMetaSoundEditorSubsystem::GetConstChecked().SetFocusedPage(*Builder.Get(), PageID, bOpenEditor);
-						IAssetEditorInstance* AssetEditor = GEditor->GetEditorSubsystem<UAssetEditorSubsystem>()->FindEditorForAsset(&GetMetaSoundChecked(), /*bFocusIfOpen=*/false);
+						IAssetEditorInstance* AssetEditor = GEditor->GetEditorSubsystem<UAssetEditorSubsystem>()->FindEditorForAsset(&GetMetaSound(), /*bFocusIfOpen=*/false);
 						FEditor* MetaSoundEditor = static_cast<FEditor*>(AssetEditor);
 						if (MetaSoundEditor)
 						{
@@ -617,7 +557,7 @@ namespace Metasound::Editor
 					{
 						using namespace Frontend;
 						const FScopedTransaction Transaction(FText::Format(LOCTEXT("RemovePageTransactionFormat", "Remove MetaSound Page '{0}'"), FText::FromName(InName)));
-						UObject& MetaSound = GetMetaSoundChecked();
+						UObject& MetaSound = GetMetaSound();
 						MetaSound.Modify();
 
 						// Removal may modify the builder's build page ID if it is the currently set value
@@ -733,16 +673,11 @@ namespace Metasound::Editor
 		AddableItems.Reset();
 		ImplementedNames.Reset();
 
-		if (!IsBuilderValid())
-		{
-			return;
-		}
-
 		const UMetaSoundSettings* Settings = GetDefault<UMetaSoundSettings>();
 		check(Settings);
 
 		TSet<FGuid> ImplementedGuids;
-		const FMetasoundFrontendDocument& Document = GetConstDocumentChecked();
+		const FMetasoundFrontendDocument& Document = Builder->GetBuilder().GetConstDocumentChecked();
 		Document.RootGraph.IterateGraphPages([&ImplementedGuids](const FMetasoundFrontendGraph& Graph)
 		{
 			ImplementedGuids.Add(Graph.PageID);
@@ -793,7 +728,7 @@ namespace Metasound::Editor
 				if (PageSettings->Name != ParentPtr->BuildPageName)
 				{
 					ParentPtr->BuildPageName = PageSettings->Name;
-					FGraphBuilder::RegisterGraphWithFrontend(ParentPtr->GetMetaSoundChecked());
+					FGraphBuilder::RegisterGraphWithFrontend(ParentPtr->GetMetaSound());
 				}
 
 				ParentPtr->AddableItems.RemoveAll([&PageSettings](const TSharedPtr<FString>& Item) { return Item->Compare(PageSettings->Name.ToString()) == 0; });
@@ -830,7 +765,7 @@ namespace Metasound::Editor
 				if (PageSettings->Name != ParentPtr->BuildPageName)
 				{
 					ParentPtr->BuildPageName = PageSettings->Name;
-					FGraphBuilder::RegisterGraphWithFrontend(ParentPtr->GetMetaSoundChecked());
+					FGraphBuilder::RegisterGraphWithFrontend(ParentPtr->GetMetaSound());
 				}
 
 				ParentPtr->AddableItems.Add(MakeShared<FString>(PageSettings->Name.ToString()));
@@ -848,7 +783,7 @@ namespace Metasound::Editor
 
 		// Only support modifying a single MetaSound at a time (Multiple
 		// MetaSound editing will be covered most likely by separate tool).
-		if (Objects.Num() != 1 || !Objects.Last().IsValid())
+		if (Objects.Num() > 1)
 		{
 			return;
 		}
diff --git a/Engine/Plugins/Runtime/Metasound/Source/MetasoundEditor/Private/MetasoundDetailCustomization.h b/Engine/Plugins/Runtime/Metasound/Source/MetasoundEditor/Private/MetasoundDetailCustomization.h
--- a/Engine/Plugins/Runtime/Metasound/Source/MetasoundEditor/Private/MetasoundDetailCustomization.h
+++ b/Engine/Plugins/Runtime/Metasound/Source/MetasoundEditor/Private/MetasoundDetailCustomization.h
@@ -2,7 +2,6 @@
 #pragma once
 
 #include "IDetailCustomization.h"
-#include "MetasoundFrontendDocument.h"
 #include "MetasoundFrontendDocumentModifyDelegates.h"
 #include "SGraphActionMenu.h"
 #include "SSearchableComboBox.h"
@@ -48,7 +47,6 @@ namespace Metasound
 			// End of IDetailCustomization interface
 
 		private:
-			FName GetDocConfigPropertyPath() const;
 			FName GetInterfaceVersionsPropertyPath() const;
 			FName GetRootClassPropertyPath() const;
 			FName GetMetadataPropertyPath() const;
@@ -66,9 +64,7 @@ namespace Metasound
 			// End of IDetailCustomization interface
 
 		private:
-			const FMetasoundFrontendDocument& GetConstDocumentChecked() const;
-			UObject& GetMetaSoundChecked() const;
-			bool IsBuilderValid() const;
+			UObject& GetMetaSound() const;
 			void RebuildImplemented();
 			void RefreshView();
 			void UpdateItemNames();
diff --git a/Engine/Plugins/Runtime/Metasound/Source/MetasoundEditor/Private/MetasoundEditor.cpp b/Engine/Plugins/Runtime/Metasound/Source/MetasoundEditor/Private/MetasoundEditor.cpp
--- a/Engine/Plugins/Runtime/Metasound/Source/MetasoundEditor/Private/MetasoundEditor.cpp
+++ b/Engine/Plugins/Runtime/Metasound/Source/MetasoundEditor/Private/MetasoundEditor.cpp
@@ -159,11 +159,6 @@ namespace Metasound
 
 		bool PageEditorEnabled(const FMetaSoundFrontendDocumentBuilder& Builder, bool bHasProjectPageValues, bool bPresetCanEditPageValues)
 		{
-			if (!Builder.IsValid())
-			{
-				return false;
-			}
-
 			IMetasoundEditorModule& MetaSoundEditorModule = FModuleManager::GetModuleChecked<IMetasoundEditorModule>("MetaSoundEditor");
 			if (MetaSoundEditorModule.IsRestrictedMode())
 			{
@@ -588,10 +583,10 @@ namespace Metasound
 						if (GetHoveredGraph() == OwningGraph)
 						{
 							const UMetaSoundBuilderBase& Builder = OwningGraph->GetBuilderChecked();
-							const bool bIsEditable = Builder.GetConstBuilder().GetIsGraphEditable();
-							if (!bIsEditable)
+							const bool bIsPreset = Builder.IsPreset();
+							if (bIsPreset)
 							{
-								Message = FText::Format(LOCTEXT("DropTargetFailIsPreset", "'{0}': Graph is not editable"), GraphMember->GetDisplayName());
+								Message = FText::Format(LOCTEXT("DropTargetFailIsPreset", "'{0}': Graph is Preset"), GraphMember->GetDisplayName());
 							}
 							else if (UMetasoundEditorGraphInput* Input = Cast<UMetasoundEditorGraphInput>(GraphMember.Get()))
 							{
@@ -1694,11 +1689,12 @@ namespace Metasound
 				.OnActionDragged(this, &FEditor::OnActionDragged)
 				.OnActionMatchesName(this, &FEditor::HandleActionMatchesName)
 				.OnActionSelected(this, &FEditor::OnActionSelected)
+// 				.OnCategoryTextCommitted(this, &FEditor::OnCategoryNameCommitted)
 				.OnCollectAllActions(this, &FEditor::CollectAllActions)
 				.OnCollectStaticSections(this, &FEditor::CollectStaticSections)
-				.OnContextMenuOpening(this, &FEditor::OnContextMenuOpening)
+ 				.OnContextMenuOpening(this, &FEditor::OnContextMenuOpening)
 				.OnCreateWidgetForAction(this, &FEditor::OnCreateWidgetForAction)
-				.OnCanRenameSelectedAction(this, &FEditor::CanRenameOnActionNode)
+  				.OnCanRenameSelectedAction(this, &FEditor::CanRenameOnActionNode)
 				.OnGetFilterText(this, &FEditor::GetFilterText)
 				.OnGetSectionTitle(this, &FEditor::OnGetSectionTitle)
 				.OnGetSectionWidget(this, &FEditor::OnGetMenuSectionWidget)
@@ -1751,7 +1747,7 @@ namespace Metasound
 				PagesDetails->SetVisibility(VisibilityAttr);
 			}
 
-			if (Builder.IsValid() && !Builder->GetConstBuilder().GetIsGraphEditable())
+			if (Builder.IsValid() && Builder->IsPreset())
 			{
 				FLevelEditorModule& LevelEditor = FModuleManager::LoadModuleChecked<FLevelEditorModule>("LevelEditor");
 				LevelEditor.OnMapChanged().AddSP(this, &FEditor::ChangeUserPresetWidgetWorld);
@@ -1801,7 +1797,7 @@ namespace Metasound
 		{
 			using namespace Metasound::Frontend;
 
-			if (!Builder.IsValid() || Builder->GetConstBuilder().GetIsGraphEditable())
+			if (!Builder.IsValid() || !Builder->IsPreset())
 			{
 				return nullptr;
 			}
@@ -3051,8 +3047,8 @@ namespace Metasound
 					FCanExecuteAction::CreateLambda([this]() { return CanPromoteToDeferredVariable(); }));
 
 				GraphEditorCommands->MapAction(FGraphEditorCommands::Get().HideNoConnectionPins,
-					FExecuteAction::CreateSP(this, &FEditor::HideUnconnectedPins));
-
+					FExecuteAction::CreateSP(this, &FEditor::HideUnconnectedPins));				
+				
 				GraphEditorCommands->MapAction(FGraphEditorCommands::Get().ShowAllPins,
 					FExecuteAction::CreateSP(this, &FEditor::ShowUnconnectedPins));
 
@@ -3217,7 +3213,7 @@ namespace Metasound
 					}
 				}
 
-				if (NextToSelect && GraphMembersMenu.IsValid())
+				if (NextToSelect)
 				{
 					if (GraphMembersMenu->SelectItemByName(NextToSelect->GetMemberName(), ESelectInfo::Direct, static_cast<int32>(NextToSelect->GetSectionID())))
 					{
@@ -3643,28 +3639,11 @@ namespace Metasound
 			return nullptr;
 		}
 
-		const FMetaSoundFrontendDocumentConfiguration* FEditor::GetConfiguration() const
-		{
-			if (const FMetaSoundFrontendDocumentBuilder* DocBuilder = GetFrontendBuilder())
-			{
-				if (DocBuilder->IsValid())
-				{
-					return DocBuilder->GetConstDocumentConfiguration().GetPtr();
-				}
-			}
-
-			return nullptr;
-		}
-
 		FMetaSoundFrontendDocumentBuilder* FEditor::GetFrontendBuilder() const
 		{
 			if (Builder.IsValid())
 			{
-				FMetaSoundFrontendDocumentBuilder& DocBuilder = Builder->GetBuilder();
-				if (DocBuilder.IsValid())
-				{
-					return &DocBuilder;
-				}
+				return &Builder->GetBuilder();
 			}
 
 			return nullptr;
@@ -3694,18 +3673,15 @@ namespace Metasound
 				}
 
 				// Renameable member nodes (inputs/outputs/variables)
-				if (GraphMembersMenu.IsValid())
+				if (UMetasoundEditorGraphMemberNode* MemberNode = Cast<UMetasoundEditorGraphMemberNode>(*SelectedIter))
 				{
-					if (UMetasoundEditorGraphMemberNode* MemberNode = Cast<UMetasoundEditorGraphMemberNode>(*SelectedIter))
+					if (const UMetasoundEditorGraphMember* Member = MemberNode->GetMember()) 
 					{
-						if (const UMetasoundEditorGraphMember* Member = MemberNode->GetMember())
+						if (Member->CanRename())
 						{
-							if (Member->CanRename())
-							{
-								GraphMembersMenu->SelectItemByName(Member->GetMemberName(), ESelectInfo::Direct, static_cast<int32>(Member->GetSectionID()));
-								GraphMembersMenu->RefreshAllActions(/*bPreserveExpansion=*/ true, /*bHandleOnSelectionEvent=*/ true);
-								GraphMembersMenu->OnRequestRenameOnActionNode();
-							}
+							GraphMembersMenu->SelectItemByName(Member->GetMemberName(), ESelectInfo::Direct, static_cast<int32>(Member->GetSectionID()));
+							GraphMembersMenu->RefreshAllActions(/*bPreserveExpansion=*/ true, /*bHandleOnSelectionEvent=*/ true);
+							GraphMembersMenu->OnRequestRenameOnActionNode();
 						}
 					}
 				}
@@ -3825,7 +3801,7 @@ namespace Metasound
 					// Duplicate the Sources Frontend Node and add a new member from it
 					if (const UMetasoundEditorGraphVariable* SourceGraphVariable = Cast<UMetasoundEditorGraphVariable>(SourceGraphMember))
 					{
-						const FName NewName = DocBuilder.GenerateUniqueVariableName(SourceGraphVariable->GetMemberName().ToString());
+						const FName NewName = FGraphBuilder::GenerateUniqueVariableName(DocBuilder, SourceGraphVariable->GetMemberName().ToString());
 						const FMetasoundFrontendVariable* NewFrontendVariable = DocBuilder.DuplicateGraphVariable(SourceGraphVariable->GetMemberName(), NewName);
 						if (ensure(NewFrontendVariable))
 						{
@@ -3837,9 +3813,10 @@ namespace Metasound
 						const FName SourceMemberName = SourceGraphVertex->GetMemberName();
 						const EMetasoundFrontendClassType ClassType = SourceGraphVertex->GetClassType();
 						
+						const FName Name = FGraphBuilder::GenerateUniqueNameByClassType(Graph.GetMetasoundChecked(), ClassType, SourceMemberName.ToString());
+
 						if (ClassType == EMetasoundFrontendClassType::Input)
 						{
-							const FName Name = DocBuilder.GenerateUniqueInputName(SourceMemberName.ToString());
 							if (const FMetasoundFrontendClassInput* NewInput = DocBuilder.DuplicateGraphInput(SourceMemberName, Name))
 							{
 								FGraphBuilder::SynchronizeGraphMembers(DocBuilder, Graph);
@@ -3848,7 +3825,6 @@ namespace Metasound
 						}
 						else if (ClassType == EMetasoundFrontendClassType::Output)
 						{
-							const FName Name = DocBuilder.GenerateUniqueOutputName(SourceMemberName.ToString());
 							if (const FMetasoundFrontendClassOutput* NewOutput = DocBuilder.DuplicateGraphOutput(SourceMemberName, Name))
 							{
 								FGraphBuilder::SynchronizeGraphMembers(DocBuilder, Graph);
@@ -4048,7 +4024,7 @@ namespace Metasound
 				else if (const UMetasoundEditorGraphVariable* SourceGraphVariable = Cast<UMetasoundEditorGraphVariable>(Member))
 				{
 					const FMetasoundEditorGraphVariableBreadcrumb Breadcrumb = SourceGraphVariable->GetBreadcrumb();
-					const FName VariableName = DocumentBuilder.GenerateUniqueVariableName(Breadcrumb.MemberName.ToString());
+					const FName VariableName = FGraphBuilder::GenerateUniqueVariableName(DocumentBuilder, Breadcrumb.MemberName.ToString());
 					const FMetasoundFrontendVariable* FrontendVariable = DocumentBuilder.AddGraphVariable(
 						VariableName,
 						SourceGraphVariable->GetDataType(),
@@ -4361,14 +4337,14 @@ PRAGMA_ENABLE_DEPRECATION_WARNINGS
 
 		void FEditor::CollectStaticSections(TArray<int32>& StaticSectionIDs)
 		{
-			const bool bIsReadOnly = Builder.IsValid() ? !Builder->GetConstBuilder().GetIsGraphEditable() : false;
+			const bool bIsPreset = Builder.IsValid() ? Builder->IsPreset() : false;
 
 			for (int32 i = 0; i < static_cast<int32>(ENodeSection::COUNT); ++i)
 			{
 				if (static_cast<ENodeSection>(i) != ENodeSection::None)
 				{
-					// ReadOnly graphs do not show variables
-					if (bIsReadOnly && static_cast<ENodeSection>(i) == ENodeSection::Variables)
+					// Presets do not have variables
+					if (bIsPreset && static_cast<ENodeSection>(i) == ENodeSection::Variables)
 					{
 						continue;
 					}
@@ -4612,9 +4588,9 @@ PRAGMA_ENABLE_DEPRECATION_WARNINGS
 			TArray<TSharedPtr<FEdGraphSchemaAction>> Actions;
 			GraphMembersMenu->GetSelectedActions(Actions);
 
-			if (Actions.IsEmpty()) //Section is selected
+			if (Actions.IsEmpty())//Section is selected
 			{
-				if (Builder->GetConstBuilder().GetIsGraphEditable())
+				if (!Builder->IsPreset())
 				{
 					MenuBuilder.BeginSection("GraphActionMenuSectionActions", LOCTEXT("SectionActionsMenuHeader", "Section Actions"));
 
@@ -5053,7 +5029,7 @@ PRAGMA_ENABLE_DEPRECATION_WARNINGS
 					MetaSound->Modify();
 					Graph.Modify();
 
-					const FName VariableName = DocBuilder.GenerateUniqueVariableName(TEXT("Variable"));
+					const FName VariableName = FGraphBuilder::GenerateUniqueVariableName(DocBuilder, TEXT("Variable"));
 					const FMetasoundFrontendVariable* FrontendVariable = DocBuilder.AddGraphVariable(VariableName, DataTypeName);
 					if (ensure(FrontendVariable))
 					{
@@ -5525,7 +5501,7 @@ PRAGMA_ENABLE_DEPRECATION_WARNINGS
 					const FMetasoundFrontendVertex* InputVertex = DocBuilder.FindNodeInput(InputVertexHandle.NodeID, InputVertexHandle.VertexID);
 					check(InputVertex);
 
-					FName Name = DocBuilder.GenerateUniqueInputName(Pin->GetName());
+					FName Name = FGraphBuilder::GenerateUniqueNameByClassType(ParentMetasound, EMetasoundFrontendClassType::Input, Pin->GetName());
 
 					EMetaSoundBuilderResult Result = EMetaSoundBuilderResult::Failed;
 					FMetasoundFrontendLiteral Literal;
@@ -5632,7 +5608,7 @@ PRAGMA_ENABLE_DEPRECATION_WARNINGS
 				FName TypeName = Pair.Key.Value;
 				UEdGraphPin* SourcePin = Pair.Value[0];
 				FMetasoundFrontendVertexHandle InputVertexHandle = FGraphBuilder::GetPinVertexHandle(DocBuilder, SourcePin);
-				FName InputName = DocBuilder.GenerateUniqueInputName(PinName.ToString());
+				FName InputName = FGraphBuilder::GenerateUniqueNameByClassType(ParentMetasound, EMetasoundFrontendClassType::Input, PinName.ToString());
 				
 				EMetaSoundBuilderResult Result = EMetaSoundBuilderResult::Failed;
 				FMetasoundFrontendLiteral Literal;
diff --git a/Engine/Plugins/Runtime/Metasound/Source/MetasoundEditor/Private/MetasoundEditor.h b/Engine/Plugins/Runtime/Metasound/Source/MetasoundEditor/Private/MetasoundEditor.h
--- a/Engine/Plugins/Runtime/Metasound/Source/MetasoundEditor/Private/MetasoundEditor.h
+++ b/Engine/Plugins/Runtime/Metasound/Source/MetasoundEditor/Private/MetasoundEditor.h
@@ -49,7 +49,6 @@ class IToolkitHost;
 class UUserWidget;
 
 struct FMetaSoundFrontendDocumentBuilder;
-struct FMetaSoundFrontendDocumentConfiguration;
 
 namespace Metasound::Editor
 {
@@ -136,7 +135,6 @@ namespace Metasound::Editor
 
 		UAudioComponent* GetAudioComponent() const;
 
-		const FMetaSoundFrontendDocumentConfiguration* GetConfiguration() const;
 		FMetaSoundFrontendDocumentBuilder* GetFrontendBuilder() const;
 
 		/** IMetasoundEditor interface */
@@ -347,7 +345,7 @@ namespace Metasound::Editor
 		bool CanCopySelectedMemberItems() const;
 
 		/** Copy selected Member items. */
-		void CopySelectedMemberItems();
+		void CopySelectedMemberItems();		
 		
 		/** Whether the currently selected Member item(s) can be cut. */
 		bool CanCutSelectedMemberItems() const;
@@ -495,7 +493,7 @@ namespace Metasound::Editor
 		/** Converts the MetaSound from a preset to a fully modifiable MetaSound. */
 		void ConvertFromPreset();
 
-		/* Whether or not page details should be visible/enabled. */
+		/* Whether or not page details should be visible. */
 		bool ShowPageGraphDetails() const;
 
 		/** Creates audition menu options */
diff --git a/Engine/Plugins/Runtime/Metasound/Source/MetasoundEditor/Private/MetasoundEditorDocumentClipboardUtils.cpp b/Engine/Plugins/Runtime/Metasound/Source/MetasoundEditor/Private/MetasoundEditorDocumentClipboardUtils.cpp
--- a/Engine/Plugins/Runtime/Metasound/Source/MetasoundEditor/Private/MetasoundEditorDocumentClipboardUtils.cpp
+++ b/Engine/Plugins/Runtime/Metasound/Source/MetasoundEditor/Private/MetasoundEditorDocumentClipboardUtils.cpp
@@ -266,7 +266,7 @@ namespace Metasound::Editor
 				else
 				{
 					const FMetasoundFrontendLiteral* Literal = Breadcrumb.DefaultLiterals.Find(Frontend::DefaultPageID);
-					const FName VariableName = DocBuilder.GenerateUniqueVariableName(BaseName.ToString());
+					const FName VariableName = FGraphBuilder::GenerateUniqueVariableName(DocBuilder, BaseName.ToString());
 					FrontendVariable = DocBuilder.AddGraphVariable(
 						VariableName,
 						Breadcrumb.DataType,
diff --git a/Engine/Plugins/Runtime/Metasound/Source/MetasoundEditor/Private/MetasoundEditorGraph.cpp b/Engine/Plugins/Runtime/Metasound/Source/MetasoundEditor/Private/MetasoundEditorGraph.cpp
--- a/Engine/Plugins/Runtime/Metasound/Source/MetasoundEditor/Private/MetasoundEditorGraph.cpp
+++ b/Engine/Plugins/Runtime/Metasound/Source/MetasoundEditor/Private/MetasoundEditorGraph.cpp
@@ -376,7 +376,7 @@ bool UMetasoundEditorGraphVertex::SetIsAdvancedDisplay(const bool IsAdvancedDisp
 		return false;
 	}
 
-	const FText TransactionLabel = FText::Format(LOCTEXT("SetGraphVertexIsAdvancedDisplayState", "Set Metasound {0} IsAdvancedDisplay"), GetGraphMemberLabel());
+	const FText TransactionLabel = FText::Format(LOCTEXT("SetGraphVertexIsAdvancedDisplayState", "Set Metasound {0} IsAdvancedDislay"), GetGraphMemberLabel());
 	const FScopedTransaction Transaction(TransactionLabel, true);
 
 	Graph->Modify();
@@ -500,7 +500,7 @@ bool UMetasoundEditorGraphVertex::NameContainsInterfaceNamespace(FMetasoundFront
 bool UMetasoundEditorGraphVertex::CanRename() const
 {
 	FMetaSoundFrontendDocumentBuilder& Builder = GetFrontendBuilderChecked();
-	return Builder.GetIsGraphEditable() && !IsInterfaceMember();
+	return !Builder.IsPreset() && !IsInterfaceMember();
 }
 
 bool UMetasoundEditorGraphVertex::CanRename(const FText& InNewText, FText& OutError) const
@@ -533,9 +533,9 @@ bool UMetasoundEditorGraphVertex::CanRename(const FText& InNewText, FText& OutEr
 	}
 
 	const FMetaSoundFrontendDocumentBuilder& Builder = GetFrontendBuilderChecked();
-	if (!Builder.GetIsGraphEditable())
+	if (Builder.IsPreset())
 	{
-		OutError = FText::Format(LOCTEXT("GraphVertexRenameInvalid_Preset", "{0} is a vertex in a read-only graph and cannot be renamed."), InNewText);
+		OutError = FText::Format(LOCTEXT("GraphVertexRenameInvalid_Preset", "{0} is a vertex in a preset graph and cannot be renamed."), InNewText);
 		return false;
 	}
 
@@ -2601,14 +2601,12 @@ void UMetasoundEditorGraph::MigrateEditorDocumentData(FMetaSoundFrontendDocument
 		}
 #endif // WITH_EDITOR
 
-PRAGMA_DISABLE_DEPRECATION_WARNINGS
 		// Presets get rebuilt anyway and may have invalid connections (newly referenced vertices)
 		// that need to be rebuilt later in asset load.
-		if (OutBuilder.GetConstDocumentChecked().RootGraph.PresetOptions.bIsPreset)
+		if (OutBuilder.IsPreset())
 		{
 			return;
 		}
-PRAGMA_ENABLE_DEPRECATION_WARNINGS
 
 		// Cache data to be used in edge swapping below, being careful to not reference the input node pointer
 		// or vertex reference within the lower inner loop as the underlying node array may be reallocated by
@@ -2711,12 +2709,10 @@ PRAGMA_ENABLE_DEPRECATION_WARNINGS
 		}
 	});
 
-PRAGMA_DISABLE_DEPRECATION_WARNINGS
 	// 4. Add comment nodes as builder graph comments to frontend document
 	// (No need to propagate comments for presets)
-	if (!OutBuilder.GetConstDocumentChecked().RootGraph.PresetOptions.bIsPreset)
+	if (!OutBuilder.IsPreset())
 	{
-PRAGMA_ENABLE_DEPRECATION_WARNINGS
 		TArray<UEdGraphNode_Comment*> CommentNodes;
 		GetNodesOfClass(CommentNodes);
 		for (UEdGraphNode_Comment* Node : CommentNodes)
diff --git a/Engine/Plugins/Runtime/Metasound/Source/MetasoundEditor/Private/MetasoundEditorGraphBuilder.cpp b/Engine/Plugins/Runtime/Metasound/Source/MetasoundEditor/Private/MetasoundEditorGraphBuilder.cpp
--- a/Engine/Plugins/Runtime/Metasound/Source/MetasoundEditor/Private/MetasoundEditorGraphBuilder.cpp
+++ b/Engine/Plugins/Runtime/Metasound/Source/MetasoundEditor/Private/MetasoundEditorGraphBuilder.cpp
@@ -94,6 +94,20 @@ namespace Metasound
 				return nullptr;
 			}
 
+			FName GenerateUniqueName(const TSet<FName>& InExistingNames, const FString& InBaseName)
+			{
+				int32 PostFixInt = 0;
+				FString NewName = InBaseName;
+
+				while (InExistingNames.Contains(*NewName))
+				{
+					PostFixInt++;
+					NewName = FString::Format(TEXT("{0} {1}"), { InBaseName, PostFixInt });
+				}
+
+				return FName(*NewName);
+			}
+
 			void RecurseClearDocumentModified(FMetasoundAssetBase& InAssetBase)
 			{
 				using namespace Metasound::Frontend;
@@ -614,6 +628,18 @@ namespace Metasound
 			return Categories;
 		}
 
+		FName FGraphBuilder::GenerateUniqueNameByClassType(UObject& InMetaSound, EMetasoundFrontendClassType InClassType, const FString& InBaseName)
+		{
+			using namespace Engine;
+
+			TSet<FName> ExistingNames;
+			FMetaSoundFrontendDocumentBuilder& Builder = FDocumentBuilderRegistry::GetChecked().FindOrBeginBuilding(&InMetaSound);
+			auto GetName = [&](const FMetasoundFrontendClass&, const FMetasoundFrontendNode& Node) { ExistingNames.Add(Node.Name); };
+			Builder.IterateNodesByClassType(GetName, InClassType);
+
+			return GraphBuilderPrivate::GenerateUniqueName(ExistingNames, InBaseName);
+		}
+
 		TSharedPtr<FEditor> FGraphBuilder::GetEditorForMetasound(const UObject& Metasound)
 		{
 			// TODO: FToolkitManager is deprecated. Replace with UAssetEditorSubsystem.
@@ -984,17 +1010,13 @@ PRAGMA_ENABLE_DEPRECATION_WARNINGS
 		bool FGraphBuilder::IsPreviewingMetaSound(const UObject& InMetaSound)
 		{
 			using namespace Engine;
-
-			if (GEditor)
+			if (const UAudioComponent* PreviewComponent = GEditor->GetPreviewAudioComponent())
 			{
-				if (const UAudioComponent* PreviewComponent = GEditor->GetPreviewAudioComponent())
+				if (PreviewComponent->IsPlaying())
 				{
-					if (PreviewComponent->IsPlaying())
+					if (const USoundBase* Sound = PreviewComponent->Sound)
 					{
-						if (const USoundBase* Sound = PreviewComponent->Sound)
-						{
-							return Sound->GetUniqueID() == InMetaSound.GetUniqueID();
-						}
+						return Sound->GetUniqueID() == InMetaSound.GetUniqueID();
 					}
 				}
 			}
@@ -1013,9 +1035,9 @@ PRAGMA_ENABLE_DEPRECATION_WARNINGS
 					if (ExternalNode->GetBreadcrumb().ClassName == FRerouteNodeTemplate::ClassName)
 					{
 						auto IsInput = [](UEdGraphPin* Pin) { check(Pin); return Pin->Direction == EGPD_Input; };
-						if (UEdGraphPin** RerouteInput = ExternalNode->Pins.FindByPredicate(IsInput); RerouteInput && *RerouteInput)
+						if (UEdGraphPin* RerouteInput = *ExternalNode->Pins.FindByPredicate(IsInput))
 						{
-							TArray<UEdGraphPin*>& LinkedTo = (*RerouteInput)->LinkedTo;
+							TArray<UEdGraphPin*>& LinkedTo = RerouteInput->LinkedTo;
 							if (!LinkedTo.IsEmpty())
 							{
 								UEdGraphPin* ReroutedOutput = LinkedTo.Last();
@@ -1040,9 +1062,9 @@ PRAGMA_ENABLE_DEPRECATION_WARNINGS
 					if (ExternalNode->GetBreadcrumb().ClassName == FRerouteNodeTemplate::ClassName)
 					{
 						auto IsInput = [](const UEdGraphPin* Pin) { check(Pin); return Pin->Direction == EGPD_Input; };
-						if (const UEdGraphPin* const* RerouteInput = ExternalNode->Pins.FindByPredicate(IsInput))
+						if (const UEdGraphPin* RerouteInput = *ExternalNode->Pins.FindByPredicate(IsInput))
 						{
-							const TArray<UEdGraphPin*>& LinkedTo = (*RerouteInput)->LinkedTo;
+							const TArray<UEdGraphPin*>& LinkedTo = RerouteInput->LinkedTo;
 							if (!LinkedTo.IsEmpty())
 							{
 								const UEdGraphPin* ReroutedOutput = LinkedTo.Last();
@@ -1067,9 +1089,9 @@ PRAGMA_ENABLE_DEPRECATION_WARNINGS
 					if (ExternalNode->GetBreadcrumb().ClassName == FRerouteNodeTemplate::ClassName)
 					{
 						auto IsInput = [](UEdGraphPin* Pin) { check(Pin); return Pin->Direction == EGPD_Input; };
-						if (const UEdGraphPin* const* RerouteInput = ExternalNode->Pins.FindByPredicate(IsInput))
+						if (const UEdGraphPin* RerouteInput = *ExternalNode->Pins.FindByPredicate(IsInput))
 						{
-							const TArray<UEdGraphPin*>& LinkedTo = (*RerouteInput)->LinkedTo;
+							const TArray<UEdGraphPin*>& LinkedTo = RerouteInput->LinkedTo;
 							if (!LinkedTo.IsEmpty())
 							{
 								const UEdGraphPin* ReroutedOutput = LinkedTo.Last();
@@ -1290,11 +1312,8 @@ PRAGMA_ENABLE_DEPRECATION_WARNINGS
 
 		FMetasoundFrontendClassInput FGraphBuilder::CreateUniqueClassInput(UObject& InMetaSound, const FCreateNodeVertexParams& InParams, const TArray<FMetasoundFrontendClassInputDefault>& InDefaultLiterals, const FName* InNameBase)
 		{
-			using namespace Engine;
-
 			FMetasoundFrontendClassInput ClassInput;
-			UMetaSoundBuilderBase& Builder = FDocumentBuilderRegistry::GetChecked().FindOrBeginBuilding(InMetaSound);
-			ClassInput.Name = Builder.GetConstBuilder().GenerateUniqueInputName(InNameBase ? InNameBase->ToString() : TEXT("Input"));
+			ClassInput.Name = GenerateUniqueNameByClassType(InMetaSound, EMetasoundFrontendClassType::Input, InNameBase ? InNameBase->ToString() : TEXT("Input"));
 			ClassInput.TypeName = InParams.DataType;
 			ClassInput.VertexID = FGuid::NewGuid();
 			ClassInput.NodeID = FGuid::NewGuid();
@@ -1331,11 +1350,8 @@ PRAGMA_ENABLE_DEPRECATION_WARNINGS
 
 		FMetasoundFrontendClassOutput FGraphBuilder::CreateUniqueClassOutput(UObject& InMetaSound, const FCreateNodeVertexParams& InParams, const FName* InNameBase)
 		{
-			using namespace Metasound::Engine;
-
 			FMetasoundFrontendClassOutput ClassOutput;
-			UMetaSoundBuilderBase& Builder = FDocumentBuilderRegistry::GetChecked().FindOrBeginBuilding(InMetaSound);
-			ClassOutput.Name = Builder.GetConstBuilder().GenerateUniqueOutputName(InNameBase ? InNameBase->ToString() : TEXT("Output"));
+			ClassOutput.Name = GenerateUniqueNameByClassType(InMetaSound, EMetasoundFrontendClassType::Output, InNameBase ? InNameBase->ToString() : TEXT("Output"));
 			ClassOutput.TypeName = InParams.DataType;
 			ClassOutput.VertexID = FGuid::NewGuid();
 			ClassOutput.NodeID = FGuid::NewGuid();
@@ -1346,6 +1362,19 @@ PRAGMA_ENABLE_DEPRECATION_WARNINGS
 			return ClassOutput;
 		}
 
+		FName FGraphBuilder::GenerateUniqueVariableName(const FMetaSoundFrontendDocumentBuilder& InBuilder, const FString& InBaseName)
+		{
+			using namespace Frontend;
+
+			TSet<FName> ExistingVariableNames;
+
+			// Get all the names from the existing variables on
+			// the build graph and place into the ExistingVariableNames array.
+			Algo::Transform(InBuilder.FindConstBuildGraphChecked().Variables, ExistingVariableNames, [](const FMetasoundFrontendVariable& Var) { return Var.Name; });
+
+			return GraphBuilderPrivate::GenerateUniqueName(ExistingVariableNames, InBaseName);
+		}
+
 		FMetasoundAssetBase& FGraphBuilder::GetOutermostMetaSoundChecked(UObject& InSubObject)
 		{
 			FMetasoundAssetBase* MetaSoundAsset = IMetasoundUObjectRegistry::Get().GetObjectAsAssetBase(InSubObject.GetOutermostObject());
@@ -1925,24 +1954,21 @@ PRAGMA_ENABLE_DEPRECATION_WARNINGS
 				TSet<FMetasoundAssetBase*> EditedReferencingMetaSounds;
 				if (GEditor)
 				{
-					if (UAssetEditorSubsystem* AssetSubsystem = GEditor->GetEditorSubsystem<UAssetEditorSubsystem>())
+					TArray<UObject*> EditedAssets = GEditor->GetEditorSubsystem<UAssetEditorSubsystem>()->GetAllEditedAssets();
+					for (UObject* Asset : EditedAssets)
 					{
-						TArray<UObject*> EditedAssets = AssetSubsystem->GetAllEditedAssets();
-						for (UObject* Asset : EditedAssets)
+						if (Asset != &MetaSound)
 						{
-							if (Asset != &MetaSound)
+							if (FMetasoundAssetBase* EditedMetaSound = IMetasoundUObjectRegistry::Get().GetObjectAsAssetBase(Asset))
 							{
-								if (FMetasoundAssetBase* EditedMetaSound = IMetasoundUObjectRegistry::Get().GetObjectAsAssetBase(Asset))
+								if (EditedMetaSound->IsReferencedAsset(*MetaSoundAsset))
 								{
-									if (EditedMetaSound->IsReferencedAsset(*MetaSoundAsset))
+									// Ensure the graph has been initialized. Its possible the "edited asset" is a MetaSound
+									// but its transient EdGraph has yet to be initialized due to the active editor not necessarily
+									// being the MetaSound editor (ex. property matrix editor).
+									if (const UEdGraph* EdGraph = EditedMetaSound->GetGraph())
 									{
-										// Ensure the graph has been initialized. Its possible the "edited asset" is a MetaSound
-										// but its transient EdGraph has yet to be initialized due to the active editor not necessarily
-										// being the MetaSound editor (ex. property matrix editor).
-										if (const UEdGraph* EdGraph = EditedMetaSound->GetGraph())
-										{
-											EditedReferencingMetaSounds.Add(EditedMetaSound);
-										}
+										EditedReferencingMetaSounds.Add(EditedMetaSound);
 									}
 								}
 							}
@@ -2134,7 +2160,7 @@ PRAGMA_ENABLE_DEPRECATION_WARNINGS
 					{
 						case EMetasoundFrontendClassType::External:
 						{
-							NewGraphNode = Cast<UMetasoundEditorGraphNode>(AddExternalNode(MetaSoundObject, NodeID, ClassMetadata, false));
+							NewGraphNode = CastChecked<UMetasoundEditorGraphNode>(AddExternalNode(MetaSoundObject, NodeID, ClassMetadata, false));
 						}
 						break;
 
@@ -2142,18 +2168,18 @@ PRAGMA_ENABLE_DEPRECATION_WARNINGS
 						{
 							if (Pair.Value->Metadata.GetClassName() == FInputNodeTemplate::ClassName)
 							{
-								NewGraphNode = Cast<UMetasoundEditorGraphNode>(AddInputNode(MetaSoundObject, NodeID, false));
+								NewGraphNode = CastChecked<UMetasoundEditorGraphNode>(AddInputNode(MetaSoundObject, NodeID, false));
 							}
 							else
 							{
-								NewGraphNode = Cast<UMetasoundEditorGraphNode>(AddTemplateNode(MetaSoundObject, NodeID, ClassMetadata, false));
+								NewGraphNode = CastChecked<UMetasoundEditorGraphNode>(AddTemplateNode(MetaSoundObject, NodeID, ClassMetadata, false));
 							}
 						}
 						break;
 
 						case EMetasoundFrontendClassType::Output:
 						{
-							NewGraphNode = Cast<UMetasoundEditorGraphNode>(AddOutputNode(MetaSoundObject, NodeID, false));
+							NewGraphNode = CastChecked<UMetasoundEditorGraphNode>(AddOutputNode(MetaSoundObject, NodeID, false));
 						}
 						break;
 
@@ -2162,7 +2188,7 @@ PRAGMA_ENABLE_DEPRECATION_WARNINGS
 						case EMetasoundFrontendClassType::VariableDeferredAccessor:
 						case EMetasoundFrontendClassType::Variable:
 						{
-							NewGraphNode = Cast<UMetasoundEditorGraphNode>(AddVariableNode(MetaSoundObject, NodeID, false));
+							NewGraphNode = CastChecked<UMetasoundEditorGraphNode>(AddVariableNode(MetaSoundObject, NodeID, false));
 						}
 						break;
 
diff --git a/Engine/Plugins/Runtime/Metasound/Source/MetasoundEditor/Private/MetasoundEditorGraphBuilder.h b/Engine/Plugins/Runtime/Metasound/Source/MetasoundEditor/Private/MetasoundEditorGraphBuilder.h
--- a/Engine/Plugins/Runtime/Metasound/Source/MetasoundEditor/Private/MetasoundEditorGraphBuilder.h
+++ b/Engine/Plugins/Runtime/Metasound/Source/MetasoundEditor/Private/MetasoundEditorGraphBuilder.h
@@ -143,6 +143,9 @@ namespace Metasound
 			// Adds an output node to the editor graph that corresponds to the provided output node ID.
 			static UMetasoundEditorGraphOutputNode* AddOutputNode(UObject& InMetaSound, const FGuid& InNodeID, bool bInSelectNewNode = true);
 
+			// Create a unique name for the variable.
+			static FName GenerateUniqueVariableName(const FMetaSoundFrontendDocumentBuilder& InBuilder, const FString& InBaseName);
+
 			// Convenience method for walking to the outermost object and transforming to a base MetaSound
 			static FMetasoundAssetBase& GetOutermostMetaSoundChecked(UObject& InSubObject);
 			static const FMetasoundAssetBase& GetOutermostConstMetaSoundChecked(const UObject& InSubObject);
@@ -170,6 +173,9 @@ namespace Metasound
 			// disconnect the EdGraph pins.
 			static void DisconnectPinVertex(UEdGraphPin& InPin);
 
+			// Generates a unique output name for the given MetaSound object
+			static FName GenerateUniqueNameByClassType(UObject& InMetaSound, EMetasoundFrontendClassType InClassType, const FString& InBaseName);
+
 			static UMetaSoundBuilderBase& GetBuilderFromPinChecked(const UEdGraphPin& InPin);
 
 			static TArray<FString> GetDataTypeNameCategories(const FName& InDataTypeName);
diff --git a/Engine/Plugins/Runtime/Metasound/Source/MetasoundEditor/Private/MetasoundEditorGraphNode.cpp b/Engine/Plugins/Runtime/Metasound/Source/MetasoundEditor/Private/MetasoundEditorGraphNode.cpp
--- a/Engine/Plugins/Runtime/Metasound/Source/MetasoundEditor/Private/MetasoundEditorGraphNode.cpp
+++ b/Engine/Plugins/Runtime/Metasound/Source/MetasoundEditor/Private/MetasoundEditorGraphNode.cpp
@@ -774,7 +774,7 @@ void UMetasoundEditorGraphNode::GetNodeContextMenuActions(UToolMenu* Menu, UGrap
 
 	if (Context->Node)
 	{
-		if (GetBuilderChecked().GetConstBuilder().GetIsGraphEditable())
+		if (!GetBuilderChecked().IsPreset())
 		{
 			{
 				FToolMenuSection& Section = Menu->FindOrAddSection("MetasoundGraphNodeActionsOrganization", LOCTEXT("NodeActionsOrganizationMenuHeader", "Organization"));
diff --git a/Engine/Plugins/Runtime/Metasound/Source/MetasoundEditor/Private/MetasoundEditorGraphSchema.cpp b/Engine/Plugins/Runtime/Metasound/Source/MetasoundEditor/Private/MetasoundEditorGraphSchema.cpp
--- a/Engine/Plugins/Runtime/Metasound/Source/MetasoundEditor/Private/MetasoundEditorGraphSchema.cpp
+++ b/Engine/Plugins/Runtime/Metasound/Source/MetasoundEditor/Private/MetasoundEditorGraphSchema.cpp
@@ -290,7 +290,7 @@ PRAGMA_ENABLE_DEPRECATION_WARNINGS
 
 			UMetasoundEditorGraph* MetaSoundGraph = CastChecked<UMetasoundEditorGraph>(ConnectedNode->GetGraph());
 			FMetaSoundFrontendDocumentBuilder& DocBuilder = MetaSoundGraph->GetBuilderChecked().GetBuilder();
-			const FName NodeName = DocBuilder.GenerateUniqueVariableName(BaseName.ToString());
+			const FName NodeName = FGraphBuilder::GenerateUniqueVariableName(DocBuilder, BaseName.ToString());
 
 			const FScopedTransaction Transaction(FText::Format(
 				LOCTEXT("PromoteNodeVertexToGraphVariableFormat", "Promote MetaSound Node {0} to {1}"),
@@ -1649,7 +1649,7 @@ void UMetasoundEditorGraphSchema::GetContextMenuActions(class UToolMenu* Menu, c
 		const UMetasoundEditorGraph* EdGraph = CastChecked<UMetasoundEditorGraph>(Context->Graph);
 		UMetaSoundBuilderBase& Builder = Metasound::Engine::FDocumentBuilderRegistry::GetChecked().FindOrBeginBuilding(EdGraph->GetMetasoundChecked());
 
-		if (Builder.GetConstBuilder().GetIsGraphEditable())
+		if (!Builder.IsPreset())
 		{
 			FToolMenuSection& Section = Menu->FindOrAddSection("Pin Actions");
 
@@ -2807,7 +2807,7 @@ void UMetasoundEditorGraphSchema::QueryNodeClasses(TFunctionRef<void(TUniquePtr<
 
 	// Gather all loaded classes
 	{
-		constexpr bool bIncludeUnloadedAssets = true;
+		static bool bIncludeUnloadedAssets = true;
 		const ISearchEngine::EResultVersion Options = SchemaPrivate::EnableAllVersionsMetaSoundNodeClassCreationCVar == 0
 			? ISearchEngine::EResultVersion::Highest
 			: ISearchEngine::EResultVersion::All;
diff --git a/Engine/Plugins/Runtime/Metasound/Source/MetasoundEditor/Public/MetasoundEditorGraph.h b/Engine/Plugins/Runtime/Metasound/Source/MetasoundEditor/Public/MetasoundEditorGraph.h
--- a/Engine/Plugins/Runtime/Metasound/Source/MetasoundEditor/Public/MetasoundEditorGraph.h
+++ b/Engine/Plugins/Runtime/Metasound/Source/MetasoundEditor/Public/MetasoundEditorGraph.h
@@ -107,15 +107,15 @@ public:
 		return Value;
 	}
 
-	UE_API virtual void ForceRefresh();
+	UE_API virtual void ForceRefresh() override;
 	UE_API virtual FName GetDataType() const;
-	UE_API virtual EMetasoundFrontendLiteralType GetLiteralType() const;
+	UE_API virtual EMetasoundFrontendLiteralType GetLiteralType() const override;
 	UE_API virtual void InitDefault(const FGuid& InPageID = Metasound::Frontend::DefaultPageID);
 	virtual void Initialize() { };
 	UE_API virtual void IterateDefaults(TFunctionRef<void(const FGuid&, FMetasoundFrontendLiteral)> Iter) const;
 	UE_API virtual bool RemoveDefault(const FGuid& InPageID);
 	UE_API virtual void ResetDefaults();
-	UE_API virtual void SetFromLiteral(const FMetasoundFrontendLiteral& InLiteral, const FGuid& InPageID = Metasound::Frontend::DefaultPageID);
+	UE_API virtual void SetFromLiteral(const FMetasoundFrontendLiteral& InLiteral, const FGuid& InPageID = Metasound::Frontend::DefaultPageID) override;
 	UE_API virtual bool TryFindDefault(FMetasoundFrontendLiteral& OutLiteral, const FGuid* InPageID = nullptr) const;
 
 	virtual void UpdatePreviewInstance(const Metasound::FVertexName& InParameterName, TScriptInterface<IAudioParameterControllerInterface>& InParameterInterface) const
diff --git a/Engine/Plugins/Runtime/Metasound/Source/MetasoundEngine/Private/DocumentConfigurations/MetasoundDocumentConfigurator.cpp b/Engine/Plugins/Runtime/Metasound/Source/MetasoundEngine/Private/DocumentConfigurations/MetasoundDocumentConfigurator.cpp
deleted file mode 100644
diff --git a/Engine/Plugins/Runtime/Metasound/Source/MetasoundEngine/Private/DocumentConfigurations/MetasoundRandomizerDocumentConfiguration.cpp b/Engine/Plugins/Runtime/Metasound/Source/MetasoundEngine/Private/DocumentConfigurations/MetasoundRandomizerDocumentConfiguration.cpp
deleted file mode 100644
diff --git a/Engine/Plugins/Runtime/Metasound/Source/MetasoundEngine/Private/MetasoundBuilderBase.cpp b/Engine/Plugins/Runtime/Metasound/Source/MetasoundEngine/Private/MetasoundBuilderBase.cpp
--- a/Engine/Plugins/Runtime/Metasound/Source/MetasoundEngine/Private/MetasoundBuilderBase.cpp
+++ b/Engine/Plugins/Runtime/Metasound/Source/MetasoundEngine/Private/MetasoundBuilderBase.cpp
@@ -488,15 +488,9 @@ TArray<FMetaSoundBuilderNodeInputHandle> UMetaSoundBuilderBase::ConnectNodeOutpu
 	TArray<FMetaSoundBuilderNodeInputHandle> ConnectedVertices;
 	Algo::Transform(NewEdges, ConnectedVertices, [this](const FMetasoundFrontendEdge* NewEdge)
 	{
-#if DO_CHECK
-		{
-			check(NewEdge);
-			const FMetasoundFrontendVertex* Vertex = Builder.FindNodeInput(NewEdge->ToNodeID, NewEdge->ToVertexID);
-			checkf(Vertex, TEXT("Edge node output to interface output connection reported success but vertex not found."));
-		}
-#endif // DO_CHECK
-
-		return FMetaSoundBuilderNodeInputHandle(NewEdge->ToNodeID, NewEdge->ToVertexID);
+		const FMetasoundFrontendVertex* Vertex = Builder.FindNodeInput(NewEdge->ToNodeID, NewEdge->ToVertexID);
+		checkf(Vertex, TEXT("Edge connection reported success but vertex not found."));
+		return FMetaSoundBuilderNodeInputHandle(NewEdge->ToNodeID, Vertex->VertexID);
 	});
 
 	return ConnectedVertices;
@@ -511,15 +505,9 @@ TArray<FMetaSoundBuilderNodeOutputHandle> UMetaSoundBuilderBase::ConnectNodeInpu
 	TArray<FMetaSoundBuilderNodeOutputHandle> ConnectedVertices;
 	Algo::Transform(NewEdges, ConnectedVertices, [this](const FMetasoundFrontendEdge* NewEdge)
 	{
-#if DO_CHECK
-		{
-			check(NewEdge);
-			const FMetasoundFrontendVertex* Vertex = Builder.FindNodeOutput(NewEdge->FromNodeID, NewEdge->FromVertexID);
-			checkf(Vertex, TEXT("Edge node input to interface input connection reported success but vertex not found."));
-		}
-#endif // DO_CHECK
-
-		return FMetaSoundBuilderNodeOutputHandle(NewEdge->FromNodeID, NewEdge->FromVertexID);
+		const FMetasoundFrontendVertex* Vertex = Builder.FindNodeOutput(NewEdge->FromNodeID, NewEdge->FromVertexID);
+		checkf(Vertex, TEXT("Edge connection reported success but vertex not found."));
+		return FMetaSoundBuilderNodeOutputHandle(NewEdge->ToNodeID, Vertex->VertexID);
 	});
 
 	return ConnectedVertices;
@@ -929,6 +917,13 @@ UMetaSoundFrontendMemberMetadata* UMetaSoundBuilderBase::FindMemberMetadata(cons
 }
 #endif // WITH_EDITOR
 
+UMetaSoundBuilderDocument* UMetaSoundBuilderBase::CreateTransientDocumentObject() const
+{
+PRAGMA_DISABLE_DEPRECATION_WARNINGS
+	return &UMetaSoundBuilderDocument::Create(GetBaseMetaSoundUClass());
+PRAGMA_ENABLE_DEPRECATION_WARNINGS
+}
+
 FMetaSoundFrontendDocumentBuilder& UMetaSoundBuilderBase::GetBuilder()
 {
 	return Builder;
@@ -975,8 +970,6 @@ TArray<FName> UMetaSoundBuilderBase::GetGraphInputNames(EMetaSoundBuilderResult&
 	{
 		GraphInputNames.Emplace(Input.Name);
 	}
-
-	OutResult = EMetaSoundBuilderResult::Succeeded;
 	return GraphInputNames;
 }
 
@@ -988,8 +981,6 @@ TArray<FName> UMetaSoundBuilderBase::GetGraphOutputNames(EMetaSoundBuilderResult
 	{
 		GraphOutputNames.Emplace(Output.Name);
 	}
-
-	OutResult = EMetaSoundBuilderResult::Succeeded;
 	return GraphOutputNames;
 }
 
@@ -1060,6 +1051,11 @@ bool UMetaSoundBuilderBase::InterfaceIsDeclared(FName InterfaceName) const
 	return Builder.IsInterfaceDeclared(InterfaceName);
 }
 
+void UMetaSoundBuilderBase::InvalidateCache(bool bPrimeCache)
+{
+	Reload({ }, bPrimeCache);
+}
+
 bool UMetaSoundBuilderBase::IsPreset() const
 {
 	return Builder.IsPreset();
diff --git a/Engine/Plugins/Runtime/Metasound/Source/MetasoundEngine/Private/MetasoundBuilderSubsystem.cpp b/Engine/Plugins/Runtime/Metasound/Source/MetasoundEngine/Private/MetasoundBuilderSubsystem.cpp
--- a/Engine/Plugins/Runtime/Metasound/Source/MetasoundEngine/Private/MetasoundBuilderSubsystem.cpp
+++ b/Engine/Plugins/Runtime/Metasound/Source/MetasoundEngine/Private/MetasoundBuilderSubsystem.cpp
@@ -138,6 +138,59 @@ void UMetaSoundPatchBuilder::OnRemovingAssetReference(TScriptInterface<IMetaSoun
 	Patch.ReferencedAssetClassKeys.Remove(RegistryKey.ToString());
 }
 
+UMetaSoundBuilderBase& UMetaSoundBuilderSubsystem::AttachBuilderToAssetChecked(UObject& InObject) const
+{
+	const UClass* BaseClass = InObject.GetClass();
+	if (BaseClass == UMetaSoundSource::StaticClass())
+	{
+PRAGMA_DISABLE_DEPRECATION_WARNINGS
+		UMetaSoundSourceBuilder* NewBuilder = AttachSourceBuilderToAsset(CastChecked<UMetaSoundSource>(&InObject));
+PRAGMA_ENABLE_DEPRECATION_WARNINGS
+		return *NewBuilder;
+	}
+	else if (BaseClass == UMetaSoundPatch::StaticClass())
+	{
+PRAGMA_DISABLE_DEPRECATION_WARNINGS
+		UMetaSoundPatchBuilder* NewBuilder = AttachPatchBuilderToAsset(CastChecked<UMetaSoundPatch>(&InObject));
+PRAGMA_ENABLE_DEPRECATION_WARNINGS
+		return *NewBuilder;
+	}
+	else
+	{
+		checkf(false, TEXT("UClass '%s' is not a base MetaSound that supports attachment via the MetaSoundBuilderSubsystem"), *BaseClass->GetFullName());
+		return *NewObject<UMetaSoundPatchBuilder>();
+	}
+}
+
+UMetaSoundPatchBuilder* UMetaSoundBuilderSubsystem::AttachPatchBuilderToAsset(UMetaSoundPatch* InPatch) const
+{
+#if WITH_EDITORONLY_DATA
+	using namespace Metasound::Engine;
+
+	if (InPatch)
+	{
+		return &FDocumentBuilderRegistry::GetChecked().FindOrBeginBuilding<UMetaSoundPatchBuilder>(*InPatch);
+	}
+#endif // WITH_EDITORONLY_DATA
+
+	return nullptr;
+}
+
+UMetaSoundSourceBuilder* UMetaSoundBuilderSubsystem::AttachSourceBuilderToAsset(UMetaSoundSource* InSource) const
+{
+#if WITH_EDITORONLY_DATA
+	using namespace Metasound::Engine;
+
+	if (InSource)
+	{
+		UMetaSoundSourceBuilder& SourceBuilder = FDocumentBuilderRegistry::GetChecked().FindOrBeginBuilding<UMetaSoundSourceBuilder>(*InSource);
+		return &SourceBuilder;
+	}
+#endif // WITH_EDITORONLY_DATA
+
+	return nullptr;
+}
+
 void UMetaSoundSourceBuilder::Audition(UObject* Parent, UAudioComponent* AudioComponent, FOnCreateAuditionGeneratorHandleDelegate CreateGenerator, bool bLiveUpdatesEnabled)
 {
 	using namespace Metasound;
@@ -352,10 +405,9 @@ void UMetaSoundSourceBuilder::InitDelegates(Metasound::Frontend::FDocumentModify
 
 	if (!IsRunningCookCommandlet())
 	{
-		OutDocumentDelegates.OnDocumentConfigurationChanged.AddUObject(this, &UMetaSoundSourceBuilder::OnDocumentConfigurationChanged);
-
 		OutDocumentDelegates.PageDelegates.OnPageAdded.AddUObject(this, &UMetaSoundSourceBuilder::OnPageAdded);
 		OutDocumentDelegates.PageDelegates.OnRemovingPage.AddUObject(this, &UMetaSoundSourceBuilder::OnRemovingPage);
+		OutDocumentDelegates.OnPresetStateChanged.AddUObject(this, &UMetaSoundSourceBuilder::OnPresetStateChanged);
 
 		OutDocumentDelegates.InterfaceDelegates.OnInputAdded.AddUObject(this, &UMetaSoundSourceBuilder::OnInputAdded);
 		OutDocumentDelegates.InterfaceDelegates.OnInputDataTypeChanged.AddUObject(this, &UMetaSoundSourceBuilder::OnInputDataTypeChanged);
@@ -436,13 +488,6 @@ void UMetaSoundSourceBuilder::OnAssetReferenceAdded(TScriptInterface<IMetaSoundD
 	Source.ReferencedAssetClassKeys.Add(RegistryKey.ToString());
 }
 
-void UMetaSoundSourceBuilder::OnDocumentConfigurationChanged(const Metasound::Frontend::FDocumentConfigChangedArgs&) const
-{
-	using namespace Metasound::Frontend;
-
-	GetMetaSoundSource().InvalidateCachedRuntimeInputData();
-}
-
 void UMetaSoundSourceBuilder::OnEdgeAdded(int32 EdgeIndex) const
 {
 	using namespace Metasound::DynamicGraph;
@@ -481,10 +526,7 @@ void UMetaSoundSourceBuilder::OnInputAdded(int32 InputIndex)
 				{
 					AudioDevice->SendCommandToActiveSounds(AudioComponentID, [NewInputName = NewInput.Name](FActiveSound& ActiveSound)
 					{
-						if (FMetaSoundParameterTransmitter* Transmitter = static_cast<FMetaSoundParameterTransmitter*>(ActiveSound.GetTransmitter()))
-						{
-							Transmitter->AddAvailableParameter(NewInputName);
-						}
+						static_cast<FMetaSoundParameterTransmitter*>(ActiveSound.GetTransmitter())->AddAvailableParameter(NewInputName);
 					});
 				}
 			}
@@ -549,11 +591,10 @@ void UMetaSoundSourceBuilder::OnInputNameChanged(FName OldName, FName NewName)
 					{
 						AudioDevice->SendCommandToActiveSounds(AudioComponentID, [OldName, NewInputName = Input->Name](FActiveSound& ActiveSound)
 						{
-							if (FMetaSoundParameterTransmitter* Transmitter = static_cast<FMetaSoundParameterTransmitter*>(ActiveSound.GetTransmitter()))
-							{
-								Transmitter->RemoveAvailableParameter(OldName);
-								Transmitter->AddAvailableParameter(NewInputName);
-							}
+							FMetaSoundParameterTransmitter* Transmitter = static_cast<FMetaSoundParameterTransmitter*>(ActiveSound.GetTransmitter());
+							check(Transmitter);
+							Transmitter->RemoveAvailableParameter(OldName);
+							Transmitter->AddAvailableParameter(NewInputName);
 						});
 					}
 				}
@@ -765,6 +806,13 @@ void UMetaSoundSourceBuilder::OnPageAdded(const Metasound::Frontend::FDocumentMu
 	InitTargetPageDelegates(DocDelegates);
 }
 
+void UMetaSoundSourceBuilder::OnPresetStateChanged(const Metasound::Frontend::FDocumentPresetStateChangedArgs& InArgs) const
+{
+	using namespace Metasound::Frontend;
+
+	GetMetaSoundSource().InvalidateCachedRuntimeInputData();
+}
+
 void UMetaSoundSourceBuilder::OnRemovingPage(const Metasound::Frontend::FDocumentMutatePageArgs& InArgs)
 {
 	using namespace Metasound::Frontend;
@@ -893,12 +941,9 @@ void UMetaSoundSourceBuilder::OnRemovingInput(int32 InputIndex)
 				if (FAudioDevice* AudioDevice = AudioComponent->GetAudioDevice())
 				{
 					AudioDevice->SendCommandToActiveSounds(AudioComponentID, [InputRemoved = InputBeingRemoved.Name](FActiveSound& ActiveSound)
-					{
-						if (FMetaSoundParameterTransmitter* Transmitter = static_cast<FMetaSoundParameterTransmitter*>(ActiveSound.GetTransmitter()))
 						{
-							Transmitter->RemoveAvailableParameter(InputRemoved);
-						}
-					});
+							static_cast<FMetaSoundParameterTransmitter*>(ActiveSound.GetTransmitter())->RemoveAvailableParameter(InputRemoved);
+						});
 				}
 			}
 		}
@@ -1340,6 +1385,12 @@ FMetasoundFrontendLiteral UMetaSoundBuilderSubsystem::CreateMetaSoundLiteralFrom
 	return FMetasoundFrontendLiteral { Param };
 }
 
+bool UMetaSoundBuilderSubsystem::DetachBuilderFromAsset(const FMetasoundFrontendClassName& InClassName) const
+{
+	using namespace Metasound::Frontend;
+	return IDocumentBuilderRegistry::GetChecked().FinishBuilding(InClassName);
+}
+
 UMetaSoundBuilderBase* UMetaSoundBuilderSubsystem::FindBuilder(FName BuilderName)
 {
 	return NamedBuilders.FindRef(BuilderName);
@@ -1435,6 +1486,12 @@ UMetaSoundSourceBuilder* UMetaSoundBuilderSubsystem::FindSourceBuilder(FName Bui
 	return nullptr;
 }
 
+void UMetaSoundBuilderSubsystem::InvalidateDocumentCache(const FMetasoundFrontendClassName& InClassName) const
+{
+	using namespace Metasound::Engine;
+	FDocumentBuilderRegistry::GetChecked().ReloadBuilder(InClassName);
+}
+
 bool UMetaSoundBuilderSubsystem::IsInterfaceRegistered(FName InInterfaceName) const
 {
 	using namespace Metasound::Frontend;
diff --git a/Engine/Plugins/Runtime/Metasound/Source/MetasoundEngine/Public/DocumentConfigurations/MetasoundDocumentConfigurator.h b/Engine/Plugins/Runtime/Metasound/Source/MetasoundEngine/Public/DocumentConfigurations/MetasoundDocumentConfigurator.h
deleted file mode 100644
diff --git a/Engine/Plugins/Runtime/Metasound/Source/MetasoundEngine/Public/DocumentConfigurations/MetasoundRandomizerDocumentConfiguration.h b/Engine/Plugins/Runtime/Metasound/Source/MetasoundEngine/Public/DocumentConfigurations/MetasoundRandomizerDocumentConfiguration.h
deleted file mode 100644
diff --git a/Engine/Plugins/Runtime/Metasound/Source/MetasoundEngine/Public/MetasoundBuilderBase.h b/Engine/Plugins/Runtime/Metasound/Source/MetasoundEngine/Public/MetasoundBuilderBase.h
--- a/Engine/Plugins/Runtime/Metasound/Source/MetasoundEngine/Public/MetasoundBuilderBase.h
+++ b/Engine/Plugins/Runtime/Metasound/Source/MetasoundEngine/Public/MetasoundBuilderBase.h
@@ -583,6 +583,12 @@ protected:
 	virtual void OnAssetReferenceAdded(TScriptInterface<IMetaSoundDocumentInterface> DocInterface) PURE_VIRTUAL(UMetaSoundBuilderBase::OnAssetReferenceAdded, );
 	virtual void OnRemovingAssetReference(TScriptInterface<IMetaSoundDocumentInterface> DocInterface) PURE_VIRTUAL(UMetaSoundBuilderBase::OnRemovingAssetReference, );
 
+	UE_DEPRECATED(5.5, "Moved to internal implementation")
+	virtual void CreateTransientBuilder() { };
+
+	UE_DEPRECATED(5.5, "Moved to 'Reload', to enforce generation of new delegates")
+	UE_API void InvalidateCache(bool bPrimeCache = false);
+
 	// Runs build, conforming the document and corresponding object data on a MetaSound UObject to that managed by this builder.
 	template <typename UClassType>
 	UClassType& BuildInternal(UObject* Parent, const FMetaSoundBuilderOptions& BuilderOptions) const
@@ -644,6 +650,9 @@ protected:
 		return *MetaSound;
 	}
 
+	UE_DEPRECATED(5.4, "Use UMetaSoundBuilderDocument::Create instead")
+	UE_API UMetaSoundBuilderDocument* CreateTransientDocumentObject() const;
+
 	// Only registers provided MetaSound's graph class and referenced graphs recursively if
 	// it has yet to be registered or if it has an attached builder reporting outstanding
 	// transactions that have yet to be registered.
diff --git a/Engine/Plugins/Runtime/Metasound/Source/MetasoundEngine/Public/MetasoundBuilderSubsystem.h b/Engine/Plugins/Runtime/Metasound/Source/MetasoundEngine/Public/MetasoundBuilderSubsystem.h
--- a/Engine/Plugins/Runtime/Metasound/Source/MetasoundEngine/Public/MetasoundBuilderSubsystem.h
+++ b/Engine/Plugins/Runtime/Metasound/Source/MetasoundEngine/Public/MetasoundBuilderSubsystem.h
@@ -122,7 +122,6 @@ private:
 
 	UE_API void InitTargetPageDelegates(Metasound::Frontend::FDocumentModifyDelegates& OutDocumentDelegates);
 
-	UE_API void OnDocumentConfigurationChanged(const Metasound::Frontend::FDocumentConfigChangedArgs& Args) const;
 	UE_API void OnEdgeAdded(int32 EdgeIndex) const;
 	UE_API void OnInputAdded(int32 InputIndex);
 	UE_API void OnInputDataTypeChanged(int32 InputIndex);
@@ -134,6 +133,7 @@ private:
 	UE_API void OnNodeInputLiteralSet(int32 NodeIndex, int32 VertexIndex, int32 LiteralIndex) const;
 	UE_API void OnOutputAdded(int32 OutputIndex) const;
 	UE_API void OnPageAdded(const Metasound::Frontend::FDocumentMutatePageArgs& InArgs);
+	UE_API void OnPresetStateChanged(const Metasound::Frontend::FDocumentPresetStateChangedArgs& InArgs) const;
 	UE_API void OnRemoveSwappingEdge(int32 SwapIndex, int32 LastIndex) const;
 	UE_API void OnRemovingInput(int32 InputIndex);
 	UE_API void OnRemoveSwappingNode(int32 SwapIndex, int32 LastIndex) const;
@@ -169,11 +169,26 @@ private:
 	TMap<FName, TObjectPtr<UMetaSoundBuilderBase>> NamedBuilders;
 
 public:	
+	UE_DEPRECATED(5.5, "Call 'ReloadBuilder' in IDocumentBuilderRegistry instead ")
+	UE_API virtual void InvalidateDocumentCache(const FMetasoundFrontendClassName& InClassName) const;
+
 	static UE_API UMetaSoundBuilderSubsystem* Get();
 	static UE_API UMetaSoundBuilderSubsystem& GetChecked();
 	static UE_API const UMetaSoundBuilderSubsystem* GetConst();
 	static UE_API const UMetaSoundBuilderSubsystem& GetConstChecked();
 
+	UE_DEPRECATED(5.5, "Use FDocumentBuilderRegistry::FindOrBeginBuilding, which is now only supported in builds loading editor-only data.")
+	UE_API UMetaSoundBuilderBase& AttachBuilderToAssetChecked(UObject& InObject) const;
+
+	UE_DEPRECATED(5.5, "Use FDocumentBuilderRegistry::FindOrBeginBuilding (when editor only data is loaded) or MetaSoundEditorSubsystem::FindOrBeginBuilding call")
+	UE_API UMetaSoundPatchBuilder* AttachPatchBuilderToAsset(UMetaSoundPatch* InPatch) const;
+
+	UE_DEPRECATED(5.5, "Use FDocumentBuilderRegistry::FindOrBeginBuilding (when editor only data is loaded) or MetaSoundEditorSubsystem::FindOrBeginBuilding call")
+	UE_API UMetaSoundSourceBuilder* AttachSourceBuilderToAsset(UMetaSoundSource* InSource) const;
+
+	UE_DEPRECATED(5.5, "Moved to IDocumentBuilderRegistry::RemoveBuilderFromAsset")
+	UE_API bool DetachBuilderFromAsset(const FMetasoundFrontendClassName& InClassName) const;
+
 	UFUNCTION(BlueprintCallable, Category = "Audio|MetaSound|Builder",  meta = (ExpandEnumAsExecs = "OutResult"))
 	UE_API UPARAM(DisplayName = "Patch Builder") UMetaSoundPatchBuilder* CreatePatchBuilder(FName BuilderName, EMetaSoundBuilderResult& OutResult);
 
diff --git a/Engine/Plugins/Runtime/Metasound/Source/MetasoundFrontend/Private/DocumentConfigurations/MetasoundPresetDocumentConfiguration.cpp b/Engine/Plugins/Runtime/Metasound/Source/MetasoundFrontend/Private/DocumentConfigurations/MetasoundPresetDocumentConfiguration.cpp
deleted file mode 100644
diff --git a/Engine/Plugins/Runtime/Metasound/Source/MetasoundFrontend/Private/MetasoundAssetBase.cpp b/Engine/Plugins/Runtime/Metasound/Source/MetasoundFrontend/Private/MetasoundAssetBase.cpp
--- a/Engine/Plugins/Runtime/Metasound/Source/MetasoundFrontend/Private/MetasoundAssetBase.cpp
+++ b/Engine/Plugins/Runtime/Metasound/Source/MetasoundFrontend/Private/MetasoundAssetBase.cpp
@@ -7,7 +7,6 @@
 #include "Algo/ForEach.h"
 #include "Algo/Transform.h"
 #include "Containers/Set.h"
-#include "DocumentConfigurations/MetasoundPresetDocumentConfiguration.h"
 #include "HAL/FileManager.h"
 #include "HAL/IConsoleManager.h"
 #include "IAudioParameterTransmitter.h"
@@ -211,7 +210,7 @@ namespace Metasound
 			TEXT("au.MetaSound.SampleRate"),
 			AssetBasePrivate::SampleRateOverride,
 			TEXT("Overrides the sample rate of metasounds. Negative values default to audio mixer sample rate.\n")
-			TEXT("Default: 0, Min: 8000, Max: 96000"),
+			TEXT("Default: 0, Min: 8000, Max: 48000"),
 			FConsoleVariableDelegate::CreateLambda([](IConsoleVariable* Var) { CVarMetaSoundSampleRateChanged.Broadcast(Var); }),
 			ECVF_Default);
 
@@ -380,8 +379,7 @@ void FMetasoundAssetBase::UpdateAndRegisterForExecution(const Metasound::Fronten
 		if (Owner->IsAsset())
 		{
 			DocBuilder = &IDocumentBuilderRegistry::GetChecked().FindOrBeginBuilding(Owner);
-			FAutoUpdateRootGraph AutoUpdateTransform(GetOwningAssetName(), InRegistrationOptions.bAutoUpdateLogWarningOnDroppedConnection);
-			bDidUpdate = AutoUpdateTransform.Transform(*DocBuilder);
+			VersionDependencies(*DocBuilder, InRegistrationOptions.bAutoUpdateLogWarningOnDroppedConnection);
 		}
 #else // !WITH_EDITORONLY_DATA
 		constexpr bool bDidUpdate = false;
@@ -472,12 +470,7 @@ void FMetasoundAssetBase::UpdateAndRegisterForSerialization(const Metasound::Fro
 	}
 
 	// Auto update must be done after all referenced asset classes are registered
-#if WITH_EDITORONLY_DATA
-	constexpr bool bAutoUpdateLogWarningOnDroppedConnection = true;
-	FAutoUpdateRootGraph AutoUpdateTransform(GetOwningAssetName(), bAutoUpdateLogWarningOnDroppedConnection);
-	bDidUpdate |= AutoUpdateTransform.Transform(DocBuilder);
-#endif // WITH_EDITORONLY_DATA
-
+	bDidUpdate |= VersionDependencies(DocBuilder, /*bAutoUpdateLogWarningOnDroppedConnection=*/true);
 #if WITH_EDITOR
 	if (bDidUpdate)
 	{
@@ -885,6 +878,23 @@ void FMetasoundAssetBase::RegisterAssetDependencies(const Metasound::Frontend::F
 }
 
 #if WITH_EDITORONLY_DATA
+void FMetasoundAssetBase::UpdateAndRegisterReferencesForSerialization(FName InCookPlatformName)
+{
+	using namespace Metasound::Frontend;
+	Metasound::Frontend::FMetaSoundAssetCookOptions CookOptions{ .bStripUnusedPages = !InCookPlatformName.IsNone() };
+
+	if (CookOptions.bStripUnusedPages)
+	{
+		const IDocumentBuilderRegistry& DocBuilderRegistry = IDocumentBuilderRegistry::GetChecked();
+PRAGMA_DISABLE_DEPRECATION_WARNINGS
+		CookOptions.PagesToTarget = DocBuilderRegistry.GetCookedTargetPages(InCookPlatformName);
+		CookOptions.PageOrder = DocBuilderRegistry.GetCookedPageOrder(InCookPlatformName);
+PRAGMA_ENABLE_DEPRECATION_WARNINGS
+	}
+
+	UpdateAndRegisterReferencesForSerialization(CookOptions);
+}
+
 void FMetasoundAssetBase::UpdateAndRegisterReferencesForSerialization(const Metasound::Frontend::FMetaSoundAssetCookOptions& InCookOptions)
 {
 	using namespace Metasound::Frontend;
@@ -918,6 +928,20 @@ bool FMetasoundAssetBase::TryUpdateInterfaceFromVersion(const FMetasoundFrontend
 	return false;
 }
 
+bool FMetasoundAssetBase::VersionDependencies(FMetaSoundFrontendDocumentBuilder& Builder, bool bInLogWarningsOnDroppedConnection)
+{
+	using namespace Metasound::Frontend;
+
+	bool bDocumentModified = false;
+
+#if WITH_EDITORONLY_DATA
+	FAutoUpdateRootGraph AutoUpdateTransform(GetOwningAssetName(), bInLogWarningsOnDroppedConnection);
+	bDocumentModified |= AutoUpdateTransform.Transform(Builder);
+#endif // WITH_EDITORONLY_DATA
+
+	return bDocumentModified;
+}
+
 FMetasoundFrontendVersion FMetasoundAssetBase::GetTargetInterfaceVersion(const FMetasoundFrontendVersion& InterfaceVersion) const
 {
 	using namespace Metasound::Frontend;
diff --git a/Engine/Plugins/Runtime/Metasound/Source/MetasoundFrontend/Private/MetasoundAssetManager.cpp b/Engine/Plugins/Runtime/Metasound/Source/MetasoundFrontend/Private/MetasoundAssetManager.cpp
--- a/Engine/Plugins/Runtime/Metasound/Source/MetasoundFrontend/Private/MetasoundAssetManager.cpp
+++ b/Engine/Plugins/Runtime/Metasound/Source/MetasoundFrontend/Private/MetasoundAssetManager.cpp
@@ -20,6 +20,28 @@ namespace Metasound::Frontend
 		static TUniquePtr<IMetaSoundAssetManager> Instance;
 	} // AssetManagerPrivate
 
+	namespace AssetTags
+	{
+		// Deprecated Tag Identifiers
+PRAGMA_DISABLE_DEPRECATION_WARNINGS
+		const FString ArrayDelim = TEXT(",");
+
+		const FName AssetClassID = "AssetClassID";
+
+#if WITH_EDITORONLY_DATA
+		const FName IsPreset = "bIsPreset";
+#endif // WITH_EDITORONLY_DATA
+
+		const FName RegistryVersionMajor = "RegistryVersionMajor";
+		const FName RegistryVersionMinor = "RegistryVersionMinor";
+
+#if WITH_EDITORONLY_DATA
+		const FName RegistryInputTypes = "RegistryInputTypes";
+		const FName RegistryOutputTypes = "RegistryOutputTypes";
+#endif // WITH_EDITORONLY_DATA
+PRAGMA_ENABLE_DEPRECATION_WARNINGS
+	} // namespace AssetTags
+
 	namespace AssetTagsPrivate
 	{
 		static const FLazyName AssetClassID = "AssetClassID";
diff --git a/Engine/Plugins/Runtime/Metasound/Source/MetasoundFrontend/Private/MetasoundDocumentInterface.cpp b/Engine/Plugins/Runtime/Metasound/Source/MetasoundFrontend/Private/MetasoundDocumentInterface.cpp
new file mode 100644
--- /dev/null
+++ b/Engine/Plugins/Runtime/Metasound/Source/MetasoundFrontend/Private/MetasoundDocumentInterface.cpp
@@ -0,0 +1,41 @@
+// Copyright Epic Games, Inc. All Rights Reserved.
+
+#include "MetasoundDocumentInterface.h"
+#include "MetasoundFrontendDocumentIdGenerator.h"
+#include "MetasoundFrontendGraph.h"
+#include "MetasoundFrontendProxyDataCache.h"
+#include "MetasoundOperatorBuilder.h"
+#include "MetasoundOperatorBuilderSettings.h"
+
+#include UE_INLINE_GENERATED_CPP_BY_NAME(MetasoundDocumentInterface)
+
+
+namespace Metasound::Frontend
+{
+	namespace DocumentInterfacePrivate
+	{
+		TUniquePtr<IDocumentBuilderRegistry> Instance;
+	}
+
+	IDocumentBuilderRegistry* IDocumentBuilderRegistry::Get()
+	{
+		return DocumentInterfacePrivate::Instance.Get();
+	}
+
+	IDocumentBuilderRegistry& IDocumentBuilderRegistry::GetChecked()
+	{
+		return *DocumentInterfacePrivate::Instance.Get();
+	}
+
+	void IDocumentBuilderRegistry::Deinitialize()
+	{
+		check(DocumentInterfacePrivate::Instance.IsValid());
+		DocumentInterfacePrivate::Instance.Reset();
+	}
+
+	void IDocumentBuilderRegistry::Initialize(TUniquePtr<IDocumentBuilderRegistry>&& InInstance)
+	{
+		check(!DocumentInterfacePrivate::Instance.IsValid());
+		DocumentInterfacePrivate::Instance = MoveTemp(InInstance);
+	}
+} // namespace Metasound::Frontend
diff --git a/Engine/Plugins/Runtime/Metasound/Source/MetasoundFrontend/Private/MetasoundFrontendDocument.cpp b/Engine/Plugins/Runtime/Metasound/Source/MetasoundFrontend/Private/MetasoundFrontendDocument.cpp
--- a/Engine/Plugins/Runtime/Metasound/Source/MetasoundFrontend/Private/MetasoundFrontendDocument.cpp
+++ b/Engine/Plugins/Runtime/Metasound/Source/MetasoundFrontend/Private/MetasoundFrontendDocument.cpp
@@ -11,7 +11,6 @@
 #include "Logging/LogMacros.h"
 #include "MetasoundDocumentInterface.h"
 #include "MetasoundFrontend.h"
-#include "MetasoundFrontendDocumentBuilderRegistry.h"
 #include "MetasoundFrontendDocumentIdGenerator.h"
 #include "MetasoundFrontendDocumentVersioning.h"
 #include "MetasoundFrontendPages.h"
diff --git a/Engine/Plugins/Runtime/Metasound/Source/MetasoundFrontend/Private/MetasoundFrontendDocumentBuilder.cpp b/Engine/Plugins/Runtime/Metasound/Source/MetasoundFrontend/Private/MetasoundFrontendDocumentBuilder.cpp
--- a/Engine/Plugins/Runtime/Metasound/Source/MetasoundFrontend/Private/MetasoundFrontendDocumentBuilder.cpp
+++ b/Engine/Plugins/Runtime/Metasound/Source/MetasoundFrontend/Private/MetasoundFrontendDocumentBuilder.cpp
@@ -9,7 +9,6 @@
 #include "Algo/Sort.h"
 #include "Algo/Transform.h"
 #include "AudioParameter.h"
-#include "DocumentConfigurations/MetasoundPresetDocumentConfiguration.h"
 #include "Interfaces/MetasoundFrontendInterfaceBindingRegistry.h"
 #include "Interfaces/MetasoundFrontendInterfaceRegistry.h"
 #include "MetasoundAssetBase.h"
@@ -18,12 +17,10 @@
 #include "MetasoundFrontendController.h"
 #include "MetasoundFrontendDataTypeRegistry.h"
 #include "MetasoundFrontendDocumentCache.h"
-#include "MetasoundFrontendDocumentConfiguration.h"
 #include "MetasoundFrontendDocument.h"
 #include "MetasoundFrontendDocumentIdGenerator.h"
 #include "MetasoundFrontendDocumentModifyDelegates.h"
 #include "MetasoundFrontendDocumentVersioning.h"
-#include "MetasoundFrontendGraphLayer.h"
 #include "MetasoundFrontendNodeTemplateRegistry.h"
 #include "MetasoundFrontendNodeUpdateTransform.h"
 #include "MetasoundFrontendRegistries.h"
@@ -91,20 +88,6 @@ namespace Metasound::Frontend
 			return false;
 		}
 
-		FName GenerateUniqueName(const TSet<FName>& InExistingNames, const FString& InBaseName)
-		{
-			int32 PostFixInt = 0;
-			FString NewName = InBaseName;
-
-			while (InExistingNames.Contains(*NewName))
-			{
-				PostFixInt++;
-				NewName = FString::Format(TEXT("{0} {1}"), { InBaseName, PostFixInt });
-			}
-
-			return FName(*NewName);
-		}
-
 		bool NameContainsInterfaceNamespace(FName VertexName, FMetasoundFrontendInterface* OutInterface)
 		{
 			using namespace Metasound::Frontend;
@@ -782,6 +765,14 @@ UMetaSoundBuilderDocument::UMetaSoundBuilderDocument(const FObjectInitializer& O
 	Document.RootGraph.ID = FGuid::NewGuid();
 }
 
+UMetaSoundBuilderDocument& UMetaSoundBuilderDocument::Create(const UClass& InMetaSoundUClass)
+{
+	UMetaSoundBuilderDocument* DocObject = NewObject<UMetaSoundBuilderDocument>();
+	check(DocObject);
+	DocObject->MetaSoundUClass = &InMetaSoundUClass;
+	return *DocObject;
+}
+
 UMetaSoundBuilderDocument& UMetaSoundBuilderDocument::Create(const IMetaSoundDocumentInterface& InDocToCopy)
 {
 	UMetaSoundBuilderDocument* DocObject = NewObject<UMetaSoundBuilderDocument>();
@@ -810,14 +801,6 @@ const FMetasoundFrontendDocument& UMetaSoundBuilderDocument::GetConstDocument()
 	return Document;
 }
 
-#if WITH_EDITORONLY_DATA
-TConstStructView<FMetaSoundFrontendDocumentConfiguration> FMetaSoundFrontendDocumentBuilder::GetConstDocumentConfiguration() const
-{
-	const TInstancedStruct<FMetaSoundFrontendDocumentConfiguration>& Config = GetConstDocumentChecked().Configuration;
-	return Config;
-}
-#endif // WITH_EDITORONLY_DATA
-
 EMetasoundFrontendClassAccessFlags UMetaSoundBuilderDocument::GetDefaultAccessFlags() const
 {
 	return EMetasoundFrontendClassAccessFlags::None;
@@ -1040,7 +1023,7 @@ bool FMetaSoundFrontendDocumentBuilder::AddNamedEdges(const TSet<Metasound::Fron
 			checkf(Input, TEXT("Prior loop to check edge validity should protect against missing input vertex"));
 
 			const FMetasoundFrontendVertex* Output = FindNodeOutput(EdgeToAdd.NewEdge.FromNodeID, EdgeToAdd.NewEdge.FromVertexID, &PageID);
-			checkf(Output, TEXT("Prior loop to check edge validity should protect against missing output vertex"));
+			checkf(Input, TEXT("Prior loop to check edge validity should protect against missing output vertex"));
 
 			UE_LOG(LogMetaSound, Warning, TEXT("Connection between MetaSound output '%s' and input '%s' not added: Input already connected to '%s'."), *Output->Name.ToString(), *Input->Name.ToString(), *Output->Name.ToString());
 #endif // !NO_LOGGING
@@ -1281,11 +1264,6 @@ const FMetasoundFrontendNode* FMetaSoundFrontendDocumentBuilder::AddGraphInput(F
 
 		DocumentDelegates->InterfaceDelegates.OnInputAdded.Broadcast(NewIndex);
 #if WITH_EDITORONLY_DATA
-		if (FMetaSoundFrontendDocumentConfiguration* DocConfig = Document.Configuration.GetMutablePtr())
-		{
-			DocConfig->AddVertexMetadata(*this, NewNode->GetID());
-		}
-
 		Document.Metadata.ModifyContext.AddMemberIDModified(NewInput.NodeID);
 #endif // WITH_EDITORONLY_DATA
 
@@ -1395,13 +1373,7 @@ const FMetasoundFrontendNode* FMetaSoundFrontendDocumentBuilder::AddGraphOutput(
 		}
 
 		DocumentDelegates->InterfaceDelegates.OnOutputAdded.Broadcast(NewIndex);
-
 #if WITH_EDITORONLY_DATA
-		if (FMetaSoundFrontendDocumentConfiguration* DocConfig = Document.Configuration.GetMutablePtr())
-		{
-			DocConfig->AddVertexMetadata(*this, NewOutput.NodeID);
-		}
-
 		Document.Metadata.ModifyContext.AddMemberIDModified(NewOutput.NodeID);
 #endif // WITH_EDITORONLY_DATA
 	}
@@ -2353,7 +2325,6 @@ void FMetaSoundFrontendDocumentBuilder::ClearDocument(TSharedRef<Metasound::Fron
 
 	GraphClass.GetDefaultInterface().Inputs.Empty();
 	GraphClass.GetDefaultInterface().Outputs.Empty();
-	GraphClass.GetDefaultInterface().Environment.Empty();
 
 #if WITH_EDITOR
 	GraphClass.GetDefaultInterface().SetInputStyle({ });
@@ -2380,6 +2351,10 @@ void FMetaSoundFrontendDocumentBuilder::ClearDocument(TSharedRef<Metasound::Fron
 	});
 #endif // !WITH_EDITORONLY_DATA
 
+	GraphClass.GetDefaultInterface().Inputs.Empty();
+	GraphClass.GetDefaultInterface().Outputs.Empty();
+	GraphClass.GetDefaultInterface().Environment.Empty();
+
 	Doc.Interfaces.Empty();
 	Doc.Dependencies.Empty();
 
@@ -2505,40 +2480,6 @@ bool FMetaSoundFrontendDocumentBuilder::ConformGraphOutputNodeToClass(const FMet
 	return false;
 }
 
-#if WITH_EDITORONLY_DATA
-bool FMetaSoundFrontendDocumentBuilder::ConfigureDocument(bool bLogFailure)
-{
-	using namespace Metasound::Frontend;
-
-	if (bInDocConfigScope)
-	{
-		UE_LOG(LogMetaSound, Error, TEXT("Cannot call Configure Document: Configure already called higher in stack"));
-		return false;
-	}
-
-	bInDocConfigScope = true;
-
-	FMetasoundFrontendDocument& Document = GetDocumentChecked();
-	if (FMetaSoundFrontendDocumentConfiguration* DocConfig = Document.Configuration.GetMutablePtr())
-	{
-		if (DocConfig->ConfigureDocument(*this, bLogFailure))
-		{
-			GetDocumentInterfaceChecked().ConformObjectToDocument();
-			bInDocConfigScope = false;
-			return true;
-		}
-
-		TSharedRef<FDocumentModifyDelegates> ModifyDelegates = MakeShared<FDocumentModifyDelegates>(*DocumentDelegates);
-		ClearDocument(ModifyDelegates);
-	}
-
-	SetIsGraphEditable(true);
-
-	bInDocConfigScope = false;
-	return false;
-}
-#endif // WITH_EDITORONLY_DATA
-
 bool FMetaSoundFrontendDocumentBuilder::ContainsDependencyOfType(EMetasoundFrontendClassType ClassType) const
 {
 	return DocumentCache->ContainsDependencyOfType(ClassType);
@@ -2571,7 +2512,7 @@ bool FMetaSoundFrontendDocumentBuilder::ConvertFromPreset()
 		Style.bIsGraphEditable = true;
 #endif // WITH_EDITOR
 
-		DocumentDelegates->OnDocumentConfigurationChanged.Broadcast(FDocumentConfigChangedArgs { });
+		DocumentDelegates->OnPresetStateChanged.Broadcast(FDocumentPresetStateChangedArgs { });
 		return true;
 	}
 
@@ -2605,7 +2546,7 @@ bool FMetaSoundFrontendDocumentBuilder::ConvertToPreset(const FMetaSoundFrontend
 	if (RebuildPresetRootGraph.Transform(*this))
 	{
 		DocumentInterface->ConformObjectToDocument();
-		DocumentDelegates->OnDocumentConfigurationChanged.Broadcast(FDocumentConfigChangedArgs { });
+		DocumentDelegates->OnPresetStateChanged.Broadcast(FDocumentPresetStateChangedArgs { });
 		return true;
 	}
 
@@ -2771,6 +2712,26 @@ bool FMetaSoundFrontendDocumentBuilder::DiffAgainstRegistryInterface(const FGuid
 }
 #endif // WITH_EDITORONLY_DATA
 
+const FMetasoundFrontendNode* FMetaSoundFrontendDocumentBuilder::DuplicateGraphInput(const FMetasoundFrontendClassInput& InClassInput, const FName InName, const FGuid* InPageID)
+{
+	using namespace Metasound;
+
+	Frontend::FDocumentIDGenerator& IDGenerator = Frontend::FDocumentIDGenerator::Get();
+	const FMetasoundFrontendDocument& Doc = GetConstDocumentChecked();
+
+	const FGuid& PageID = InPageID ? *InPageID : BuildPageID;
+
+	FMetasoundFrontendClassInput ClassInput = InClassInput;
+	ClassInput.NodeID = IDGenerator.CreateNodeID(Doc);
+	ClassInput.VertexID = IDGenerator.CreateVertexID(Doc);
+#if WITH_EDITORONLY_DATA
+	ClassInput.Metadata.SetDisplayName(FText::GetEmpty());
+#endif // WITH_EDITORONLY_DATA
+	ClassInput.Name = InName;
+
+	return AddGraphInput(MoveTemp(ClassInput), &PageID);
+}
+
 const FMetasoundFrontendClassInput* FMetaSoundFrontendDocumentBuilder::DuplicateGraphInput(FName ExistingName, FName NewName)
 {
 	using namespace Metasound;
@@ -2803,6 +2764,26 @@ const FMetasoundFrontendClassInput* FMetaSoundFrontendDocumentBuilder::Duplicate
 	return FindGraphInput(NewName);
 }
 
+const FMetasoundFrontendNode* FMetaSoundFrontendDocumentBuilder::DuplicateGraphOutput(const FMetasoundFrontendClassOutput& InClassOutput, const FName InName, const FGuid* InPageID)
+{
+	using namespace Metasound::Frontend;
+
+	FDocumentIDGenerator& IDGenerator = FDocumentIDGenerator::Get();
+	const FMetasoundFrontendDocument& Doc = GetConstDocumentChecked();
+
+	const FGuid& PageID = InPageID ? *InPageID : BuildPageID;
+
+	FMetasoundFrontendClassOutput ClassOutput = InClassOutput;
+	ClassOutput.NodeID = IDGenerator.CreateNodeID(Doc);
+	ClassOutput.VertexID = IDGenerator.CreateVertexID(Doc);
+#if WITH_EDITORONLY_DATA
+	ClassOutput.Metadata.SetDisplayName(FText::GetEmpty());
+#endif // WITH_EDITORONLY_DATA
+	ClassOutput.Name = InName;
+
+	return AddGraphOutput(MoveTemp(ClassOutput), &PageID);
+}
+
 const FMetasoundFrontendClassOutput* FMetaSoundFrontendDocumentBuilder::DuplicateGraphOutput(FName ExistingName, FName NewName)
 {
 	using namespace Metasound::Frontend;
@@ -3623,35 +3604,6 @@ FMetasoundFrontendClassName FMetaSoundFrontendDocumentBuilder::GenerateNewClassN
 	return NewClassName;
 }
 
-FName FMetaSoundFrontendDocumentBuilder::GenerateUniqueInputName(const FString& InBaseName) const
-{
-	using namespace Metasound::Frontend;
-
-	TSet<FName> ExistingNames;
-	auto GetName = [&](const FMetasoundFrontendClass&, const FMetasoundFrontendNode& Node) { ExistingNames.Add(Node.Name); };
-	IterateNodesByClassType(GetName, EMetasoundFrontendClassType::Input);
-	return DocumentBuilderPrivate::GenerateUniqueName(ExistingNames, InBaseName);
-}
-
-FName FMetaSoundFrontendDocumentBuilder::GenerateUniqueOutputName(const FString& InBaseName) const
-{
-	using namespace Metasound::Frontend;
-
-	TSet<FName> ExistingNames;
-	auto GetName = [&](const FMetasoundFrontendClass&, const FMetasoundFrontendNode& Node) { ExistingNames.Add(Node.Name); };
-	IterateNodesByClassType(GetName, EMetasoundFrontendClassType::Output);
-	return DocumentBuilderPrivate::GenerateUniqueName(ExistingNames, InBaseName);
-}
-
-FName FMetaSoundFrontendDocumentBuilder::GenerateUniqueVariableName(const FString& InBaseName) const
-{
-	using namespace Metasound::Frontend;
-
-	TSet<FName> ExistingVariableNames;
-	Algo::Transform(FindConstBuildGraphChecked().Variables, ExistingVariableNames, [](const FMetasoundFrontendVariable& Var) { return Var.Name; });
-	return DocumentBuilderPrivate::GenerateUniqueName(ExistingVariableNames, InBaseName);
-}
-
 const FTopLevelAssetPath FMetaSoundFrontendDocumentBuilder::GetBuilderClassPath() const
 {
 	IMetaSoundDocumentInterface* Interface = DocumentInterface.GetInterface();
@@ -3679,6 +3631,13 @@ const FString FMetaSoundFrontendDocumentBuilder::GetDebugName() const
 	return MetaSoundObject.GetPathName();
 }
 
+const FMetasoundFrontendDocument& FMetaSoundFrontendDocumentBuilder::GetDocument() const
+{
+	const IMetaSoundDocumentInterface* Interface = DocumentInterface.GetInterface();
+	checkf(Interface, TEXT("Failed to return document; interface must always be valid while builder is operating on MetaSound UObject! Builder constructed with asset at %s"), *HintPath.ToString());
+	return Interface->GetConstDocument();
+}
+
 FMetasoundFrontendDocument& FMetaSoundFrontendDocumentBuilder::GetDocumentChecked() const
 {
 	return GetDocumentInterfaceChecked().GetDocument();
@@ -3689,6 +3648,13 @@ Metasound::Frontend::FDocumentModifyDelegates& FMetaSoundFrontendDocumentBuilder
 	return *DocumentDelegates;
 }
 
+const IMetaSoundDocumentInterface& FMetaSoundFrontendDocumentBuilder::GetDocumentInterface() const
+{
+	const IMetaSoundDocumentInterface* Interface = DocumentInterface.GetInterface();
+	checkf(Interface, TEXT("Failed to return document; interface must always be valid while builder is operating on MetaSound UObject! Builder constructed with asset at %s"), *HintPath.ToString());
+	return *Interface;
+}
+
 IMetaSoundDocumentInterface& FMetaSoundFrontendDocumentBuilder::GetDocumentInterfaceChecked() const
 {
 	IMetaSoundDocumentInterface* Interface = DocumentInterface.GetInterface();
@@ -3826,22 +3792,6 @@ const FMetasoundFrontendLiteral* FMetaSoundFrontendDocumentBuilder::GetGraphVari
 	return nullptr;
 }
 
-#if WITH_EDITORONLY_DATA
-bool FMetaSoundFrontendDocumentBuilder::GetIsGraphEditable(const FGuid* PageID) const
-{
-	bool bIsEditable = true;
-	const FMetasoundFrontendDocument& Document = GetConstDocumentChecked();
-	Document.RootGraph.IterateGraphPages([&bIsEditable, &PageID](const FMetasoundFrontendGraph& PageGraph)
-	{
-		if (!PageID || PageGraph.PageID == *PageID)
-		{
-			bIsEditable &= PageGraph.Style.bIsGraphEditable;
-		}
-	});
-	return bIsEditable;
-}
-#endif // WITH_EDITORONLY_DATA
-
 EMetasoundFrontendVertexAccessType FMetaSoundFrontendDocumentBuilder::GetNodeInputAccessType(const FGuid& InNodeID, const FGuid& InVertexID, const FGuid* InPageID) const
 {
 	using namespace Metasound::Frontend;
@@ -3940,8 +3890,10 @@ FText FMetaSoundFrontendDocumentBuilder::GetNodeInputDisplayName(const FGuid& In
 		case EMetasoundFrontendClassType::External:
 		{
 			auto MatchesName = [&VertexName](const FMetasoundFrontendClassVertex& OtherVertex) { return OtherVertex.Name == VertexName; };
-			const FMetasoundFrontendClassVertex* ClassVertex = Class->GetInterfaceForNode(*Node).Inputs.FindByPredicate(MatchesName);
-			if (ClassVertex)
+			const FMetasoundFrontendVertex* Vertex = nullptr;
+			const FMetasoundFrontendClassVertex* ClassVertex = nullptr;
+			ClassVertex = Class->GetInterfaceForNode(*Node).Inputs.FindByPredicate(MatchesName);
+			if (Vertex && ClassVertex)
 			{
 				FName Namespace, ParamName;
 				ClassVertex->SplitName(Namespace, ParamName);
@@ -4025,8 +3977,10 @@ FText FMetaSoundFrontendDocumentBuilder::GetNodeOutputDisplayName(const FGuid& I
 		case EMetasoundFrontendClassType::External:
 		{
 			auto PinMatchesClassVertex = [&VertexName](const FMetasoundFrontendClassVertex& OtherVertex) { return OtherVertex.Name == VertexName; };
-			const FMetasoundFrontendClassVertex* ClassVertex = Class->GetInterfaceForNode(*Node).Outputs.FindByPredicate(PinMatchesClassVertex);
-			if (ClassVertex)
+			const FMetasoundFrontendVertex* Vertex = nullptr;
+			const FMetasoundFrontendClassVertex* ClassVertex = nullptr;
+			ClassVertex = Class->GetInterfaceForNode(*Node).Inputs.FindByPredicate(PinMatchesClassVertex);
+			if (Vertex && ClassVertex)
 			{
 				FName Namespace, ParamName;
 				ClassVertex->SplitName(Namespace, ParamName);
@@ -4071,6 +4025,81 @@ FText FMetaSoundFrontendDocumentBuilder::GetNodeOutputDisplayName(const FGuid& I
 }
 #endif // WITH_EDITORONLY_DATA
 
+const FMetasoundFrontendLiteral* FMetaSoundFrontendDocumentBuilder::GetNodeInputClassDefault(const FGuid& InNodeID, const FGuid& InVertexID, const FGuid* InPageID) const
+{
+	using namespace Metasound;
+	using namespace Metasound::Frontend;
+
+	const FGuid& PageID = InPageID ? *InPageID : BuildPageID;
+	const Frontend::IDocumentGraphNodeCache& NodeCache = DocumentCache->GetNodeCache(PageID);
+	if (const int32* NodeIndex = NodeCache.FindNodeIndex(InNodeID))
+	{
+		const FMetasoundFrontendDocument& Document = GetConstDocumentChecked();
+		const FMetasoundFrontendNode& Node = Document.RootGraph.FindConstGraphChecked(PageID).Nodes[*NodeIndex];
+		auto IsVertexID = [&InVertexID](const FMetasoundFrontendVertex& Vertex) { return Vertex.VertexID == InVertexID; };
+		if (const FMetasoundFrontendVertex* Vertex = Node.Interface.Inputs.FindByPredicate(IsVertexID))
+		{
+			if (const FMetasoundFrontendClass* Class = DocumentCache->FindDependency(Node.ClassID))
+			{
+				const EMetasoundFrontendClassType ClassType = Class->Metadata.GetType();
+				const FMetasoundFrontendClassInterface& ClassInterface = Class->GetInterfaceForNode(Node);
+				switch (ClassType)
+				{
+					case EMetasoundFrontendClassType::Output:
+					{
+						const FMetasoundFrontendClassInput& ClassInput = ClassInterface.Inputs.Last();
+						return ClassInput.FindConstDefault(Frontend::DefaultPageID);
+					}
+
+					default:
+					{
+						auto IsClassInput = [VertexName = Vertex->Name](const FMetasoundFrontendClassInput& Input) { return Input.Name == VertexName; };
+						if (const FMetasoundFrontendClassInput* ClassInput = ClassInterface.Inputs.FindByPredicate(IsClassInput))
+						{
+							return ClassInput->FindConstDefault(Frontend::DefaultPageID);
+						}
+						static_assert(static_cast<uint32>(EMetasoundFrontendClassType::Invalid) == 10, "Potential missing case coverage for EMetasoundFrontendClassType "
+							"(default may not be sufficient for newly added class types)");
+					}
+					break;
+				}
+				static_assert(static_cast<uint32>(EMetasoundFrontendClassType::Invalid) == 10, "Potential missing case coverage for EMetasoundFrontendClassType");
+			}
+		}
+	}
+
+	return nullptr;
+}
+
+const FMetasoundFrontendLiteral* FMetaSoundFrontendDocumentBuilder::GetNodeInputDefault(const FGuid& InNodeID, const FGuid& InVertexID, const FGuid* InPageID) const
+{
+	using namespace Metasound::Frontend;
+
+	const FGuid& PageID = InPageID ? *InPageID : BuildPageID;
+	const IDocumentGraphNodeCache& NodeCache = DocumentCache->GetNodeCache(PageID);
+	if (const int32* NodeIndex = NodeCache.FindNodeIndex(InNodeID))
+	{
+		const FMetasoundFrontendGraph& Graph = GetConstDocumentChecked().RootGraph.FindConstGraphChecked(PageID);
+		const FMetasoundFrontendNode& Node = Graph.Nodes[*NodeIndex];
+
+		auto IsVertex = [&InVertexID](const FMetasoundFrontendVertex& Vertex) { return Vertex.VertexID == InVertexID; };
+		const int32 VertexIndex = Node.Interface.Inputs.IndexOfByPredicate(IsVertex);
+		if (VertexIndex != INDEX_NONE)
+		{
+			const FMetasoundFrontendVertex& NodeInput = Node.Interface.Inputs[VertexIndex];
+
+			auto IsLiteral = [&InVertexID](const FMetasoundFrontendVertexLiteral& Literal) { return Literal.VertexID == InVertexID; };
+			const int32 LiteralIndex = Node.InputLiterals.IndexOfByPredicate(IsLiteral);
+			if (LiteralIndex != INDEX_NONE)
+			{
+				return &Node.InputLiterals[LiteralIndex].Value;
+			}
+		}
+	}
+
+	return nullptr;
+}
+
 EMetasoundFrontendVertexAccessType FMetaSoundFrontendDocumentBuilder::GetNodeOutputAccessType(const FGuid& InNodeID, const FGuid& InVertexID, const FGuid* InPageID) const
 {
 	using namespace Metasound::Frontend;
@@ -4195,7 +4224,6 @@ void FMetaSoundFrontendDocumentBuilder::InitDocument(const FMetasoundFrontendDoc
 	// Initialize class using default data
 	else
 	{
-		// Initialize class using default data
 		if (Document.RootGraph.GetConstGraphPages().IsEmpty())
 		{
 			Document.RootGraph.InitDefaultGraphPage();
@@ -4206,7 +4234,6 @@ void FMetaSoundFrontendDocumentBuilder::InitDocument(const FMetasoundFrontendDoc
 		InitGraphClassMetadata(Document.RootGraph.Metadata, bResetVersion, InNewClassName);
 
 #if WITH_EDITORONLY_DATA
-		// Set default doc version Metadata
 		{
 			FMetasoundFrontendDocumentMetadata& DocMetadata = Document.Metadata;
 			DocMetadata.Version.Number = GetMaxDocumentVersion();
@@ -5230,13 +5257,6 @@ bool FMetaSoundFrontendDocumentBuilder::RemoveGraphInput(FName InputName, bool b
 
 		if (bNodesRemoved)
 		{
-#if WITH_EDITORONLY_DATA
-			if (FMetaSoundFrontendDocumentConfiguration* DocConfig = Document.Configuration.GetMutablePtr())
-			{
-				DocConfig->RemoveMetadata(NodeID);
-			}
-#endif // WITH_EDITORONLY_DATA
-
 			const int32 Index = *IndexPtr;
 			DocumentDelegates->InterfaceDelegates.OnRemovingInput.Broadcast(Index);
 
@@ -5539,6 +5559,11 @@ bool FMetaSoundFrontendDocumentBuilder::RemoveUnusedDependencies()
 	return bDidEdit;
 }
 
+bool FMetaSoundFrontendDocumentBuilder::RenameRootGraphClass(const FMetasoundFrontendClassName& InName)
+{
+	return false;
+}
+
 bool FMetaSoundFrontendDocumentBuilder::ReplaceDependency(const Metasound::Frontend::FNodeClassRegistryKey& OldNodeClassKey, const Metasound::Frontend::FNodeClassRegistryKey& NewNodeClassKey, TArray<FVertexNameAndType>* OutDisconnectedInputs, TArray<FVertexNameAndType>* OutDisconnectedOutputs)
 {
 	using namespace Metasound::Frontend;
@@ -5606,6 +5631,13 @@ bool FMetaSoundFrontendDocumentBuilder::ReplaceDependency(const Metasound::Front
 	return bSuccess;
 }
 
+void FMetaSoundFrontendDocumentBuilder::ReloadCache()
+{
+	using namespace Metasound::Frontend;
+
+	Reload(DocumentDelegates, true);
+}
+
 void FMetaSoundFrontendDocumentBuilder::ApplyNodeInstanceReplacementData(FMetasoundFrontendNode& InReplacementNode, FMetaSoundFrontendDocumentBuilder::FNodeInstanceReplacementData&& InInstanceData, TArray<FVertexNameAndType>* OutDisconnectedInputs, TArray<FVertexNameAndType>* OutDisconnectedOutputs)
 {
 	using namespace Metasound::Frontend;
@@ -5894,16 +5926,6 @@ bool FMetaSoundFrontendDocumentBuilder::SetBuildPageID(const FGuid& InBuildPageI
 
 	return false;
 }
-#endif // WITH_EDITORONLY_DATA
-
-#if WITH_EDITORONLY_DATA
-void FMetaSoundFrontendDocumentBuilder::SetDocumentConfiguration(TInstancedStruct<FMetaSoundFrontendDocumentConfiguration> InConfig)
-{
-	using namespace Metasound::Frontend;
-
-	GetDocumentChecked().Configuration = MoveTemp(InConfig);
-	DocumentDelegates->OnDocumentConfigurationChanged.Broadcast({ });
-}
 
 bool FMetaSoundFrontendDocumentBuilder::SetGraphInputAdvancedDisplay(const FName InputName, const bool InAdvancedDisplay)
 {
@@ -5924,18 +5946,6 @@ bool FMetaSoundFrontendDocumentBuilder::SetGraphInputAdvancedDisplay(const FName
 
 	return false;
 }
-
-void FMetaSoundFrontendDocumentBuilder::SetIsGraphEditable(bool bIsEditable, const FGuid* PageID) const
-{
-	FMetasoundFrontendDocument& Document = GetDocumentChecked();
-	Document.RootGraph.IterateGraphPages([bIsEditable, &PageID](FMetasoundFrontendGraph& PageGraph)
-	{
-		if (!PageID || PageGraph.PageID == *PageID)
-		{
-			PageGraph.Style.bIsGraphEditable = bIsEditable;
-		}
-	});
-}
 #endif // WITH_EDITORONLY_DATA
 
 bool FMetaSoundFrontendDocumentBuilder::SetGraphInputAccessType(FName InputName, EMetasoundFrontendVertexAccessType AccessType)
@@ -7336,6 +7346,18 @@ bool FMetaSoundFrontendDocumentBuilder::UpdateDependencyRegistryData(const TMap<
 #endif // WITH_EDITOR
 
 #if WITH_EDITORONLY_DATA
+bool FMetaSoundFrontendDocumentBuilder::VersionInterfaces()
+{
+	FMetasoundFrontendDocument& Document = GetDocumentChecked();
+	if (Document.RequiresInterfaceVersioning())
+	{
+		Document.VersionInterfaces();
+		return true;
+	}
+
+	return false;
+}
+
 FMetasoundFrontendDocument& FMetaSoundFrontendDocumentBuilder::IPropertyVersionTransform::GetDocumentUnsafe(const FMetaSoundFrontendDocumentBuilder& Builder)
 {
 	return Builder.GetDocumentChecked();
diff --git a/Engine/Plugins/Runtime/Metasound/Source/MetasoundFrontend/Private/MetasoundFrontendDocumentBuilderRegistry.cpp b/Engine/Plugins/Runtime/Metasound/Source/MetasoundFrontend/Private/MetasoundFrontendDocumentBuilderRegistry.cpp
deleted file mode 100644
diff --git a/Engine/Plugins/Runtime/Metasound/Source/MetasoundFrontend/Private/MetasoundFrontendDocumentConfiguration.cpp b/Engine/Plugins/Runtime/Metasound/Source/MetasoundFrontend/Private/MetasoundFrontendDocumentConfiguration.cpp
deleted file mode 100644
diff --git a/Engine/Plugins/Runtime/Metasound/Source/MetasoundFrontend/Private/MetasoundFrontendDocumentModifyDelegates.cpp b/Engine/Plugins/Runtime/Metasound/Source/MetasoundFrontend/Private/MetasoundFrontendDocumentModifyDelegates.cpp
--- a/Engine/Plugins/Runtime/Metasound/Source/MetasoundFrontend/Private/MetasoundFrontendDocumentModifyDelegates.cpp
+++ b/Engine/Plugins/Runtime/Metasound/Source/MetasoundFrontend/Private/MetasoundFrontendDocumentModifyDelegates.cpp
@@ -9,16 +9,11 @@
 namespace Metasound::Frontend
 {
 	FDocumentModifyDelegates::FDocumentModifyDelegates()
-PRAGMA_DISABLE_DEPRECATION_WARNINGS
-		: OnPresetStateChanged()
-PRAGMA_ENABLE_DEPRECATION_WARNINGS
 	{
 	}
 
 	FDocumentModifyDelegates::FDocumentModifyDelegates(const FMetasoundFrontendDocument& Document)
-PRAGMA_DISABLE_DEPRECATION_WARNINGS
-		: OnPresetStateChanged()
-PRAGMA_ENABLE_DEPRECATION_WARNINGS
+		: FDocumentModifyDelegates()
 	{
 		Document.RootGraph.IterateGraphPages([this](const FMetasoundFrontendGraph& Graph)
 		{
diff --git a/Engine/Plugins/Runtime/Metasound/Source/MetasoundFrontend/Private/MetasoundFrontendDocumentVersioning.cpp b/Engine/Plugins/Runtime/Metasound/Source/MetasoundFrontend/Private/MetasoundFrontendDocumentVersioning.cpp
--- a/Engine/Plugins/Runtime/Metasound/Source/MetasoundFrontend/Private/MetasoundFrontendDocumentVersioning.cpp
+++ b/Engine/Plugins/Runtime/Metasound/Source/MetasoundFrontend/Private/MetasoundFrontendDocumentVersioning.cpp
@@ -5,7 +5,6 @@
 #include "Algo/Transform.h"
 #include "Algo/Unique.h"
 #include "CoreGlobals.h"
-#include "DocumentConfigurations/MetasoundPresetDocumentConfiguration.h"
 #include "HAL/IConsoleManager.h"
 #include "Interfaces/MetasoundFrontendInterface.h"
 #include "Interfaces/MetasoundFrontendInterfaceRegistry.h"
diff --git a/Engine/Plugins/Runtime/Metasound/Source/MetasoundFrontend/Private/MetasoundFrontendGraphController.cpp b/Engine/Plugins/Runtime/Metasound/Source/MetasoundFrontend/Private/MetasoundFrontendGraphController.cpp
--- a/Engine/Plugins/Runtime/Metasound/Source/MetasoundFrontend/Private/MetasoundFrontendGraphController.cpp
+++ b/Engine/Plugins/Runtime/Metasound/Source/MetasoundFrontend/Private/MetasoundFrontendGraphController.cpp
@@ -781,11 +781,7 @@ namespace Metasound
 				Graph.Edges.Reset();
 				GraphClass->GetDefaultInterface().Inputs.Reset();
 				GraphClass->GetDefaultInterface().Outputs.Reset();
-
-PRAGMA_DISABLE_DEPRECATION_WARNINGS
 				GraphClass->PresetOptions.InputsInheritingDefault.Reset();
-PRAGMA_ENABLE_DEPRECATION_WARNINGS
-
 				OwningDocument->ClearInterfaceVersions();
 				OwningDocument->RemoveUnreferencedDependencies();
 
diff --git a/Engine/Plugins/Runtime/Metasound/Source/MetasoundFrontend/Private/MetasoundFrontendLiteral.cpp b/Engine/Plugins/Runtime/Metasound/Source/MetasoundFrontend/Private/MetasoundFrontendLiteral.cpp
--- a/Engine/Plugins/Runtime/Metasound/Source/MetasoundFrontend/Private/MetasoundFrontendLiteral.cpp
+++ b/Engine/Plugins/Runtime/Metasound/Source/MetasoundFrontend/Private/MetasoundFrontendLiteral.cpp
@@ -251,15 +251,12 @@ void FMetasoundFrontendLiteral::Set(FDefault InValue)
 	AsNumDefault = 1;
 }
 
-void FMetasoundFrontendLiteral::Set(FDefaultArray InValue)
+void FMetasoundFrontendLiteral::Set(const FDefaultArray& InValue)
 {
-	AsBoolean.Empty();
-	AsInteger.Empty();
-	AsFloat.Empty();
-	AsString.Empty();
-	AsUObject.Empty();
+	Empty();
 
 	Type = EMetasoundFrontendLiteralType::NoneArray;
+
 	AsNumDefault = FMath::Max(0, InValue.Num);
 }
 
@@ -272,16 +269,13 @@ void FMetasoundFrontendLiteral::Set(bool InValue)
 	AsBoolean.Add(InValue);
 }
 
-void FMetasoundFrontendLiteral::Set(TArray<bool> InValue)
+void FMetasoundFrontendLiteral::Set(const TArray<bool>& InValue)
 {
-	AsNumDefault = 0;
-	AsInteger.Empty();
-	AsFloat.Empty();
-	AsString.Empty();
-	AsUObject.Empty();
+	Empty();
 
 	Type = EMetasoundFrontendLiteralType::BooleanArray;
-	AsBoolean = MoveTemp(InValue);
+
+	AsBoolean = InValue;
 }
 
 void FMetasoundFrontendLiteral::Set(int32 InValue)
@@ -293,16 +287,13 @@ void FMetasoundFrontendLiteral::Set(int32 InValue)
 	AsInteger.Add(InValue);
 }
 
-void FMetasoundFrontendLiteral::Set(TArray<int32> InValue)
+void FMetasoundFrontendLiteral::Set(const TArray<int32>& InValue)
 {
-	AsNumDefault = 0;
-	AsBoolean.Empty();
-	AsFloat.Empty();
-	AsString.Empty();
-	AsUObject.Empty();
+	Empty();
 
 	Type = EMetasoundFrontendLiteralType::IntegerArray;
-	AsInteger = MoveTemp(InValue);
+
+	AsInteger = InValue;
 }
 
 void FMetasoundFrontendLiteral::Set(float InValue)
@@ -314,40 +305,31 @@ void FMetasoundFrontendLiteral::Set(float InValue)
 	AsFloat.Add(InValue);
 }
 
-void FMetasoundFrontendLiteral::Set(TArray<float> InValue)
+void FMetasoundFrontendLiteral::Set(const TArray<float>& InValue)
 {
-	AsNumDefault = 0;
-	AsBoolean.Empty();
-	AsInteger.Empty();
-	AsString.Empty();
-	AsUObject.Empty();
+	Empty();
 
 	Type = EMetasoundFrontendLiteralType::FloatArray;
-	AsFloat = MoveTemp(InValue);
+
+	AsFloat = InValue;
 }
 
-void FMetasoundFrontendLiteral::Set(FString InValue)
+void FMetasoundFrontendLiteral::Set(const FString& InValue)
 {
-	AsNumDefault = 0;
-	AsBoolean.Empty();
-	AsInteger.Empty();
-	AsFloat.Empty();
-	AsUObject.Empty();
+	Empty();
 
 	Type = EMetasoundFrontendLiteralType::String;
-	AsString = { MoveTemp(InValue) };
+
+	AsString.Add(InValue);
 }
 
-void FMetasoundFrontendLiteral::Set(TArray<FString> InValue)
+void FMetasoundFrontendLiteral::Set(const TArray<FString>& InValue)
 {
-	AsNumDefault = 0;
-	AsBoolean.Empty();
-	AsInteger.Empty();
-	AsFloat.Empty();
-	AsUObject.Empty();
+	Empty();
 
 	Type = EMetasoundFrontendLiteralType::StringArray;
-	AsString = MoveTemp(InValue);
+
+	AsString = InValue;
 }
 
 void FMetasoundFrontendLiteral::Set(UObject* InValue)
@@ -359,16 +341,13 @@ void FMetasoundFrontendLiteral::Set(UObject* InValue)
 	AsUObject.Add(InValue);
 }
 
-void FMetasoundFrontendLiteral::Set(TArray<UObject*> InValue)
+void FMetasoundFrontendLiteral::Set(const TArray<UObject*>& InValue)
 {
-	AsNumDefault = 0;
-	AsBoolean.Empty();
-	AsInteger.Empty();
-	AsFloat.Empty();
-	AsString.Empty();
+	Empty();
 
 	Type = EMetasoundFrontendLiteralType::UObjectArray;
-	AsUObject = MoveTemp(InValue);
+
+	AsUObject = InValue;
 }
 
 void FMetasoundFrontendLiteral::SetFromLiteral(const Metasound::FLiteral& InLiteral)
@@ -475,7 +454,7 @@ void FMetasoundFrontendLiteral::SetFromLiteral(const Metasound::FLiteral& InLite
 
 void FMetasoundFrontendLiteral::SetType(EMetasoundFrontendLiteralType InType)
 {
-	switch (InType)
+	switch (Type)
 	{
 		case EMetasoundFrontendLiteralType::NoneArray:
 		{
diff --git a/Engine/Plugins/Runtime/Metasound/Source/MetasoundFrontend/Private/MetasoundFrontendTransform.cpp b/Engine/Plugins/Runtime/Metasound/Source/MetasoundFrontend/Private/MetasoundFrontendTransform.cpp
--- a/Engine/Plugins/Runtime/Metasound/Source/MetasoundFrontend/Private/MetasoundFrontendTransform.cpp
+++ b/Engine/Plugins/Runtime/Metasound/Source/MetasoundFrontend/Private/MetasoundFrontendTransform.cpp
@@ -883,9 +883,8 @@ namespace Metasound
 				check(ReferencedNodeInputVertex);
 				
 				ReferencedNodeInputVertices.Add(ReferencedNodeInputVertex);
-
-				// If not editor-only data, just make connection directly between input and output node
-#if !WITH_EDITORONLY_DATA
+				// If not editor, just make connection directly between input and output node
+				#if !WITH_EDITORONLY_DATA
 				InOutBuilderToTransform.AddEdge(FMetasoundFrontendEdge
 				{
 					InputNodeID,
@@ -893,9 +892,8 @@ namespace Metasound
 					InReferencedNodeID,
 					ReferencedNodeInputVertex->VertexID
 				});
-
-				// If editor-only data, add template node and connections
-#else // WITH_EDITORONLY_DATA
+				// If editor, add template node and connections
+				#else
 				// template node takes on data type of concrete input node's output type
 				const FName DataType = InputNode->Interface.Outputs.Last().TypeName;
 				
@@ -921,7 +919,7 @@ namespace Metasound
 					InReferencedNodeID,
 					ReferencedNodeInputVertex->VertexID
 				});
-#endif  // WITH_EDITORONLY_DATA
+				#endif
 			}
 
 #if WITH_EDITOR
diff --git a/Engine/Plugins/Runtime/Metasound/Source/MetasoundFrontend/Public/DocumentConfigurations/MetasoundPresetDocumentConfiguration.h b/Engine/Plugins/Runtime/Metasound/Source/MetasoundFrontend/Public/DocumentConfigurations/MetasoundPresetDocumentConfiguration.h
deleted file mode 100644
diff --git a/Engine/Plugins/Runtime/Metasound/Source/MetasoundFrontend/Public/MetasoundAssetBase.h b/Engine/Plugins/Runtime/Metasound/Source/MetasoundFrontend/Public/MetasoundAssetBase.h
--- a/Engine/Plugins/Runtime/Metasound/Source/MetasoundFrontend/Public/MetasoundAssetBase.h
+++ b/Engine/Plugins/Runtime/Metasound/Source/MetasoundFrontend/Public/MetasoundAssetBase.h
@@ -222,12 +222,19 @@ PRAGMA_ENABLE_DEPRECATION_WARNINGS
 
 private:
 #if WITH_EDITORONLY_DATA
+	UE_DEPRECATED(5.7, "Call version of function which accepts FMetaSoundAssetCookOptions.")
+	UE_API void UpdateAndRegisterReferencesForSerialization(FName CookPlatformName);
+
 	UE_API void UpdateAndRegisterReferencesForSerialization(const Metasound::Frontend::FMetaSoundAssetCookOptions& InCookOptions);
 #endif // WITH_EDITORONLY_DATA
 
 	// Checks if version is up-to-date. If so, returns true. If false, updates the interfaces within the given asset's document to the most recent version.
 	UE_API bool TryUpdateInterfaceFromVersion(const FMetasoundFrontendVersion& Version, FMetaSoundFrontendDocumentBuilder& Builder);
 
+	// Versions dependencies to most recent version where applicable. If asset is a preset, MetaSound is rebuilt to accommodate any referenced node class interface changes.
+	// Otherwise, automatically updates any nodes and respective dependent classes to accommodate changes to interfaces therein preserving edges/connections where possible.
+	UE_API bool VersionDependencies(FMetaSoundFrontendDocumentBuilder& Builder, bool bInLogWarningsOnDroppedConnection);
+
 	// Returns new interface to be versioned to from the given version. If no interface versioning is
 	// required, returns invalid interface version (no name and invalid version number).
 	FMetasoundFrontendVersion GetTargetInterfaceVersion(const FMetasoundFrontendVersion& InterfaceVersion) const;
diff --git a/Engine/Plugins/Runtime/Metasound/Source/MetasoundFrontend/Public/MetasoundAssetManager.h b/Engine/Plugins/Runtime/Metasound/Source/MetasoundFrontend/Public/MetasoundAssetManager.h
--- a/Engine/Plugins/Runtime/Metasound/Source/MetasoundFrontend/Public/MetasoundAssetManager.h
+++ b/Engine/Plugins/Runtime/Metasound/Source/MetasoundFrontend/Public/MetasoundAssetManager.h
@@ -5,7 +5,6 @@
 #include "MetasoundAssetManagerTransaction.h"
 #include "MetasoundDocumentInterface.h"
 #include "MetasoundFrontendDocument.h"
-#include "MetasoundFrontendDocumentBuilderRegistry.h"
 #include "MetasoundFrontendQuery.h"
 #include "MetasoundFrontendRegistryKey.h"
 #include "MetasoundVertex.h"
@@ -59,6 +58,36 @@ struct FMetaSoundDocumentInfo
 
 namespace Metasound::Frontend
 {
+	namespace AssetTags
+	{
+		UE_DEPRECATED(5.6, "AssetTags no longer public. Construct FMetaSoundAssetClassInfo from AssetData to parse tag fields.")
+		extern const FString UE_API ArrayDelim;
+
+#if WITH_EDITORONLY_DATA
+		UE_DEPRECATED(5.6, "AssetTags no longer public. Construct FMetaSoundAssetClassInfo from AssetData to parse tag fields.")
+		extern const FName UE_API IsPreset;
+#endif // WITH_EDITORONLY_DATA
+
+
+		UE_DEPRECATED(5.6, "AssetTags no longer public. Construct FMetaSoundAssetClassInfo from AssetData to parse tag fields.")
+		extern const FName UE_API AssetClassID;
+
+		UE_DEPRECATED(5.6, "AssetTags no longer public. Construct FMetaSoundAssetClassInfo from AssetData to parse tag fields.")
+		extern const FName UE_API RegistryVersionMajor;
+
+		UE_DEPRECATED(5.6, "AssetTags no longer public. Construct FMetaSoundAssetClassInfo from AssetData to parse tag fields.")
+		extern const FName UE_API RegistryVersionMinor;
+
+#if WITH_EDITORONLY_DATA
+		UE_DEPRECATED(5.6, "AssetTags no longer public. Construct FMetaSoundAssetClassInfo from AssetData to parse tag fields.")
+		extern const FName UE_API RegistryInputTypes;
+
+		UE_DEPRECATED(5.6, "AssetTags no longer public. Construct FMetaSoundAssetClassInfo from AssetData to parse tag fields.")
+		extern const FName UE_API RegistryOutputTypes;
+#endif // WITH_EDITORONLY_DATA
+	} // namespace AssetTags
+
+
 	struct FMetaSoundAssetRegistrationOptions
 	{
 		// If true, forces a re-register of this class (and all class dependencies
diff --git a/Engine/Plugins/Runtime/Metasound/Source/MetasoundFrontend/Public/MetasoundDocumentInterface.h b/Engine/Plugins/Runtime/Metasound/Source/MetasoundFrontend/Public/MetasoundDocumentInterface.h
--- a/Engine/Plugins/Runtime/Metasound/Source/MetasoundFrontend/Public/MetasoundDocumentInterface.h
+++ b/Engine/Plugins/Runtime/Metasound/Source/MetasoundFrontend/Public/MetasoundDocumentInterface.h
@@ -12,12 +12,13 @@
 
 #include "MetasoundDocumentInterface.generated.h"
 
-
 #define UE_API METASOUNDFRONTEND_API
 
+
 // Forward Declarations
 struct FMetaSoundFrontendDocumentBuilder;
 
+
 // UInterface for all MetaSound UClasses that implement a MetaSound document
 // as a means for accessing data via code, scripting, execution, or node
 // class generation.
@@ -66,9 +67,76 @@ private:
 	friend struct FMetaSoundFrontendDocumentBuilder;
 };
 
-// Forward declaration for building/backwards compat. Now found in separate header: MetasoundFrontendDocumentBuilderRegistry.h.
 namespace Metasound::Frontend
 {
-	class IDocumentBuilderRegistry;
+	class IDocumentBuilderRegistry
+	{
+	public:
+		virtual ~IDocumentBuilderRegistry() = default;
+
+#if WITH_EDITORONLY_DATA
+		// Given the provided builder, removes paged data within the associated document for a cooked build.
+		// This function removes graphs and input defaults which are not to ever be used by a given cook
+		// platform, allowing users to optimize away data and scale the amount of memory required for
+		// initial load of input UObjects and graph topology, which can also positively effect runtime
+		// performance as well, etc. Returns true if builder modified the document, false if not.
+		UE_DEPRECATED(5.7, "Use StripUnusedPages(...) instead.")
+		virtual bool CookPages(FName CookPlatformName, FMetaSoundFrontendDocumentBuilder& Builder) const = 0;
+#endif // WITH_EDITORONLY_DATA
+
+		virtual FMetaSoundFrontendDocumentBuilder* FindBuilder(TScriptInterface<IMetaSoundDocumentInterface> MetaSound) const = 0;
+		virtual FMetaSoundFrontendDocumentBuilder* FindBuilder(const FMetasoundFrontendClassName& ClassName, const FTopLevelAssetPath& AssetPath) const = 0;
+		virtual FMetaSoundFrontendDocumentBuilder* FindOutermostBuilder(const UObject& InSubObject) const = 0;
+
+#if WITH_EDITORONLY_DATA
+		// Find the existing builder for the given MetaSound, or optionally begin building by attaching a new builder.  Only available
+		// in builds with editor only data as building serialized assets (which may have template nodes, cooked builds do not) is only
+		// supported when editor data is loaded. Creating transient builders can simply be done by passing a new MetaSound asset to
+		// a FMetaSoundFrontendDocumentBuilder constructor, or this registry's implementation may supply its own create call for tracking
+		// and reuse purposes.
+		virtual FMetaSoundFrontendDocumentBuilder& FindOrBeginBuilding(TScriptInterface<IMetaSoundDocumentInterface> MetaSound) = 0;
+#endif // WITH_EDITORONLY_DATA
+
+		// Removes builder from registry, clearing any cached builder state. (Optionally) forces unregistration from the Frontend Node Class Registry
+		// (If the builder has outstanding transactions, unregistration from the Node Class Registry will occur regardless).
+		virtual bool FinishBuilding(const FMetasoundFrontendClassName& InClassName, bool bForceUnregisterNodeClass = false) const = 0;
+		virtual bool FinishBuilding(const FMetasoundFrontendClassName& InClassName, const FTopLevelAssetPath& InAssetPath, bool bForceUnregisterNodeClass = false) const = 0;
+
+		// Reloads the given builder, maintaining all modify delegate subscriptions. Returns true if builder was found and reloaded,
+		// false if not found.
+		virtual bool ReloadBuilder(const FMetasoundFrontendClassName& InClassName) const = 0;
+
+		// Given the provided document and its respective pages, provides the PageID to be used for runtime IGraph and proxy generation.
+		// Returns true if valid page was found, false if not.
+		UE_DEPRECATED(5.7, "This is no longer supported. FindPreferredPage(...)")
+		virtual FGuid ResolveTargetPageID(const FMetasoundFrontendGraphClass& InGraphClass) const = 0;
+
+		// Given the provided input and its respective paged default values, provides the default PageID to be used for runtime IGraph and proxy generation.
+		// Returns true if valid page was found, false if not.
+		UE_DEPRECATED(5.7, "This is no longer supported. FindPreferredPage(...)")
+		virtual FGuid ResolveTargetPageID(const FMetasoundFrontendClassInput& InClassInput) const = 0;
+
+		// Given the provided array of default values, provides the default PageID to be used for runtime IGraph and proxy generation.
+		// Returns true if valid page was found, false if not.
+		UE_DEPRECATED(5.7, "This is no longer supported. FindPreferredPage(...)")
+		virtual FGuid ResolveTargetPageID(const TArray<FMetasoundFrontendClassInputDefault>& Defaults) const = 0;
+
+		UE_DEPRECATED(5.7, "This is only used to support existing deprecated functionality")
+		virtual TArrayView<const FGuid> GetPageOrder() const = 0;
+
+#if WITH_EDITORONLY_DATA
+		UE_DEPRECATED(5.7, "This is only used to support existing deprecated functionality")
+		virtual TArray<FGuid> GetCookedTargetPages(FName InPlatformName) const = 0;
+
+		UE_DEPRECATED(5.7, "This is only used to support existing deprecated functionality")
+		virtual TArray<FGuid> GetCookedPageOrder(FName InPlatformName) const = 0;
+#endif // WITH_EDITORONLY_DATA
+
+		static UE_API IDocumentBuilderRegistry* Get();
+		static UE_API IDocumentBuilderRegistry& GetChecked();
+		static UE_API void Deinitialize();
+		static UE_API void Initialize(TUniquePtr<IDocumentBuilderRegistry>&& InBuilderRegistry);
+	};
 } // namespace Metasound::Frontend
+
 #undef UE_API
diff --git a/Engine/Plugins/Runtime/Metasound/Source/MetasoundFrontend/Public/MetasoundFrontendDocument.h b/Engine/Plugins/Runtime/Metasound/Source/MetasoundFrontend/Public/MetasoundFrontendDocument.h
--- a/Engine/Plugins/Runtime/Metasound/Source/MetasoundFrontend/Public/MetasoundFrontendDocument.h
+++ b/Engine/Plugins/Runtime/Metasound/Source/MetasoundFrontend/Public/MetasoundFrontendDocument.h
@@ -8,7 +8,6 @@
 #include "IAudioParameterInterfaceRegistry.h"
 #include "Internationalization/Text.h"
 #include "MetasoundAccessPtr.h"
-#include "MetasoundFrontendDocumentConfiguration.h"
 #include "MetasoundFrontendLiteral.h"
 #include "MetasoundOperatorData.h"
 #include "MetasoundNodeInterface.h"
@@ -2006,6 +2005,18 @@ class UMetaSoundFrontendMemberMetadata : public UObject
 	GENERATED_BODY()
 
 public:
+	UE_DEPRECATED(5.5, "Implementation moved to child editor class instead of compiled out (not required by Frontend representation")
+	virtual void ForceRefresh() { }
+
+	UE_DEPRECATED(5.5, "Default is no longer required to be stored or represented in metadata and may differ in paged or non-paged implementation")
+	FMetasoundFrontendLiteral GetDefault() const { return FMetasoundFrontendLiteral(); }
+
+	UE_DEPRECATED(5.5, "Implementation moved to child editor class instead of compiled out (not required by Frontend representation")
+	virtual EMetasoundFrontendLiteralType GetLiteralType() const { return EMetasoundFrontendLiteralType::None; }
+
+	UE_DEPRECATED(5.5, "Default is no longer required to be stored or represented in metadata and may differ in paged or non-paged implementation")
+	virtual void SetFromLiteral(const FMetasoundFrontendLiteral& InLiteral, const FGuid& InPageID = Metasound::Frontend::DefaultPageID) { }
+
 #if WITH_EDITORONLY_DATA
 	UPROPERTY()
 	FGuid MemberID;
@@ -2061,11 +2072,6 @@ public:
 	UPROPERTY()
 	TArray<FMetasoundFrontendClass> Dependencies;
 
-#if WITH_EDITORONLY_DATA
-	UPROPERTY()
-	TInstancedStruct<FMetaSoundFrontendDocumentConfiguration> Configuration;
-#endif // WITH_EDITORONLY_DATA
-
 	uint32 GetNextIdCounter() const
 	{
 		return IdCounter++;
diff --git a/Engine/Plugins/Runtime/Metasound/Source/MetasoundFrontend/Public/MetasoundFrontendDocumentBuilder.h b/Engine/Plugins/Runtime/Metasound/Source/MetasoundFrontend/Public/MetasoundFrontendDocumentBuilder.h
--- a/Engine/Plugins/Runtime/Metasound/Source/MetasoundFrontend/Public/MetasoundFrontendDocumentBuilder.h
+++ b/Engine/Plugins/Runtime/Metasound/Source/MetasoundFrontend/Public/MetasoundFrontendDocumentBuilder.h
@@ -8,16 +8,13 @@
 #include "MetasoundDocumentInterface.h"
 #include "MetasoundFrontendDocument.h"
 #include "MetasoundFrontendDocumentModifyDelegates.h"
-#include "MetasoundFrontendGraphLayer.h"
 #include "MetasoundFrontendRegistries.h"
 #include "MetasoundFrontendTransform.h"
 #include "MetasoundLiteral.h"
 #include "MetasoundOperatorSettings.h"
 #include "MetasoundVertex.h"
 #include "StructUtils/InstancedStruct.h"
-#include "StructUtils/StructView.h"
 #include "Templates/Function.h"
-#include "UObject/Class.h"
 #include "UObject/ScriptInterface.h"
 
 #include "MetasoundFrontendDocumentBuilder.generated.h"
@@ -26,7 +23,6 @@
 
 
 // Forward Declarations
-struct FMetaSoundFrontendDocumentConfiguration;
 class FMetasoundAssetBase;
 
 
@@ -127,6 +123,9 @@ class UMetaSoundBuilderDocument : public UObject, public IMetaSoundDocumentInter
 public:
 	UE_API UMetaSoundBuilderDocument(const FObjectInitializer& ObjectInitializer);
 
+	UE_DEPRECATED(5.5, "Use overload supplying MetaSound to copy (builder documents no longer supported for cases outside of cloned document registration.")
+	static UE_API UMetaSoundBuilderDocument& Create(const UClass& InBuilderClass);
+
 	// Create and return a valid builder document which copies the provided interface's document & class
 	static UE_API UMetaSoundBuilderDocument& Create(const IMetaSoundDocumentInterface& InDocToCopy);
 
@@ -153,17 +152,6 @@ public:
 private:
 	UE_API virtual FMetasoundFrontendDocument& GetDocument() override;
 
-#if WITH_EDITORONLY_DATA
-	// Convenience template function for attempting to get the underlying pointer of the set configuration.
-	// If the provided configuration type is not that which the current document is set to or no configuration
-	// is set, returns null.
-	template<typename TConfigType>
-	TConfigType* GetDocumentConfigurationAs()
-	{
-		return GetDocument().Configuration.GetMutablePtr<TConfigType>();
-	}
-#endif // WITH_EDITORONLY_DATA
-
 	UE_API virtual void OnBeginActiveBuilder() override;
 	UE_API virtual void OnFinishActiveBuilder() override;
 
@@ -278,11 +266,11 @@ PRAGMA_ENABLE_EXPERIMENTAL_WARNINGS
 	// so external delegates must be relinked if desired.
 	UE_API void ClearDocument(TSharedRef<Metasound::Frontend::FDocumentModifyDelegates> ModifyDelegates);
 
+	UE_DEPRECATED(5.5, "Use ClearDocument instead")
+	void ClearGraph() {  }
+
 #if WITH_EDITORONLY_DATA
 	UE_API bool ClearMemberMetadata(const FGuid& InMemberID);
-
-	UE_EXPERIMENTAL(5.8, "Configurator is in active development and subject to change")
-	UE_API bool ConfigureDocument(bool bLogFailure = false);
 #endif // WITH_EDITORONLY_DATA
 
 	UE_API bool ContainsDependencyOfType(EMetasoundFrontendClassType ClassType) const;
@@ -309,6 +297,12 @@ PRAGMA_ENABLE_EXPERIMENTAL_WARNINGS
 	UE_API const FMetasoundFrontendClassOutput* DuplicateGraphOutput(FName ExistingName, FName NewName);
 	UE_API const FMetasoundFrontendVariable* DuplicateGraphVariable(FName ExistingName, FName NewName, const FGuid* InPageID = nullptr);
 
+	UE_DEPRECATED(5.6, "Use the duplicate overload that supplies existing name and new name and returns input")
+	UE_API const FMetasoundFrontendNode* DuplicateGraphInput(const FMetasoundFrontendClassInput& InClassInput, const FName InName, const FGuid* InPageID = nullptr);
+
+	UE_DEPRECATED(5.6, "Use the duplicate overload that supplies existing name and new name and returns output")
+	UE_API const FMetasoundFrontendNode* DuplicateGraphOutput(const FMetasoundFrontendClassOutput& InClassOutput, const FName InName, const FGuid* InPageID = nullptr);
+
 	TConstStructView<FMetasoundFrontendClassInterface> FindClassInterfaceOverride(const FGuid& InNodeID, const FGuid* InPageID = nullptr) const;
 
 #if WITH_EDITORONLY_DATA
@@ -388,25 +382,12 @@ PRAGMA_ENABLE_EXPERIMENTAL_WARNINGS
 	UE_API int32 FindPageIndex(const FGuid& InPageID) const;
 
 	UE_API const FMetasoundFrontendDocument& GetConstDocumentChecked() const;
-
-#if WITH_EDITORONLY_DATA
-	// Returns struct view of document configuration if set.
-	UE_EXPERIMENTAL(5.8, "Configurator is in active development and subject to change")
-	UE_API TConstStructView<FMetaSoundFrontendDocumentConfiguration> GetConstDocumentConfiguration() const;
-
-	// Convenience template function for attempting to get the underlying const pointer of the set configuration.
-	// If the provided configuration type is not that which the current document is set to or no configuration
-	// is set, returns null.
-	template<typename TConfigType>
-	const TConfigType* GetConstDocumentConfigurationAs() const
-	{
-		return GetConstDocumentConfiguration().GetPtr<TConfigType>();
-	}
-#endif // WITH_EDITORONLY_DATA
-
 	UE_API const IMetaSoundDocumentInterface& GetConstDocumentInterfaceChecked() const;
 	UE_API const FString GetDebugName() const;
 
+	UE_DEPRECATED(5.5, "Use GetConstDocumentChecked() instead")
+	UE_API const FMetasoundFrontendDocument& GetDocument() const;
+
 	// The graph ID used when requests are made to mutate specific paged graph topology (ex. adding or removing nodes or edges)
 	UE_API const FGuid& GetBuildPageID() const;
 
@@ -437,12 +418,11 @@ PRAGMA_ENABLE_EXPERIMENTAL_WARNINGS
 	// asset validation errors.
 	UE_API FMetasoundFrontendClassName GenerateNewClassName();
 
-	UE_API FName GenerateUniqueInputName(const FString& InBaseName) const;
-	UE_API FName GenerateUniqueOutputName(const FString& InBaseName) const;
-	UE_API FName GenerateUniqueVariableName(const FString& InBaseName) const;
-
 	UE_API Metasound::Frontend::FDocumentModifyDelegates& GetDocumentDelegates();
 
+	UE_DEPRECATED(5.5, "Use GetConstDocumentInterfaceChecked instead")
+	UE_API const IMetaSoundDocumentInterface& GetDocumentInterface() const;
+
 	// Path for document object provided at construction time. Cached on builder as a useful means of debugging
 	// and enables weak reference removal from the builder registry should the object be mid-destruction.
 	UE_API const FTopLevelAssetPath& GetHintPath() const;
@@ -461,6 +441,9 @@ PRAGMA_ENABLE_EXPERIMENTAL_WARNINGS
 
 	UE_API EMetasoundFrontendVertexAccessType GetNodeInputAccessType(const FGuid& InNodeID, const FGuid& InVertexID, const FGuid* InPageID = nullptr) const;
 
+	UE_DEPRECATED(5.5, "Use FindNodeInputClass overloads instead and use GetDefaults() on result (now supports page values)")
+	UE_API const FMetasoundFrontendLiteral* GetNodeInputClassDefault(const FGuid& InNodeID, const FGuid& InVertexID, const FGuid* InPageID = nullptr) const;
+
 #if WITH_EDITORONLY_DATA
 	// Returns friendly name of the given node input vertex. Unlike a straight look-up of the vertex metadata value, this
 	// tries and determines the best display name based on the parent node's class, if no display name is given
@@ -473,6 +456,9 @@ PRAGMA_ENABLE_EXPERIMENTAL_WARNINGS
 	UE_API FText GetNodeOutputDisplayName(const FGuid& InNodeID, const FName VertexName, const FGuid* InPageID = nullptr) const;
 #endif // WITH_EDITORONLY_DATA
 
+	UE_DEPRECATED(5.5, "Use FindNodeInputDefault and returned struct Value member instead")
+	UE_API const FMetasoundFrontendLiteral* GetNodeInputDefault(const FGuid& InNodeID, const FGuid& InVertexID, const FGuid* InPageID = nullptr) const;
+
 	UE_API EMetasoundFrontendVertexAccessType GetNodeOutputAccessType(const FGuid& InNodeID, const FGuid& InVertexID, const FGuid* InPageID = nullptr) const;
 
 #if WITH_EDITORONLY_DATA
@@ -485,11 +471,6 @@ PRAGMA_ENABLE_EXPERIMENTAL_WARNINGS
 	// Returns the default value set for the variable with the given name on the given page.
 	UE_API const FMetasoundFrontendLiteral* GetGraphVariableDefault(FName InputName, const FGuid* InPageID = nullptr) const;
 
-#if WITH_EDITORONLY_DATA
-	// If PageID set, returns whether or not the given page is editable.  If no page ID is set, returns whether graph with set build ID is editable.
-	UE_API bool GetIsGraphEditable(const FGuid* InPageID = nullptr) const;
-#endif // WITH_EDITORONLY_DATA
-
 	// Initializes the builder's document, using the (optional) provided document template, (optional) class name, and (optionally) whether or not to reset the existing class version.
 	UE_API void InitDocument(const FMetasoundFrontendDocument* InDocumentTemplate = nullptr, const FMetasoundFrontendClassName* InNewClassName = nullptr, bool bResetVersion = true);
 
@@ -499,22 +480,10 @@ PRAGMA_ENABLE_EXPERIMENTAL_WARNINGS
 
 	UE_API void InitNodeLocations();
 
-	UE_API bool IsDependencyReferenced(const FGuid& InClassID) const;
-
-#if WITH_EDITORONLY_DATA
-	// Determines if the builder's document configuration is of the given type
-	template<typename TConfigType>
-	const bool IsDocumentConfiguredBy() const
-	{
-		if (const UScriptStruct* Struct = GetConstDocumentChecked().Configuration.GetScriptStruct())
-		{
-			return Struct->IsChildOf(TConfigType::StaticStruct());
-		}
-
-		return false;
-	}
-#endif // WITH_EDITORONLY_DATA
+	UE_DEPRECATED(5.5, "Use invalidate overload that is provided new version of modify delegates")
+	void InvalidateCache() { }
 
+	UE_API bool IsDependencyReferenced(const FGuid& InClassID) const;
 	UE_API bool IsNodeInputConnected(const FGuid& InNodeID, const FGuid& InVertexID, const FGuid* InPageID = nullptr) const;
 	UE_API bool IsNodeInputConnectionUserModifiable(const FGuid& InNodeID, const FGuid& InVertexID, const FGuid* InPageID = nullptr) const;
 	UE_API bool IsNodeOutputConnected(const FGuid& InNodeID, const FGuid& InVertexID, const FGuid* InPageID = nullptr) const;
@@ -542,6 +511,11 @@ PRAGMA_ENABLE_EXPERIMENTAL_WARNINGS
 
 	UE_API bool ModifyInterfaces(Metasound::Frontend::FModifyInterfaceOptions&& InOptions);
 
+	UE_DEPRECATED(5.5,
+		"Cache invalidation may require new copy of delegates. In addition, re-priming is discouraged. "
+		"To enforce this, new recommended pattern is to construct a new builder instead")
+	UE_API void ReloadCache();
+
 	// Removes all dependencies with the given ClassID. Removes any nodes (and corresponding edges) remaining in any MetaSound paged graphs.
 	UE_API bool RemoveDependency(const FGuid& InClassID);
 
@@ -589,6 +563,9 @@ PRAGMA_ENABLE_EXPERIMENTAL_WARNINGS
 
 	UE_API bool RemoveNodeInputDefault(const FGuid& InNodeID, const FGuid& InVertexID, const FGuid* InPageID = nullptr);
 	UE_API bool RemoveUnusedDependencies();
+
+	UE_DEPRECATED(5.5, "Use GenerateNewClassName instead")
+	UE_API bool RenameRootGraphClass(const FMetasoundFrontendClassName& InName);
 	
 	using FVertexNameAndType = TTuple<FName, FName>;
 	// Replaces existing dependency with a new one. Removes all nodes using this class, and readds nodes with new class, 
@@ -617,14 +594,6 @@ PRAGMA_ENABLE_EXPERIMENTAL_WARNINGS
 	// API (exposed as a mechanism for mutating via controllers while deprecating.  Option will be removed
 	// in a future build).
 	UE_API bool SetBuildPageID(const FGuid& InBuildPageID, bool bBroadcastDelegate = true);
-#endif // WITH_EDITORONLY_DATA
-
-#if WITH_EDITORONLY_DATA
-	UE_API void SetDocumentConfiguration(TInstancedStruct<FMetaSoundFrontendDocumentConfiguration> InConfig);
-
-	// If page ID is provided, updates the graph's editable flag to the given value. If no page ID is provided,
-	// sets *all* paged graph fields to the given value.
-	UE_API void SetIsGraphEditable(bool bIsEditable, const FGuid* PageID = nullptr) const;
 
 	// Sets the given input`s IsAdvancedDisplay state. AdvancedDisplay pins are hidden in the node by default.
 	// returns true if state was changed.
@@ -769,6 +738,10 @@ PRAGMA_ENABLE_EXPERIMENTAL_WARNINGS
 	// replacing nodes with other, registry-defined concrete node class instances. Returns true if any template nodes were processed.
 	UE_API bool TransformTemplateNodes();
 
+	// Versions legacy document members that contained interface information
+	UE_DEPRECATED(5.5, "Moved to internally implemented versioning logic")
+	UE_API bool VersionInterfaces();
+
 	// Struct enabling property migration of data that must be applied prior to versioning logic
 	struct IPropertyVersionTransform
 	{
@@ -784,8 +757,6 @@ PRAGMA_ENABLE_EXPERIMENTAL_WARNINGS
 #endif // WITH_EDITORONLY_DATA
 
 private:
-	bool bInDocConfigScope = false;
-
 	using FFinalizeNodeFunctionRef = TFunctionRef<void(FMetasoundFrontendNode&, const Metasound::Frontend::FNodeRegistryKey&)>;
 
 	TInstancedStruct<FMetaSoundFrontendNodeConfiguration> CreateFrontendNodeConfiguration(const FMetasoundFrontendClassMetadata& InClassMetadata) const;
@@ -905,4 +876,5 @@ private:
 
 	FTopLevelAssetPath HintPath;
 };
+
 #undef UE_API
diff --git a/Engine/Plugins/Runtime/Metasound/Source/MetasoundFrontend/Public/MetasoundFrontendDocumentBuilderRegistry.h b/Engine/Plugins/Runtime/Metasound/Source/MetasoundFrontend/Public/MetasoundFrontendDocumentBuilderRegistry.h
deleted file mode 100644
diff --git a/Engine/Plugins/Runtime/Metasound/Source/MetasoundFrontend/Public/MetasoundFrontendDocumentConfiguration.h b/Engine/Plugins/Runtime/Metasound/Source/MetasoundFrontend/Public/MetasoundFrontendDocumentConfiguration.h
deleted file mode 100644
diff --git a/Engine/Plugins/Runtime/Metasound/Source/MetasoundFrontend/Public/MetasoundFrontendDocumentModifyDelegates.h b/Engine/Plugins/Runtime/Metasound/Source/MetasoundFrontend/Public/MetasoundFrontendDocumentModifyDelegates.h
--- a/Engine/Plugins/Runtime/Metasound/Source/MetasoundFrontend/Public/MetasoundFrontendDocumentModifyDelegates.h
+++ b/Engine/Plugins/Runtime/Metasound/Source/MetasoundFrontend/Public/MetasoundFrontendDocumentModifyDelegates.h
@@ -28,14 +28,10 @@ namespace Metasound::Frontend
 		FGuid PageID;
 	};
 
+	struct FDocumentPresetStateChangedArgs
+	{
+	};
 
-	struct FDocumentConfigChangedArgs {};
-	DECLARE_MULTICAST_DELEGATE_OneParam(FOnDocumentConfigurationChanged, const FDocumentConfigChangedArgs& /* Args */);
-
-	// Soft deprecated in 5.8 in favor of FOnDocumentConfigurationChanged.  Once delegate is removed, it can be officially deprecated.
-	struct FDocumentPresetStateChangedArgs {};
-
-	// Soft deprecated in 5.8 as usage is no longer valid (migrated to OnDocumentConfigurationChanged)
 	DECLARE_MULTICAST_DELEGATE_OneParam(FOnDocumentPresetStateChanged, const FDocumentPresetStateChangedArgs& /* Args */);
 
 	DECLARE_MULTICAST_DELEGATE_OneParam(FOnDocumentPageAdded, const FDocumentMutatePageArgs& /* Args */);
@@ -112,20 +108,15 @@ namespace Metasound::Frontend
 		UE_API FDocumentModifyDelegates();
 		UE_API FDocumentModifyDelegates(const FMetasoundFrontendDocument& Document);
 
-PRAGMA_DISABLE_DEPRECATION_WARNINGS
 		~FDocumentModifyDelegates() = default;
-PRAGMA_ENABLE_DEPRECATION_WARNINGS
 
 		FOnMetaSoundFrontendDocumentMutateMetadata OnDocumentMetadataChanged;
 		FOnMetaSoundFrontendDocumentMutateArray OnDependencyAdded;
 		FOnMetaSoundFrontendDocumentRemoveSwappingArray OnRemoveSwappingDependency;
 		FOnMetaSoundFrontendDocumentRenameClass OnRenamingDependencyClass;
 
-		// Soft deprecated in 5.8: "Use OnDocumentConfigurationChanged instead"
 		FOnDocumentPresetStateChanged OnPresetStateChanged;
 
-		FOnDocumentConfigurationChanged OnDocumentConfigurationChanged;
-
 		FPageModifyDelegates PageDelegates;
 		FInterfaceModifyDelegates InterfaceDelegates;
 
diff --git a/Engine/Plugins/Runtime/Metasound/Source/MetasoundFrontend/Public/MetasoundFrontendGraphLayer.h b/Engine/Plugins/Runtime/Metasound/Source/MetasoundFrontend/Public/MetasoundFrontendGraphLayer.h
deleted file mode 100644
diff --git a/Engine/Plugins/Runtime/Metasound/Source/MetasoundFrontend/Public/MetasoundFrontendLiteral.h b/Engine/Plugins/Runtime/Metasound/Source/MetasoundFrontend/Public/MetasoundFrontendLiteral.h
--- a/Engine/Plugins/Runtime/Metasound/Source/MetasoundFrontend/Public/MetasoundFrontendLiteral.h
+++ b/Engine/Plugins/Runtime/Metasound/Source/MetasoundFrontend/Public/MetasoundFrontendLiteral.h
@@ -2,7 +2,6 @@
 
 #pragma once
 
-#include "Algo/Transform.h"
 #include "AudioParameterControllerInterface.h"
 #include "MetasoundLiteral.h"
 
@@ -75,7 +74,6 @@ struct FMetasoundFrontendLiteral
 	struct FDefaultArray { int32 Num = 0; };
 
 	FMetasoundFrontendLiteral() = default;
-
 	UE_API FMetasoundFrontendLiteral(const FAudioParameter& InParameter);
 
 private:
@@ -127,17 +125,17 @@ public:
 	UE_API void SetType(EMetasoundFrontendLiteralType InType);
 
 	UE_API void Set(FDefault InValue);
-	UE_API void Set(FDefaultArray InValue);
+	UE_API void Set(const FDefaultArray& InValue);
 	UE_API void Set(bool InValue);
-	UE_API void Set(TArray<bool> InValue);
+	UE_API void Set(const TArray<bool>& InValue);
 	UE_API void Set(int32 InValue);
-	UE_API void Set(TArray<int32> InValue);
+	UE_API void Set(const TArray<int32>& InValue);
 	UE_API void Set(float InValue);
-	UE_API void Set(TArray<float> InValue);
-	UE_API void Set(FString InValue);
-	UE_API void Set(TArray<FString> InValue);
+	UE_API void Set(const TArray<float>& InValue);
+	UE_API void Set(const FString& InValue);
+	UE_API void Set(const TArray<FString>& InValue);
 	UE_API void Set(UObject* InValue);
-	UE_API void Set(TArray<UObject*> InValue);
+	UE_API void Set(const TArray<UObject*>& InValue);
 
 	UE_API void SetFromLiteral(const Metasound::FLiteral& InLiteral);
 
@@ -171,38 +169,22 @@ public:
 		return A.IsEqual(B);
 	}
 
+
 private:
 	// Remove all values.
 	UE_API void Empty();
 };
 
-namespace Metasound::Frontend
+namespace Metasound
 {
-	template <typename ObjectType = UObject>
-	FMetasoundFrontendLiteral MakeFrontendLiteral(TConstArrayView<TObjectPtr<ObjectType>> ObjectPtrs)
+	namespace Frontend
 	{
-		TArray<UObject*> Objects;
-		Algo::Transform(ObjectPtrs, Objects,
-			[](const TObjectPtr<ObjectType> Obj) { return Obj.Get(); });
+		// Convenience function to convert Metasound::ELiteralType to EMetasoundFrontendLiteralType.
+		METASOUNDFRONTEND_API EMetasoundFrontendLiteralType GetMetasoundFrontendLiteralType(ELiteralType InLiteralType);
 
-		FMetasoundFrontendLiteral Value;
-		Value.Set(MoveTemp(Objects));
-		return Value;
+		// Convenience function to convert EMetasoundFrontendLiteralType to Metasound::ELiteralType.
+		METASOUNDFRONTEND_API ELiteralType GetMetasoundLiteralType(EMetasoundFrontendLiteralType InLiteralType);
 	}
-
-	template <typename LiteralType>
-	FMetasoundFrontendLiteral MakeFrontendLiteral(LiteralType InitValue)
-	{
-		FMetasoundFrontendLiteral Value;
-		Value.Set(MoveTemp(InitValue));
-		return Value;
-	}
-
-	// Convenience function to convert Metasound::ELiteralType to EMetasoundFrontendLiteralType.
-	UE_API EMetasoundFrontendLiteralType GetMetasoundFrontendLiteralType(ELiteralType InLiteralType);
-
-	// Convenience function to convert EMetasoundFrontendLiteralType to Metasound::ELiteralType.
-	UE_API ELiteralType GetMetasoundLiteralType(EMetasoundFrontendLiteralType InLiteralType);
 }
 
 #undef UE_API
diff --git a/Engine/Plugins/Runtime/Metasound/Source/MetasoundFrontend/Public/MetasoundFrontendTransform.h b/Engine/Plugins/Runtime/Metasound/Source/MetasoundFrontend/Public/MetasoundFrontendTransform.h
--- a/Engine/Plugins/Runtime/Metasound/Source/MetasoundFrontend/Public/MetasoundFrontendTransform.h
+++ b/Engine/Plugins/Runtime/Metasound/Source/MetasoundFrontend/Public/MetasoundFrontendTransform.h
@@ -15,7 +15,6 @@
 
 // Forward Declarations
 struct FMetaSoundFrontendDocumentBuilder;
-struct FMetaSoundFrontendPresetDocumentConfiguration;
 
 
 namespace Metasound::Frontend
