Minor bug hawk feedback
#rnx

diff --git a/Engine/Plugins/Runtime/Metasound/Source/MetasoundEditor/Private/MetasoundDetailCustomization.cpp b/Engine/Plugins/Runtime/Metasound/Source/MetasoundEditor/Private/MetasoundDetailCustomization.cpp
--- a/Engine/Plugins/Runtime/Metasound/Source/MetasoundEditor/Private/MetasoundDetailCustomization.cpp
+++ b/Engine/Plugins/Runtime/Metasound/Source/MetasoundEditor/Private/MetasoundDetailCustomization.cpp
@@ -529,17 +529,19 @@ namespace Metasound::Editor
 
 	const FMetasoundFrontendDocument& FMetasoundPagesDetailCustomization::GetConstDocumentChecked() const
 	{
+		check(IsBuilderValid());
 		return Builder->GetBuilder().GetConstDocumentChecked();
 	}
 
 	UObject& FMetasoundPagesDetailCustomization::GetMetaSoundChecked() const
 	{
+		check(IsBuilderValid());
 		return Builder->GetBuilder().CastDocumentObjectChecked<UObject>();
 	}
 
 	bool FMetasoundPagesDetailCustomization::IsBuilderValid() const
 	{
-		return Builder->GetBuilder().IsValid();
+		return Builder.IsValid() && Builder->GetBuilder().IsValid();
 	}
 
 	void FMetasoundPagesDetailCustomization::RebuildImplemented()
diff --git a/Engine/Plugins/Runtime/Metasound/Source/MetasoundEditor/Private/MetasoundEditorGraphBuilder.h b/Engine/Plugins/Runtime/Metasound/Source/MetasoundEditor/Private/MetasoundEditorGraphBuilder.h
--- a/Engine/Plugins/Runtime/Metasound/Source/MetasoundEditor/Private/MetasoundEditorGraphBuilder.h
+++ b/Engine/Plugins/Runtime/Metasound/Source/MetasoundEditor/Private/MetasoundEditorGraphBuilder.h
@@ -159,7 +159,7 @@ namespace Metasound
 				const TArray<FMetasoundFrontendClassInputDefault>& InDefaultLiterals = { },
 				const FName* InNameBase = nullptr);
 
-			// Creates a unique class output with the given default data. Output is not assigned a NodeID.
+			// Creates a unique class output with the given default data.
 			static FMetasoundFrontendClassOutput CreateUniqueClassOutput(
 				UObject& InMetaSound,
 				const FCreateNodeVertexParams& InVertexParams,
