[Metasounds] Recent & Favorite support for graph context menu actions. Better ordering of results when filtering actions.

diff --git a/Engine/Plugins/Runtime/Metasound/Source/MetasoundEditor/Private/MetasoundEditor.cpp b/Engine/Plugins/Runtime/Metasound/Source/MetasoundEditor/Private/MetasoundEditor.cpp
--- a/Engine/Plugins/Runtime/Metasound/Source/MetasoundEditor/Private/MetasoundEditor.cpp
+++ b/Engine/Plugins/Runtime/Metasound/Source/MetasoundEditor/Private/MetasoundEditor.cpp
@@ -366,6 +366,8 @@ namespace Metasound
 
 				if (UMetasoundEditorGraphInput* Input = Cast<UMetasoundEditorGraphInput>(GraphMember.Get()))
 				{
+					GetMutableDefault<UMetasoundEditorSettings>()->AddRecentlyUsedMetaSoundsGraphAction(GraphMember->GetMemberID().ToString());
+					
 					const FScopedTransaction Transaction(LOCTEXT("DropAddNewInputNode", "Drop New MetaSound Input Node"));
 					ParentMetaSound.Modify();
 					MetasoundGraph->Modify();
@@ -388,6 +390,10 @@ namespace Metasound
 
 				if (UMetasoundEditorGraphOutput* Output = Cast<UMetasoundEditorGraphOutput>(GraphMember.Get()))
 				{
+					// Note: Purposedly not registering Output nodes usage into Recently Used nodes.
+					// Indeed, once placed in a graph, they don't appear in the graph action menu anymore, so they would be hidden in the
+					// Recently Used list too anyway.
+					
 					TArray<UMetasoundEditorGraphMemberNode*> Nodes = Output->GetNodes();
 					if (Nodes.IsEmpty())
 					{
@@ -425,6 +431,8 @@ namespace Metasound
 
 				if (UMetasoundEditorGraphVariable* Variable = Cast<UMetasoundEditorGraphVariable>(GraphMember.Get()))
 				{
+					GetMutableDefault<UMetasoundEditorSettings>()->AddRecentlyUsedMetaSoundsGraphAction(GraphMember->GetMemberID().ToString());
+					
 					const FScopedTransaction Transaction(LOCTEXT("DropAddNewVariableNode", "Drop New MetaSound Variable Node"));
 					ParentMetaSound.Modify();
 					MetasoundGraph->Modify();
@@ -597,7 +605,7 @@ namespace Metasound
 							{
 								bDropTargetValid = true;
 
-								Style::GetSlateBrushSafe("MetasoundEditor.Graph.Node.Class.Input");
+								PrimarySymbol = &Style::GetSlateBrushSafe("MetasoundEditor.Graph.Node.Class.Input");
 								SecondarySymbol = nullptr;
 
 								UEdGraphPin* PinUnderCursor = GetHoveredPin();
diff --git a/Engine/Plugins/Runtime/Metasound/Source/MetasoundEditor/Private/MetasoundEditorGraphSchema.cpp b/Engine/Plugins/Runtime/Metasound/Source/MetasoundEditor/Private/MetasoundEditorGraphSchema.cpp
--- a/Engine/Plugins/Runtime/Metasound/Source/MetasoundEditor/Private/MetasoundEditorGraphSchema.cpp
+++ b/Engine/Plugins/Runtime/Metasound/Source/MetasoundEditor/Private/MetasoundEditorGraphSchema.cpp
@@ -8,12 +8,10 @@
 #include "Editor.h"
 #include "Editor/AssetReferenceFilter.h"
 #include "Editor/EditorEngine.h"
-#include "Engine/Selection.h"
 #include "Framework/Commands/GenericCommands.h"
 #include "Framework/MultiBox/MultiBoxBuilder.h"
 #include "GraphEditor.h"
 #include "GraphEditorActions.h"
-#include "GraphEditorSettings.h"
 #include "HAL/IConsoleManager.h"
 #include "Layout/SlateRect.h"
 #include "Metasound.h"
@@ -22,7 +20,6 @@
 #include "MetasoundAssetManager.h"
 #include "MetasoundAssetSubsystem.h"
 #include "MetasoundBuilderSubsystem.h"
-#include "MetasoundDataReference.h"
 #include "MetasoundDocumentBuilderRegistry.h"
 #include "MetasoundDocumentInterface.h"
 #include "MetasoundEditor.h"
@@ -32,32 +29,22 @@
 #include "MetasoundEditorGraphCommentNode.h"
 #include "MetasoundEditorGraphInputNode.h"
 #include "MetasoundEditorGraphNode.h"
-#include "MetasoundEditorGraphValidation.h"
 #include "MetasoundEditorModule.h"
 #include "MetasoundEditorSettings.h"
 #include "MetasoundFrontend.h"
 #include "MetasoundFrontendGraphLinter.h"
-#include "MetasoundFrontendNodesCategories.h"
 #include "MetasoundFrontendNodeTemplateRegistry.h"
-#include "MetasoundFrontendRegistries.h"
+#include "MetasoundFrontendNodesCategories.h"
 #include "MetasoundFrontendRegistryContainer.h"
 #include "MetasoundFrontendSearchEngine.h"
 #include "MetasoundLiteral.h"
-#include "MetasoundNodeInterface.h"
-#include "MetasoundStandardNodesCategories.h"
-#include "MetasoundFrontendNodeTemplateRegistry.h"
 #include "MetasoundUObjectRegistry.h"
-#include "MetasoundVariableNodes.h"
-#include "MetasoundVertex.h"
 #include "NodeTemplates/MetasoundFrontendNodeTemplateAudioAnalyzer.h"
 #include "NodeTemplates/MetasoundFrontendNodeTemplateInput.h"
 #include "NodeTemplates/MetasoundFrontendNodeTemplateReroute.h"
 #include "ScopedTransaction.h"
 #include "Settings/EditorStyleSettings.h"
-#include "Styling/SlateStyleRegistry.h"
-#include "Toolkits/ToolkitManager.h"
 #include "ToolMenus.h"
-#include "UObject/NoExportTypes.h"
 #include "UObject/UObjectHash.h"
 #include "UObject/UObjectIterator.h"
 
@@ -68,6 +55,11 @@
 
 namespace Metasound::Editor
 {
+	// If this fires (very unlikely), find usages on the Common category, and make sure those actions have their Grouping set to a larger number than
+	// Recent + UMetasoundEditorSettings::MaxRecentEntriesEver.
+	static_assert(static_cast<int32>(EPrimaryContextGroup::Common) - static_cast<int32>(EPrimaryContextGroup::Recents) > UMetasoundEditorSettings::MaxRecentEntriesEver,
+		"Common graph actions could be going under the most recently used actions!");
+
 	const FText& GetContextGroupDisplayName(EPrimaryContextGroup InContextGroup)
 	{
 		switch (InContextGroup)
@@ -90,6 +82,12 @@ namespace Metasound::Editor
 			case EPrimaryContextGroup::Variables:
 				return NodeCategories::Variables;
 
+			case EPrimaryContextGroup::Favorites:
+				return NodeCategories::Favorites;
+
+			case EPrimaryContextGroup::Recents:
+				return NodeCategories::Recents;
+
 			case EPrimaryContextGroup::Common:
 			default:
 			{
@@ -98,6 +96,18 @@ namespace Metasound::Editor
 		}
 	}
 
+	template<typename TAction UE_REQUIRES(TIsDerivedFrom<TAction, FMetasoundGraphSchemaAction_NewVariableNode>::Value)>
+	TSharedPtr<TAction> MakeVariableAction(const FText& InCategory, int32 Grouping, const FText& InDisplayName, FGuid InVariableID, const FText& InToolTip)
+	{
+		TSharedPtr<TAction> NewAction = MakeShared<TAction>(
+			InCategory,
+			InDisplayName,
+			InVariableID,
+			InToolTip);
+		NewAction->Grouping = Grouping;
+		return NewAction;
+	}
+
 	namespace SchemaPrivate
 	{
 		static int32 EnableAllVersionsMetaSoundNodeClassCreationCVar = 0;
@@ -241,6 +251,8 @@ PRAGMA_ENABLE_DEPRECATION_WARNINGS
 			using namespace Editor;
 			using namespace Frontend;
 
+			const UMetasoundEditorSettings* Settings = GetDefault<UMetasoundEditorSettings>();
+
 			for (const FConstNodeHandle& NodeHandle : InQuery.NodeHandles)
 			{
 				if (!InQuery.Filter || InQuery.Filter(NodeHandle))
@@ -253,6 +265,33 @@ PRAGMA_ENABLE_DEPRECATION_WARNINGS
 					const FText DisplayName = FText::Format(InQuery.DisplayNameFormat, NodeDisplayName);
 					TSharedPtr<TAction> NewNodeAction = MakeShared<TAction>(GroupName, DisplayName, NodeHandle->GetID(), Tooltip, InQuery.ContextGroup);
 					InQuery.ActionMenuBuilder.AddAction(NewNodeAction);
+
+					{
+						const FString Key = NewNodeAction->GetUniqueStableFavoriteKey();
+						const int32 RecentIndex = Settings->FindRecentlyUsedMetaSoundsGraphAction(Key);
+						if (RecentIndex != INDEX_NONE)
+						{
+							TSharedPtr<TAction> NewRecentAction = MakeShared<TAction>(
+								NodeCategories::Recents,
+								DisplayName,
+								NodeHandle->GetID(),
+								Tooltip,
+								EPrimaryContextGroup::Recents);
+							NewRecentAction->Grouping = static_cast<int32>(EPrimaryContextGroup::Recents) + RecentIndex;
+							InQuery.ActionMenuBuilder.AddAction(NewRecentAction);
+						}
+
+						if (Settings->IsFavoriteMetaSoundsGraphAction(Key))
+						{
+							TSharedPtr<TAction> NewFavAction = MakeShared<TAction>(
+								FText::FormatOrdered(INVTEXT("{0}|{1}"), NodeCategories::Favorites, GroupName),
+								DisplayName,
+								NodeHandle->GetID(),
+								Tooltip,
+								EPrimaryContextGroup::Favorites);
+							InQuery.ActionMenuBuilder.AddAction(NewFavAction);
+						}
+					}
 				}
 			}
 		}
@@ -884,6 +923,11 @@ const FLinearColor& FMetasoundGraphSchemaAction_NewNode::GetIconColor() const
 	return Super::GetIconColor();
 }
 
+FString FMetasoundGraphSchemaAction_NewNode::GetUniqueStableFavoriteKey() const
+{
+	return QueryResult->GetAssetKey().ClassName.ToString();
+}
+
 UEdGraphNode* FMetasoundGraphSchemaAction_NewNode::PerformAction(UEdGraph* ParentGraph, UEdGraphPin* FromPin, const FVector2f& Location, bool bSelectNewNode /* = true */)
 {
 	using namespace Metasound::Editor;
@@ -2158,6 +2202,27 @@ FText UMetasoundEditorGraphSchema::GetPinDisplayName(const UEdGraphPin* Pin) con
 	return Super::GetPinDisplayName(Pin);
 }
 
+float UMetasoundEditorGraphSchema::GetActionFilteredWeight(const FEdGraphSchemaAction& InCurrentAction, const TArray<FString>& InFilterTerms,
+	const TArray<FString>& InSanitizedFilterTerms, const TArray<UEdGraphPin*>& DraggedFromPins) const
+{
+	using namespace Metasound::Editor;
+
+	// Search will surface results following the priority order of EPrimaryContextGroup (highest value == highest priority)
+	// Then, for items having the same priority order, we offset their weight by their actual lexical match (BaseWeight)
+	// Exception for Recently Used section: their priority order (before BaseWeight offset) is such that more recently used items have higher prio
+	
+	const float BaseWeight = Super::GetActionFilteredWeight(InCurrentAction, InFilterTerms, InSanitizedFilterTerms, DraggedFromPins);
+	if (static_cast<int32>(EPrimaryContextGroup::Recents) <= InCurrentAction.GetGrouping() && InCurrentAction.GetGrouping() < static_cast<int32>(EPrimaryContextGroup::Common))
+	{
+		// Highest weight matches for Recent actions should follow most recently used order
+		constexpr int32 RecentsOffset = static_cast<int32>(EPrimaryContextGroup::Recents);
+		return BaseWeight + RecentsOffset * 100000.0f + (InCurrentAction.GetGrouping() - RecentsOffset) * 10000.0f;
+	}
+	
+	// Highest weight matches should follow the display order of their categories
+	return BaseWeight + InCurrentAction.GetGrouping() * 100000.0f;
+}
+
 FLinearColor UMetasoundEditorGraphSchema::GetPinTypeColor(const FEdGraphPinType& PinType) const
 {
 	return Metasound::Editor::FGraphBuilder::GetPinCategoryColor(PinType);
@@ -2590,6 +2655,8 @@ void UMetasoundEditorGraphSchema::GetFunctionActions(FGraphActionMenuBuilder& Ac
 	using namespace Metasound::Engine;
 	using namespace Metasound::Frontend;
 
+	const UMetasoundEditorSettings* Settings = GetDefault<UMetasoundEditorSettings>();
+
 	const FMetaSoundAssetManager& AssetManager = FMetaSoundAssetManager::GetChecked();
 	const FMetaSoundAssetKey ParentAssetKey(InGraphHandle->GetGraphMetadata());
 	QueryNodeClasses([&](TUniquePtr<ISchemaQueryResult>&& Result)
@@ -2654,6 +2721,35 @@ void UMetasoundEditorGraphSchema::GetFunctionActions(FGraphActionMenuBuilder& Ac
 
 		NewNodeAction->QueryResult = TSharedPtr<ISchemaQueryResult>(Result.Release());
 		ActionMenuBuilder.AddAction(NewNodeAction);
+		
+		{
+			const FString Key = NewNodeAction->GetUniqueStableFavoriteKey();
+			const int32 RecentIndex = Settings->FindRecentlyUsedMetaSoundsGraphAction(Key);
+			if (RecentIndex != INDEX_NONE)
+			{
+				TSharedPtr<FMetasoundGraphSchemaAction_NewNode> NewRecentAction = MakeShared<FMetasoundGraphSchemaAction_NewNode>(
+					Metasound::NodeCategories::Recents,
+					NewNodeAction->GetMenuDescription(),
+					NewNodeAction->GetTooltipDescription(),
+					EPrimaryContextGroup::Recents,
+					NewNodeAction->GetKeywords());
+				NewRecentAction->Grouping = static_cast<int32>(EPrimaryContextGroup::Recents) + RecentIndex;
+				NewRecentAction->QueryResult = NewNodeAction->QueryResult;
+				ActionMenuBuilder.AddAction(NewRecentAction);
+			}
+
+			if (Settings->IsFavoriteMetaSoundsGraphAction(Key))
+			{
+				TSharedPtr<FMetasoundGraphSchemaAction_NewNode> NewFavAction = MakeShared<FMetasoundGraphSchemaAction_NewNode>(
+					FText::FormatOrdered(INVTEXT("{0}|{1}"), Metasound::NodeCategories::Favorites, NewNodeAction->GetCategory()),
+					NewNodeAction->GetMenuDescription(),
+					NewNodeAction->GetTooltipDescription(),
+					EPrimaryContextGroup::Favorites,
+					NewNodeAction->GetKeywords());
+				NewFavAction->QueryResult = NewNodeAction->QueryResult;
+				ActionMenuBuilder.AddAction(NewFavAction);
+			}
+		}
 	});
 }
 
@@ -2702,12 +2798,18 @@ void UMetasoundEditorGraphSchema::GetVariableActions(FGraphActionMenuBuilder& Ac
 	}
 
 	// Create actions for each variable.
-	const FText& GroupName = GetContextGroupDisplayName(EPrimaryContextGroup::Variables);
+	static const FText& BaseGroupName = GetContextGroupDisplayName(EPrimaryContextGroup::Variables);
+	static const FText& RecentGroupName = GetContextGroupDisplayName(EPrimaryContextGroup::Recents);
+	static const FText& FavoriteGroupName = FText::FormatOrdered(INVTEXT("{0}|{1}"), GetContextGroupDisplayName(EPrimaryContextGroup::Favorites), BaseGroupName);
+	const UMetasoundEditorSettings* Settings = GetDefault<UMetasoundEditorSettings>();
 	for (FConstVariableHandle& Variable : Variables)
 	{
 		const FText VariableDisplayName = FGraphBuilder::GetDisplayName(*Variable);
 		const FGuid VariableID = Variable->GetID();
 
+		const int32 VariableRecentIndex = Settings->FindRecentlyUsedMetaSoundsGraphAction(VariableID.ToString());
+		const bool bIsVariableFavorited = Settings->IsFavoriteMetaSoundsGraphAction(VariableID.ToString());
+
 		if (bGetAccessor)
 		{
 			// Do not add the action if adding an accessor node would cause a loop.
@@ -2715,7 +2817,23 @@ void UMetasoundEditorGraphSchema::GetVariableActions(FGraphActionMenuBuilder& Ac
 			{
 				FText ActionDisplayName = FText::Format(SchemaPrivate::VariableAccessorDisplayNameFormat, VariableDisplayName);
 				FText ActionTooltip = FText::Format(SchemaPrivate::VariableAccessorTooltipFormat, VariableDisplayName);
-				ActionMenuBuilder.AddAction(MakeShared<FMetasoundGraphSchemaAction_NewVariableAccessorNode>(GroupName, ActionDisplayName, VariableID, ActionTooltip));
+				TSharedPtr<FMetasoundGraphSchemaAction_NewVariableAccessorNode> NewNodeAction = MakeVariableAction<FMetasoundGraphSchemaAction_NewVariableAccessorNode>(
+					BaseGroupName, static_cast<int32>(EPrimaryContextGroup::Variables), ActionDisplayName, VariableID, ActionTooltip);
+				ActionMenuBuilder.AddAction(NewNodeAction);
+				
+				if (VariableRecentIndex != INDEX_NONE)
+				{
+					TSharedPtr<FMetasoundGraphSchemaAction_NewVariableAccessorNode> NewRecentAction = MakeVariableAction<FMetasoundGraphSchemaAction_NewVariableAccessorNode>(
+						RecentGroupName, static_cast<int32>(EPrimaryContextGroup::Recents) + VariableRecentIndex, ActionDisplayName, VariableID, ActionTooltip);
+					ActionMenuBuilder.AddAction(NewRecentAction);
+				}
+
+				if (bIsVariableFavorited)
+				{
+					TSharedPtr<FMetasoundGraphSchemaAction_NewVariableAccessorNode> NewFavAction = MakeVariableAction<FMetasoundGraphSchemaAction_NewVariableAccessorNode>(
+						FavoriteGroupName, static_cast<int32>(EPrimaryContextGroup::Favorites), ActionDisplayName, VariableID, ActionTooltip);
+					ActionMenuBuilder.AddAction(NewFavAction);
+				}
 			}
 		}
 
@@ -2723,7 +2841,23 @@ void UMetasoundEditorGraphSchema::GetVariableActions(FGraphActionMenuBuilder& Ac
 		{
 			FText ActionDisplayName = FText::Format(SchemaPrivate::VariableDeferredAccessorDisplayNameFormat, VariableDisplayName);
 			FText ActionTooltip = FText::Format(SchemaPrivate::VariableDeferredAccessorTooltipFormat, VariableDisplayName);
-			ActionMenuBuilder.AddAction(MakeShared<FMetasoundGraphSchemaAction_NewVariableDeferredAccessorNode>(GroupName, ActionDisplayName, VariableID, ActionTooltip));
+			TSharedPtr<FMetasoundGraphSchemaAction_NewVariableDeferredAccessorNode> NewNodeAction = MakeVariableAction<FMetasoundGraphSchemaAction_NewVariableDeferredAccessorNode>(
+				BaseGroupName, static_cast<int32>(EPrimaryContextGroup::Variables), ActionDisplayName, VariableID, ActionTooltip);
+			ActionMenuBuilder.AddAction(NewNodeAction);
+				
+			if (VariableRecentIndex != INDEX_NONE)
+			{
+				TSharedPtr<FMetasoundGraphSchemaAction_NewVariableDeferredAccessorNode> NewRecentAction = MakeVariableAction<FMetasoundGraphSchemaAction_NewVariableDeferredAccessorNode>(
+					RecentGroupName, static_cast<int32>(EPrimaryContextGroup::Recents) + VariableRecentIndex, ActionDisplayName, VariableID, ActionTooltip);
+				ActionMenuBuilder.AddAction(NewRecentAction);
+			}
+
+			if (bIsVariableFavorited)
+			{
+				TSharedPtr<FMetasoundGraphSchemaAction_NewVariableDeferredAccessorNode> NewFavAction = MakeVariableAction<FMetasoundGraphSchemaAction_NewVariableDeferredAccessorNode>(
+					FavoriteGroupName, static_cast<int32>(EPrimaryContextGroup::Favorites), ActionDisplayName, VariableID, ActionTooltip);
+				ActionMenuBuilder.AddAction(NewFavAction);
+			}
 		}
 
 		if (bGetMutator)
@@ -2738,7 +2872,23 @@ void UMetasoundEditorGraphSchema::GetVariableActions(FGraphActionMenuBuilder& Ac
 				{
 					FText ActionDisplayName = FText::Format(SchemaPrivate::VariableMutatorDisplayNameFormat, VariableDisplayName);
 					FText ActionTooltip = FText::Format(SchemaPrivate::VariableMutatorTooltipFormat, VariableDisplayName);
-					ActionMenuBuilder.AddAction(MakeShared<FMetasoundGraphSchemaAction_NewVariableMutatorNode>(GroupName, ActionDisplayName, VariableID, ActionTooltip));
+					TSharedPtr<FMetasoundGraphSchemaAction_NewVariableMutatorNode> NewNodeAction = MakeVariableAction<FMetasoundGraphSchemaAction_NewVariableMutatorNode>(
+						BaseGroupName, static_cast<int32>(EPrimaryContextGroup::Variables), ActionDisplayName, VariableID, ActionTooltip);
+					ActionMenuBuilder.AddAction(NewNodeAction);
+					
+					if (VariableRecentIndex != INDEX_NONE)
+					{
+						TSharedPtr<FMetasoundGraphSchemaAction_NewVariableMutatorNode> NewRecentAction = MakeVariableAction<FMetasoundGraphSchemaAction_NewVariableMutatorNode>(
+							RecentGroupName, static_cast<int32>(EPrimaryContextGroup::Recents) + VariableRecentIndex, ActionDisplayName, VariableID, ActionTooltip);
+						ActionMenuBuilder.AddAction(NewRecentAction);
+					}
+
+					if (bIsVariableFavorited)
+					{
+						TSharedPtr<FMetasoundGraphSchemaAction_NewVariableMutatorNode> NewFavAction = MakeVariableAction<FMetasoundGraphSchemaAction_NewVariableMutatorNode>(
+							FavoriteGroupName, static_cast<int32>(EPrimaryContextGroup::Favorites), ActionDisplayName, VariableID, ActionTooltip);
+						ActionMenuBuilder.AddAction(NewFavAction);
+					}
 				}
 			}
 		}
@@ -2838,4 +2988,3 @@ TSharedPtr<IAssetReferenceFilter> UMetasoundEditorGraphSchema::MakeAssetReferenc
 }
 
 #undef LOCTEXT_NAMESPACE
-
diff --git a/Engine/Plugins/Runtime/Metasound/Source/MetasoundEditor/Private/MetasoundEditorSettings.cpp b/Engine/Plugins/Runtime/Metasound/Source/MetasoundEditor/Private/MetasoundEditorSettings.cpp
--- a/Engine/Plugins/Runtime/Metasound/Source/MetasoundEditor/Private/MetasoundEditorSettings.cpp
+++ b/Engine/Plugins/Runtime/Metasound/Source/MetasoundEditor/Private/MetasoundEditorSettings.cpp
@@ -1,9 +1,6 @@
 // Copyright Epic Games, Inc. All Rights Reserved.
 #include "MetasoundEditorSettings.h"
 
-#include "AudioMaterialSlate/AudioMaterialSlateTypes.h"
-#include "AudioWidgetsStyle.h"
-#include "Framework/Notifications/NotificationManager.h"
 #include "Internationalization/Text.h"
 #include "MetasoundDocumentBuilderRegistry.h"
 #include "MetasoundFrontendDocument.h"
@@ -11,7 +8,6 @@
 #include "PlatformInfo.h"
 #include "Styling/AppStyle.h"
 #include "Styling/SlateWidgetStyleAsset.h"
-#include "Widgets/Notifications/SNotificationList.h"
 #include "UObject/UnrealType.h"
 
 
@@ -45,6 +41,50 @@ UMetasoundEditorSettings::UMetasoundEditorSettings(const FObjectInitializer& Obj
 	VariableNodeTitleColor = FLinearColor(0.211f, 0.513f, 0.035f);				// copper
 }
 
+void UMetasoundEditorSettings::AddRecentlyUsedMetaSoundsGraphAction(FString&& ActionKey)
+{
+	const int32 ExistingIndex = RecentMetaSoundsGraphActions.Find(ActionKey);
+	if (ExistingIndex != INDEX_NONE)
+	{
+		RecentMetaSoundsGraphActions.RemoveAt(ExistingIndex, EAllowShrinking::No);
+	}
+	else if (RecentMetaSoundsGraphActions.Num() >= MaxRecentEntries)
+	{
+		RecentMetaSoundsGraphActions.RemoveAt(0, RecentMetaSoundsGraphActions.Num() - MaxRecentEntries + 1, EAllowShrinking::No);
+	}
+	RecentMetaSoundsGraphActions.Emplace(ActionKey);
+
+	SaveConfig();
+}
+
+int32 UMetasoundEditorSettings::FindRecentlyUsedMetaSoundsGraphAction(const FString& ActionKey) const
+{
+	return RecentMetaSoundsGraphActions.Find(ActionKey);
+}
+
+void UMetasoundEditorSettings::AddFavoriteMetaSoundsGraphAction(FString&& ActionKey)
+{
+	bool bAlreadyFavorite = false;
+	FavoriteMetaSoundsGraphActions.Emplace(ActionKey, &bAlreadyFavorite);
+	if (!bAlreadyFavorite)
+	{
+		SaveConfig();
+	}
+}
+
+void UMetasoundEditorSettings::RemoveFavoriteMetaSoundsGraphAction(const FString& ActionKey)
+{
+	if (FavoriteMetaSoundsGraphActions.Remove(ActionKey) > 0)
+	{
+		SaveConfig();
+	}
+}
+
+bool UMetasoundEditorSettings::IsFavoriteMetaSoundsGraphAction(const FString& ActionKey) const
+{
+	return FavoriteMetaSoundsGraphActions.Contains(ActionKey);
+}
+
 #if WITH_EDITOR
 void UMetasoundEditorSettings::PostEditChangeProperty(FPropertyChangedEvent& PropertyChangedEvent)
 {
@@ -58,7 +98,12 @@ void UMetasoundEditorSettings::PostEditChangeProperty(FPropertyChangedEvent& Pro
 			Settings->OverrideTargetPageSettings(GetPageForAudition());
 		}
 	}
+	else if (PropertyName == GET_MEMBER_NAME_CHECKED(UMetasoundEditorSettings, MaxRecentEntries))
+	{
+		MaxRecentEntries = FMath::Min(MaxRecentEntries, MaxRecentEntriesEver);
+	}
 }
+
 void UMetasoundEditorSettings::SetAuditionPageMode(EAuditionPageMode InMode)
 {
 	AuditionPageMode = InMode;
diff --git a/Engine/Plugins/Runtime/Metasound/Source/MetasoundEditor/Private/MetasoundEditorSettings.h b/Engine/Plugins/Runtime/Metasound/Source/MetasoundEditor/Private/MetasoundEditorSettings.h
--- a/Engine/Plugins/Runtime/Metasound/Source/MetasoundEditor/Private/MetasoundEditorSettings.h
+++ b/Engine/Plugins/Runtime/Metasound/Source/MetasoundEditor/Private/MetasoundEditorSettings.h
@@ -304,6 +304,30 @@ public:
 	UPROPERTY(EditAnywhere, config, Category = "Widget Styling (Experimental)", meta = (AllowedClasses = "/Script/SlateCore.SlateWidgetStyleAsset", EditCondition = "bUseAudioMaterialWidgets", DisplayName = "Meter Style"))
 	FSoftObjectPath MeterStyleOverride;
 
+	/** MaxRecentEntries can go up to this number (limit must be known, for ordering Common actions in graph action menu above Recent ones) */
+	static constexpr int32 MaxRecentEntriesEver = 200;
+
+	/** The maximum number of Recent entries kept in the Metasound Editor action menu (ie., graph context menu). Older entries are overwritten by
+	  * newer entries, and reusing one put it back at the top.
+	  */
+	UPROPERTY(EditAnywhere, config, Category = General, meta = (ClampMin = 1))
+	int32 MaxRecentEntries = 20;
+
+private:
+	/** Higher index means more recent */
+	UPROPERTY(config)
+	TArray<FString> RecentMetaSoundsGraphActions;
+	
+	UPROPERTY(config)
+	TSet<FString> FavoriteMetaSoundsGraphActions; // TODO: Implement profiles, like for blueprint favorites
+
+public:
+	void AddRecentlyUsedMetaSoundsGraphAction(FString&& ActionKey);
+	int32 FindRecentlyUsedMetaSoundsGraphAction(const FString& ActionKey) const;
+	void AddFavoriteMetaSoundsGraphAction(FString&& ActionKey);
+	void RemoveFavoriteMetaSoundsGraphAction(const FString& ActionKey);
+	bool IsFavoriteMetaSoundsGraphAction(const FString& ActionKey) const;
+	
 	//UObject
 #if WITH_EDITOR
 	UE_API virtual void PostEditChangeProperty(FPropertyChangedEvent& PropertyChangedEvent) override;
diff --git a/Engine/Plugins/Runtime/Metasound/Source/MetasoundEditor/Private/SMetasoundActionMenu.cpp b/Engine/Plugins/Runtime/Metasound/Source/MetasoundEditor/Private/SMetasoundActionMenu.cpp
--- a/Engine/Plugins/Runtime/Metasound/Source/MetasoundEditor/Private/SMetasoundActionMenu.cpp
+++ b/Engine/Plugins/Runtime/Metasound/Source/MetasoundEditor/Private/SMetasoundActionMenu.cpp
@@ -4,7 +4,6 @@
 
 #include "Framework/Application/SlateApplication.h"
 #include "EdGraphSchema_K2.h"
-#include "Editor/EditorPerProjectUserSettings.h"
 #include "Styling/AppStyle.h"
 #include "IDocumentation.h"
 #include "MetasoundEditorGraphBuilder.h"
@@ -12,10 +11,8 @@
 #include "MetasoundEditorModule.h"
 #include "MetasoundEditorSettings.h"
 #include "MetasoundFrontendDocument.h"
-#include "MetasoundFrontendRegistries.h"
 #include "SGraphActionMenu.h"
 #include "SMetasoundPalette.h"
-#include "SSubobjectEditor.h"
 #include "Styling/SlateStyleRegistry.h"
 #include "Widgets/Images/SImage.h"
 #include "Widgets/Input/SCheckBox.h"
@@ -33,6 +30,106 @@ namespace Metasound
 {
 	namespace Editor
 	{
+		class SMetasoundActionFavoriteToggle : public SCompoundWidget
+		{
+			SLATE_BEGIN_ARGS(SMetasoundActionFavoriteToggle) {}
+			SLATE_END_ARGS()
+
+		public:
+			void Construct(const FArguments& InArgs, const FCustomExpanderData& CustomExpanderData)
+			{
+				Container = CustomExpanderData.WidgetContainer;
+				
+				if (CustomExpanderData.RowAction.IsValid() && CustomExpanderData.RowAction->GetTypeId() == SchemaUtils::MetasoundGraphSchemaActionType)
+				{
+					ActionPtr = StaticCastSharedPtr<FMetasoundGraphSchemaAction>(CustomExpanderData.RowAction);
+				}
+
+				ChildSlot
+				[
+					SNew(SHorizontalBox)
+					+SHorizontalBox::Slot()
+						.VAlign(VAlign_Fill)
+						.HAlign(HAlign_Center)
+						.FillWidth(1.0)
+					[
+						SNew(SCheckBox)
+							.Visibility(this, &SMetasoundActionFavoriteToggle::IsVisible)
+							.ToolTipText(this, &SMetasoundActionFavoriteToggle::GetToolTipText)
+							.IsChecked(this, &SMetasoundActionFavoriteToggle::GetFavoritedState)
+							.OnCheckStateChanged(this, &SMetasoundActionFavoriteToggle::OnFavoriteToggled)
+							.Style(FAppStyle::Get(), "Kismet.Palette.FavoriteToggleStyle")
+					]
+				];
+			}
+
+		private:
+			/** The action that the owning palette entry represents */
+			TWeakPtr<FMetasoundGraphSchemaAction> ActionPtr;
+
+			/** The widget that this widget is nested inside */
+			TSharedPtr<SPanel> Container;
+			
+			EVisibility IsVisible() const
+			{
+				EVisibility CurrentVisibility = EVisibility::Hidden;
+				if (TSharedPtr<FMetasoundGraphSchemaAction> Action = ActionPtr.Pin())
+				{
+					if (Action->Grouping != static_cast<int32>(EPrimaryContextGroup::Common))
+					{
+						if (Container->IsHovered() || GetDefault<UMetasoundEditorSettings>()->IsFavoriteMetaSoundsGraphAction(Action->GetUniqueStableFavoriteKey()))
+						{
+							CurrentVisibility = EVisibility::Visible;
+						}
+					}
+				}
+
+				return CurrentVisibility;
+			}
+
+			/**
+			 * Retrieves tooltip that describes the current favorited state of the 
+			 * associated action.
+			 * 
+			 * @return Text describing what this toggle will do when you click on it.
+			 */
+			FText GetToolTipText() const
+			{
+				if (GetFavoritedState() == ECheckBoxState::Checked)
+				{
+					return LOCTEXT("Unfavorite", "Click to remove this item from your favorites.");
+				}
+				return LOCTEXT("Favorite", "Click to add this item to your favorites.");
+			}
+
+			ECheckBoxState GetFavoritedState() const
+			{
+				ECheckBoxState FavoriteState = ECheckBoxState::Unchecked;
+				
+				if (TSharedPtr<FMetasoundGraphSchemaAction> Action = ActionPtr.Pin())
+				{
+					FavoriteState = GetDefault<UMetasoundEditorSettings>()->IsFavoriteMetaSoundsGraphAction(Action->GetUniqueStableFavoriteKey()) ? ECheckBoxState::Checked : ECheckBoxState::Unchecked;
+				}
+				
+				return FavoriteState;
+			}
+
+			void OnFavoriteToggled(ECheckBoxState InNewState)
+			{
+				if (TSharedPtr<FMetasoundGraphSchemaAction> Action = ActionPtr.Pin())
+				{
+					if (InNewState == ECheckBoxState::Checked)
+					{
+						GetMutableDefault<UMetasoundEditorSettings>()->AddFavoriteMetaSoundsGraphAction(Action->GetUniqueStableFavoriteKey());
+					}
+					else
+					{
+						GetMutableDefault<UMetasoundEditorSettings>()->RemoveFavoriteMetaSoundsGraphAction(Action->GetUniqueStableFavoriteKey());
+					}
+				}
+			}
+		};
+		
 		void SMetasoundActionMenuExpanderArrow::Construct(const FArguments& InArgs, const FCustomExpanderData& ActionMenuData)
 		{
 			OwnerRowPtr = ActionMenuData.TableRow;
@@ -44,9 +141,7 @@ namespace Metasound
 				ChildSlot
 				.Padding(TAttribute<FMargin>(this, &SMetasoundActionMenuExpanderArrow::GetCustomIndentPadding))
 				[
-					SNullWidget::NullWidget
-					// TODO: Add favorites support
-					// SNew(SMetasoundActionFavoriteToggle, ActionMenuData)
+					SNew(SMetasoundActionFavoriteToggle, ActionMenuData)
 				];
 			}
 			else
@@ -89,7 +184,6 @@ namespace Metasound
 				{
 					FDataTypeRegistryInfo RegistryInfo;
 
-					const IMetasoundEditorModule& EditorModule = FModuleManager::GetModuleChecked<IMetasoundEditorModule>("MetaSoundEditor");
 					if (Pin->Direction == EGPD_Input)
 					{
 						FConstInputHandle InputHandle = FGraphBuilder::GetConstInputHandleFromPin(Pin);
@@ -122,6 +216,22 @@ namespace Metasound
 				}
 			}
 
+			FText MenuTitle;
+			if (DraggedFromPins.IsEmpty())
+			{
+				MenuTitle = LOCTEXT("ContextText", "All MetaSound Node Classes");
+			}
+			else if (DraggedFromPins[0]->Direction == EGPD_Input)
+			{
+				FConstInputHandle InputHandle = FGraphBuilder::GetConstInputHandleFromPin(DraggedFromPins[0]);
+				MenuTitle = FText::Format(LOCTEXT("ContextTypeFilteredText_Output", "Classes with output of type '{0}'"), FText::FromName(InputHandle->GetDataType()));
+			}
+			else
+			{
+				FConstOutputHandle OutputHandle = FGraphBuilder::GetConstOutputHandleFromPin(DraggedFromPins[0]);
+				MenuTitle = FText::Format(LOCTEXT("ContextTypeFilteredText_Input", "Classes with input of type '{0}'"), FText::FromName(OutputHandle->GetDataType()));
+			}
+
 			SBorder::Construct(SBorder::FArguments()
 				.BorderImage(FAppStyle::GetBrush("Menu.Background"))
 				.Padding(5)
@@ -143,10 +253,7 @@ namespace Metasound
 							[
 								SNew(SImage)
 								.ColorAndOpacity(TypeColor)
-								.Visibility_Lambda([this]()
-								{
-									return DraggedFromPins.IsEmpty() ? EVisibility::Hidden : EVisibility::Visible;
-								})
+								.Visibility(DraggedFromPins.IsEmpty() ? EVisibility::Hidden : EVisibility::Visible)
 								.Image(PinBrush)
 							]
 							+SHorizontalBox::Slot()
@@ -154,26 +261,7 @@ namespace Metasound
 							.VAlign(VAlign_Center)
 							[
 								SNew(STextBlock)
-								.Text_Lambda([this]()
-								{
-									if (DraggedFromPins.IsEmpty())
-									{
-										return LOCTEXT("ContextText", "All MetaSound Node Classes");
-									}
-
-									UEdGraphPin* Pin = DraggedFromPins[0];
-
-									if (DraggedFromPins[0]->Direction == EGPD_Input)
-									{
-										FConstInputHandle InputHandle = FGraphBuilder::GetConstInputHandleFromPin(Pin);
-										return FText::Format(LOCTEXT("ContextTypeFilteredText_Output", "Classes with output of type '{0}'"), FText::FromName(InputHandle->GetDataType()));
-									}
-									else
-									{
-										FConstOutputHandle OutputHandle = FGraphBuilder::GetConstOutputHandleFromPin(Pin);
-										return FText::Format(LOCTEXT("ContextTypeFilteredText_Input", "Classes with input of type '{0}'"), FText::FromName(OutputHandle->GetDataType()));
-									}
-								})
+								.Text(MenuTitle)
 								// TODO: Move to Metasound Style
 								.Font(FAppStyle::GetFontStyle("BlueprintEditor.ActionMenu.ContextDescriptionFont"))
 								.ToolTip(IDocumentation::Get()->CreateToolTip(
@@ -288,6 +376,7 @@ namespace Metasound
 				return;
 			}
 
+			UMetasoundEditorSettings* Settings = GetMutableDefault<UMetasoundEditorSettings>();
 			for (const TSharedPtr<FEdGraphSchemaAction>& Action : SelectedAction)
 			{
 				if (Action.IsValid() && Graph)
@@ -296,11 +385,12 @@ namespace Metasound
 					{
 						FSlateApplication::Get().DismissAllMenus();
 						bActionExecuted = true;
+						Settings->AddRecentlyUsedMetaSoundsGraphAction(StaticCastSharedPtr<FMetasoundGraphSchemaAction>(Action)->GetUniqueStableFavoriteKey());
 					}
 
-					if (UEdGraphNode* ResultNode = Action->PerformAction(Graph, DraggedFromPins, NewNodePosition))
+					if (Action->PerformAction(Graph, DraggedFromPins, NewNodePosition) != nullptr)
 					{
-						NewNodePosition += Metasound::Frontend::DisplayStyle::NodeLayout::DefaultOffsetX;
+						NewNodePosition += Frontend::DisplayStyle::NodeLayout::DefaultOffsetX;
 					}
 				}
 			}
diff --git a/Engine/Plugins/Runtime/Metasound/Source/MetasoundEditor/Public/MetasoundEditorGraphSchema.h b/Engine/Plugins/Runtime/Metasound/Source/MetasoundEditor/Public/MetasoundEditorGraphSchema.h
--- a/Engine/Plugins/Runtime/Metasound/Source/MetasoundEditor/Public/MetasoundEditorGraphSchema.h
+++ b/Engine/Plugins/Runtime/Metasound/Source/MetasoundEditor/Public/MetasoundEditorGraphSchema.h
@@ -35,8 +35,11 @@ namespace Metasound::Editor
 		Outputs,
 		Inputs,
 		Variables,
+		Favorites,
 
-		Common
+		// This MUST be the last one before Common, because the space between them is reserved for Recent actions ordered by most recent.
+		Recents, 
+		Common = 255
 	};
 
 	const FText& GetContextGroupDisplayName(EPrimaryContextGroup InContextGroup);
@@ -89,6 +92,8 @@ namespace Metasound::Editor
 		UEdGraphNode* PromoteToVariable(UEdGraph* ParentGraph, UEdGraphPin* FromPin, const UE::Slate::FDeprecateVector2DParameter& InLocation, bool bSelectNewNode);
 		UEdGraphNode* PromoteToDeferredVariable(UEdGraph* ParentGraph, UEdGraphPin* FromPin, const UE::Slate::FDeprecateVector2DParameter& InLocation, bool bSelectNewNode);
 		UEdGraphNode* PromoteToMutatorVariable(UEdGraph* ParentGraph, UEdGraphPin* FromPin, const UE::Slate::FDeprecateVector2DParameter& InLocation, bool bSelectNewNode);
+
+		const FName MetasoundGraphSchemaActionType("FMetasoundGraphSchemaAction");
 	} // namespace SchemaUtils
 } // namespace Metasound::Editor
 
@@ -96,7 +101,10 @@ namespace Metasound::Editor
 USTRUCT()
 struct FMetasoundGraphSchemaAction : public FEdGraphSchemaAction
 {
-	GENERATED_USTRUCT_BODY()
+	GENERATED_BODY()
+
+	static FName StaticGetTypeId() { return Metasound::Editor::SchemaUtils::MetasoundGraphSchemaActionType; }
+	virtual FName GetTypeId() const override { return StaticGetTypeId(); }
 
 	FMetasoundGraphSchemaAction()
 		: FEdGraphSchemaAction()
@@ -125,13 +133,20 @@ struct FMetasoundGraphSchemaAction : public FEdGraphSchemaAction
 		static const FLinearColor DefaultColor = FLinearColor::Black;
 		return DefaultColor;
 	}
+
+	// Used to store & match Recent and Favorite MetaSounds graph actions in the MetaSounds editor
+	virtual FString GetUniqueStableFavoriteKey() const
+	{
+		// Unlikely: If the user changes the editor locale, favorites and recents could technically be invalidated if their display names are localized.
+		return GetMenuDescription().ToString();
+	}
 };
 
 /** This is used to combine functionality for nodes that can have multiple outputs and should never be directly instantiated. */
 USTRUCT()
 struct FMetasoundGraphSchemaAction_NodeWithMultipleOutputs : public FMetasoundGraphSchemaAction
 {
-	GENERATED_USTRUCT_BODY();
+	GENERATED_BODY();
 
 	FMetasoundGraphSchemaAction_NodeWithMultipleOutputs()
 		: FMetasoundGraphSchemaAction()
@@ -154,7 +169,7 @@ struct FMetasoundGraphSchemaAction_NodeWithMultipleOutputs : public FMetasoundGr
 USTRUCT()
 struct FMetasoundGraphSchemaAction_NewInput : public FMetasoundGraphSchemaAction_NodeWithMultipleOutputs
 {
-	GENERATED_USTRUCT_BODY();
+	GENERATED_BODY();
 
 	UPROPERTY()
 	FGuid NodeID;
@@ -171,6 +186,11 @@ struct FMetasoundGraphSchemaAction_NewInput : public FMetasoundGraphSchemaAction
 	UE_API virtual const FSlateBrush* GetIconBrush() const override;
 
 	UE_API virtual const FLinearColor& GetIconColor() const override;
+
+	UE_API virtual FString GetUniqueStableFavoriteKey() const override
+	{
+		return NodeID.ToString();
+	}
 	//~ End FMetasoundGraphSchemaAction Interface
 
 	//~ Begin FEdGraphSchemaAction Interface
@@ -183,7 +203,7 @@ struct FMetasoundGraphSchemaAction_NewInput : public FMetasoundGraphSchemaAction
 USTRUCT()
 struct FMetasoundGraphSchemaAction_PromoteToInput : public FMetasoundGraphSchemaAction_NodeWithMultipleOutputs
 {
-	GENERATED_USTRUCT_BODY();
+	GENERATED_BODY();
 
 	UE_API FMetasoundGraphSchemaAction_PromoteToInput();
 
@@ -199,7 +219,7 @@ struct FMetasoundGraphSchemaAction_PromoteToInput : public FMetasoundGraphSchema
 USTRUCT()
 struct FMetasoundGraphSchemaAction_NewOutput : public FMetasoundGraphSchemaAction
 {
-	GENERATED_USTRUCT_BODY();
+	GENERATED_BODY();
 
 	UPROPERTY()
 	FGuid NodeID;
@@ -216,6 +236,11 @@ struct FMetasoundGraphSchemaAction_NewOutput : public FMetasoundGraphSchemaActio
 	UE_API virtual const FSlateBrush* GetIconBrush() const override;
 
 	UE_API virtual const FLinearColor& GetIconColor() const override;
+
+	UE_API virtual FString GetUniqueStableFavoriteKey() const override
+	{
+		return NodeID.ToString();
+	}
 	//~ End FMetasoundGraphSchemaAction Interface
 
 	//~ Begin FEdGraphSchemaAction Interface
@@ -228,7 +253,7 @@ struct FMetasoundGraphSchemaAction_NewOutput : public FMetasoundGraphSchemaActio
 USTRUCT()
 struct FMetasoundGraphSchemaAction_PromoteToOutput : public FMetasoundGraphSchemaAction
 {
-	GENERATED_USTRUCT_BODY();
+	GENERATED_BODY();
 
 	UE_API FMetasoundGraphSchemaAction_PromoteToOutput();
 
@@ -244,7 +269,7 @@ struct FMetasoundGraphSchemaAction_PromoteToOutput : public FMetasoundGraphSchem
 USTRUCT()
 struct FMetasoundGraphSchemaAction_NewVariableNode : public FMetasoundGraphSchemaAction_NodeWithMultipleOutputs
 {
-	GENERATED_USTRUCT_BODY();
+	GENERATED_BODY();
 
 	UPROPERTY()
 	FGuid VariableID;
@@ -266,6 +291,11 @@ struct FMetasoundGraphSchemaAction_NewVariableNode : public FMetasoundGraphSchem
 	UE_API virtual const FSlateBrush* GetIconBrush() const override;
 
 	UE_API virtual const FLinearColor& GetIconColor() const override;
+
+	UE_API virtual FString GetUniqueStableFavoriteKey() const override
+	{
+		return VariableID.ToString();
+	}
 	//~ End FMetasoundGraphSchemaAction Interface
 	
 protected:
@@ -285,7 +315,7 @@ protected:
 USTRUCT()
 struct FMetasoundGraphSchemaAction_NewVariableAccessorNode : public FMetasoundGraphSchemaAction_NewVariableNode
 {
-	GENERATED_USTRUCT_BODY();
+	GENERATED_BODY();
 
 	FMetasoundGraphSchemaAction_NewVariableAccessorNode()
 		: FMetasoundGraphSchemaAction_NewVariableNode()
@@ -304,7 +334,7 @@ protected:
 USTRUCT()
 struct FMetasoundGraphSchemaAction_NewVariableDeferredAccessorNode : public FMetasoundGraphSchemaAction_NewVariableNode
 {
-	GENERATED_USTRUCT_BODY();
+	GENERATED_BODY();
 
 	FMetasoundGraphSchemaAction_NewVariableDeferredAccessorNode()
 		: FMetasoundGraphSchemaAction_NewVariableNode()
@@ -322,7 +352,7 @@ protected:
 USTRUCT()
 struct FMetasoundGraphSchemaAction_NewVariableMutatorNode : public FMetasoundGraphSchemaAction_NewVariableNode
 {
-	GENERATED_USTRUCT_BODY();
+	GENERATED_BODY();
 
 	FMetasoundGraphSchemaAction_NewVariableMutatorNode()
 		: FMetasoundGraphSchemaAction_NewVariableNode()
@@ -340,7 +370,7 @@ protected:
 USTRUCT()
 struct FMetasoundGraphSchemaAction_PromoteToVariable_AccessorNode : public FMetasoundGraphSchemaAction_NodeWithMultipleOutputs
 {
-	GENERATED_USTRUCT_BODY();
+	GENERATED_BODY();
 
 	UE_API FMetasoundGraphSchemaAction_PromoteToVariable_AccessorNode();
 
@@ -356,7 +386,7 @@ struct FMetasoundGraphSchemaAction_PromoteToVariable_AccessorNode : public FMeta
 USTRUCT()
 struct FMetasoundGraphSchemaAction_PromoteToVariable_MutatorNode : public FMetasoundGraphSchemaAction
 {
-	GENERATED_USTRUCT_BODY();
+	GENERATED_BODY();
 
 	UE_API FMetasoundGraphSchemaAction_PromoteToVariable_MutatorNode();
 
@@ -372,7 +402,7 @@ struct FMetasoundGraphSchemaAction_PromoteToVariable_MutatorNode : public FMetas
 USTRUCT()
 struct FMetasoundGraphSchemaAction_PromoteToVariable_DeferredAccessorNode : public FMetasoundGraphSchemaAction_NodeWithMultipleOutputs
 {
-	GENERATED_USTRUCT_BODY();
+	GENERATED_BODY();
 
 	UE_API FMetasoundGraphSchemaAction_PromoteToVariable_DeferredAccessorNode();
 
@@ -389,7 +419,7 @@ struct FMetasoundGraphSchemaAction_PromoteToVariable_DeferredAccessorNode : publ
 USTRUCT()
 struct FMetasoundGraphSchemaAction_NewNode : public FMetasoundGraphSchemaAction_NodeWithMultipleOutputs
 {
-	GENERATED_USTRUCT_BODY();
+	GENERATED_BODY();
 
 	TSharedPtr<Metasound::Editor::ISchemaQueryResult> QueryResult;
 
@@ -407,6 +437,8 @@ struct FMetasoundGraphSchemaAction_NewNode : public FMetasoundGraphSchemaAction_
 	UE_API virtual const FSlateBrush* GetIconBrush() const override;
 
 	UE_API virtual const FLinearColor& GetIconColor() const override;
+
+	UE_API virtual FString GetUniqueStableFavoriteKey() const override;
 	//~ End FMetasoundGraphSchemaAction Interface
 
 	//~ Begin FEdGraphSchemaAction Interface
@@ -419,7 +451,7 @@ struct FMetasoundGraphSchemaAction_NewNode : public FMetasoundGraphSchemaAction_
 USTRUCT()
 struct FMetasoundGraphSchemaAction_NewFromSelected : public FMetasoundGraphSchemaAction_NewNode
 {
-	GENERATED_USTRUCT_BODY();
+	GENERATED_BODY();
 
 	FMetasoundGraphSchemaAction_NewFromSelected() 
 		: FMetasoundGraphSchemaAction_NewNode()
@@ -441,7 +473,7 @@ struct FMetasoundGraphSchemaAction_NewFromSelected : public FMetasoundGraphSchem
 USTRUCT()
 struct FMetasoundGraphSchemaAction_NewAudioAnalyzer : public FMetasoundGraphSchemaAction
 {
-	GENERATED_USTRUCT_BODY();
+	GENERATED_BODY();
 
 	UE_API FMetasoundGraphSchemaAction_NewAudioAnalyzer();
 
@@ -459,7 +491,7 @@ struct FMetasoundGraphSchemaAction_NewAudioAnalyzer : public FMetasoundGraphSche
 USTRUCT()
 struct FMetasoundGraphSchemaAction_NewReroute : public FMetasoundGraphSchemaAction
 {
-	GENERATED_USTRUCT_BODY();
+	GENERATED_BODY();
 
 	FMetasoundGraphSchemaAction_NewReroute() = default;
 	UE_API FMetasoundGraphSchemaAction_NewReroute(const FLinearColor* InIconColor, bool bInShouldTransact = true);
@@ -486,7 +518,7 @@ private:
 USTRUCT()
 struct FMetasoundGraphSchemaAction_NewComment : public FMetasoundGraphSchemaAction
 {
-	GENERATED_USTRUCT_BODY();
+	GENERATED_BODY();
 
 	FMetasoundGraphSchemaAction_NewComment()
 		: FMetasoundGraphSchemaAction()
@@ -512,7 +544,7 @@ struct FMetasoundGraphSchemaAction_NewComment : public FMetasoundGraphSchemaActi
 USTRUCT()
 struct FMetasoundGraphSchemaAction_Paste : public FMetasoundGraphSchemaAction
 {
-	GENERATED_USTRUCT_BODY();
+	GENERATED_BODY();
 
 	FMetasoundGraphSchemaAction_Paste() 
 		: FMetasoundGraphSchemaAction()
@@ -543,6 +575,7 @@ public:
 	virtual void GetGraphContextActions(FGraphContextMenuBuilder& ContextMenuBuilder) const override;
 	virtual void GetContextMenuActions(UToolMenu* Menu, UGraphNodeContextMenuContext* Context) const override;
 	virtual FText GetPinDisplayName(const UEdGraphPin* Pin) const override;
+	virtual float GetActionFilteredWeight(const FEdGraphSchemaAction& InCurrentAction, const TArray<FString>& InFilterTerms, const TArray<FString>& InSanitizedFilterTerms, const TArray<UEdGraphPin*>& DraggedFromPins) const override;
 	virtual void CreateDefaultNodesForGraph(UEdGraph& Graph) const override;
 	virtual const FPinConnectionResponse CanCreateConnection(const UEdGraphPin* A, const UEdGraphPin* B) const override;
 	virtual bool TryCreateConnection(UEdGraphPin* A, UEdGraphPin* B) const override;
diff --git a/Engine/Plugins/Runtime/Metasound/Source/MetasoundFrontend/Private/MetasoundFrontendNodesCategories.cpp b/Engine/Plugins/Runtime/Metasound/Source/MetasoundFrontend/Private/MetasoundFrontendNodesCategories.cpp
--- a/Engine/Plugins/Runtime/Metasound/Source/MetasoundFrontend/Private/MetasoundFrontendNodesCategories.cpp
+++ b/Engine/Plugins/Runtime/Metasound/Source/MetasoundFrontend/Private/MetasoundFrontendNodesCategories.cpp
@@ -28,6 +28,12 @@ namespace Metasound
 
 		// Base category for all variables 
 		const FText Variables = LOCTEXT("Metasound_VariablesCategory", "Variables");
+
+		// Base category for all favorites 
+		const FText Favorites = LOCTEXT("Metasound_FavoritesCategory", "Favorites");
+
+		// Base category for all recent actions 
+		const FText Recents = LOCTEXT("Metasound_RecentsCategory", "Recently Used");
 	}
 }
 
diff --git a/Engine/Plugins/Runtime/Metasound/Source/MetasoundFrontend/Public/MetasoundFrontendNodesCategories.h b/Engine/Plugins/Runtime/Metasound/Source/MetasoundFrontend/Public/MetasoundFrontendNodesCategories.h
--- a/Engine/Plugins/Runtime/Metasound/Source/MetasoundFrontend/Public/MetasoundFrontendNodesCategories.h
+++ b/Engine/Plugins/Runtime/Metasound/Source/MetasoundFrontend/Public/MetasoundFrontendNodesCategories.h
@@ -15,5 +15,7 @@ namespace Metasound
 		extern const FText METASOUNDFRONTEND_API Inputs;
 		extern const FText METASOUNDFRONTEND_API Outputs;
 		extern const FText METASOUNDFRONTEND_API Variables;
+		extern const FText METASOUNDFRONTEND_API Favorites;
+		extern const FText METASOUNDFRONTEND_API Recents;
 	} // namespace NodeCategories
 } // namespace Metasound
