Fix for Node Config Ensures and Ordering Schenangians

[RN] Fix for node config reording and pin synchronization

diff --git a/Engine/Plugins/Runtime/Metasound/Source/MetasoundFrontend/Private/MetasoundFrontendDocument.cpp b/Engine/Plugins/Runtime/Metasound/Source/MetasoundFrontend/Private/MetasoundFrontendDocument.cpp
--- a/Engine/Plugins/Runtime/Metasound/Source/MetasoundFrontend/Private/MetasoundFrontendDocument.cpp
+++ b/Engine/Plugins/Runtime/Metasound/Source/MetasoundFrontend/Private/MetasoundFrontendDocument.cpp
@@ -525,25 +525,25 @@ bool FMetasoundFrontendNodeInterface::Update(const FMetasoundFrontendClassInterf
 bool FMetasoundFrontendNodeInterface::Update(const FMetasoundFrontendClassInterface& InClassInterface, TFunctionRef<void(const FMetasoundFrontendVertex&)> OnPreRemoveInput, TFunctionRef<void(const FMetasoundFrontendVertex&)> OnPreRemoveOutput)
 {
 	bool bInterfaceUpdated = false;
-	TArray<const FMetasoundFrontendVertex*, TInlineAllocator<32>> UnmatchedVertices;
+	TArray<FMetasoundFrontendVertex> UnmatchedVertices;
 
 	struct FFindMatchingVertex
 	{
 		const FMetasoundFrontendClassVertex& ClassVertex;
-		bool operator()(const FMetasoundFrontendVertex* InNodeVertex)
+		bool operator()(const FMetasoundFrontendVertex& InNodeVertex)
 		{
-			return (ClassVertex.Name == InNodeVertex->Name) && (ClassVertex.TypeName == InNodeVertex->TypeName);
+			return (ClassVertex.Name == InNodeVertex.Name) && (ClassVertex.TypeName == InNodeVertex.TypeName);
 		}
 	};
 
 	// Update node inputs
-	Algo::Transform(Inputs, UnmatchedVertices, [](const FMetasoundFrontendVertex& Vertex) { return &Vertex; });
+	Algo::Transform(Inputs, UnmatchedVertices, [](const FMetasoundFrontendVertex& Vertex) { return Vertex; });
 	for (const FMetasoundFrontendClassInput& ClassInput : InClassInterface.Inputs)
 	{
 		if (const int32 Index = UnmatchedVertices.IndexOfByPredicate(FFindMatchingVertex{ClassInput}); Index != INDEX_NONE)
 		{
 			// Update the node vertex with anything new from the class vertex
-			UnmatchedVertices.RemoveAtSwap(Index, 1, EAllowShrinking::No);
+			UnmatchedVertices.RemoveAt(Index, 1, EAllowShrinking::No);
 		}
 		else
 		{
@@ -554,23 +554,27 @@ bool FMetasoundFrontendNodeInterface::Update(const FMetasoundFrontendClassInterf
 	}
 
 	// Remove any inputs that did not exist in the class interface
-	for (const FMetasoundFrontendVertex* UnmatchedInput : UnmatchedVertices)
+	for (const FMetasoundFrontendVertex& UnmatchedInput : UnmatchedVertices)
 	{
 		// Allow outside systems to react before removing the unmatched node inputs
-		OnPreRemoveInput(*UnmatchedInput);
-		Inputs.RemoveSingleSwap(*UnmatchedInput);
+		OnPreRemoveInput(UnmatchedInput);
+		if (const int32 Index = Inputs.IndexOfByPredicate([UnmatchedInput](const FMetasoundFrontendVertex& i) { return i.Name == UnmatchedInput.Name && i.TypeName == UnmatchedInput.TypeName; } ); Index != INDEX_NONE)
+		{
+			// Update the node vertex with anything new from the class vertex
+			Inputs.RemoveAt(Index, 1, EAllowShrinking::No);
+		}	
 		bInterfaceUpdated |= true;
 	}
 
 	// Update node outputs
 	UnmatchedVertices.Reset();
-	Algo::Transform(Outputs, UnmatchedVertices, [](const FMetasoundFrontendVertex& Vertex) { return &Vertex; });
+	Algo::Transform(Outputs, UnmatchedVertices, [](const FMetasoundFrontendVertex& Vertex) { return Vertex; });
 	for (const FMetasoundFrontendClassOutput& ClassOutput : InClassInterface.Outputs)
 	{
 		if (const int32 Index = UnmatchedVertices.IndexOfByPredicate(FFindMatchingVertex{ClassOutput}); Index != INDEX_NONE)
 		{
 			// Update the node vertex with anything new from the class vertex
-			UnmatchedVertices.RemoveAtSwap(Index, 1, EAllowShrinking::No);
+			UnmatchedVertices.RemoveAt(Index, 1, EAllowShrinking::No);
 		}
 		else
 		{
@@ -580,12 +584,17 @@ bool FMetasoundFrontendNodeInterface::Update(const FMetasoundFrontendClassInterf
 		}
 	}
 
-	// Remove any inputs that did not exist in the class interface
-	for (const FMetasoundFrontendVertex* UnmatchedOutput : UnmatchedVertices)
+	// Remove any outputs that did not exist in the class interface
+	for (const FMetasoundFrontendVertex& UnmatchedOutput : UnmatchedVertices)
 	{
-		// Allow outside systems to react before removing the unmatched node inputs
-		OnPreRemoveOutput(*UnmatchedOutput);
-		Outputs.RemoveSingleSwap(*UnmatchedOutput);
+		// Allow outside systems to react before removing the unmatched node outputs
+		OnPreRemoveOutput(UnmatchedOutput);
+
+		if (const int32 Index = Outputs.IndexOfByPredicate([&UnmatchedOutput](const FMetasoundFrontendVertex& i) { return i.Name == UnmatchedOutput.Name && i.TypeName == UnmatchedOutput.TypeName; } ); Index != INDEX_NONE)
+		{
+			// Update the node vertex with anything new from the class vertex
+			Outputs.RemoveAt(Index, 1, EAllowShrinking::No);
+		}	
 		bInterfaceUpdated |= true;
 	}
 	
