MetasoundMixerNode: Use TOptional AudioBuffer Inputs so we don't create nor render zero input audio buffers if no connection is provided

#tests Metasound Automation Tests

diff --git a/Engine/Plugins/Runtime/Metasound/Source/MetasoundStandardNodes/Private/MetasoundMixerNode.cpp b/Engine/Plugins/Runtime/Metasound/Source/MetasoundStandardNodes/Private/MetasoundMixerNode.cpp
--- a/Engine/Plugins/Runtime/Metasound/Source/MetasoundStandardNodes/Private/MetasoundMixerNode.cpp
+++ b/Engine/Plugins/Runtime/Metasound/Source/MetasoundStandardNodes/Private/MetasoundMixerNode.cpp
@@ -24,17 +24,17 @@ namespace Metasound
 	{
 	public:
 		// ctor
-		TAudioMixerNodeOperator(const FBuildOperatorParams& InParams, const TArray<FAudioBufferReadRef>&& InInputBuffers, const TArray<FFloatReadRef>&& InGainValues)
+		TAudioMixerNodeOperator(const FBuildOperatorParams& InParams, const TArray<TOptional<FAudioBufferReadRef>>& InInputBuffers, const TArray<FFloatReadRef>& InGainValues)
 			: Gains(InGainValues)
 			, Inputs (InInputBuffers)
 		{
 			// create write refs
-			for (uint32 i = 0; i < NumChannels; ++i)
+			for (uint32 Index = 0; Index < NumChannels; ++Index)
 			{
 				Outputs.Add(FAudioBufferWriteRef::CreateNew(InParams.OperatorSettings));
 			}
 
-			// init previous gains to current values
+			// Init previous gains to current values
 			PrevGains.Reset();
 			PrevGains.AddUninitialized(NumInputs);
 
@@ -149,17 +149,25 @@ namespace Metasound
 		{
 			const FInputVertexInterfaceData& InputData = InParams.InputData;
 			
-			TArray<FAudioBufferReadRef> InputBuffers;
+			TArray<TOptional<FAudioBufferReadRef>> InputBuffers;
 			TArray<FFloatReadRef> InputGains;
 
-			for (uint32 i = 0; i < NumInputs; ++i)
+			for (uint32 InputIndex = 0; InputIndex < NumInputs; ++InputIndex)
 			{
-				for (uint32 Chan = 0; Chan < NumChannels; ++Chan)
+				for (uint32 ChannelIndex = 0; ChannelIndex < NumChannels; ++ChannelIndex)
 				{
-					InputBuffers.Add(InputData.GetOrCreateDefaultDataReadReference<FAudioBuffer>(GetAudioInputName(i, Chan), InParams.OperatorSettings));
+					const FVertexName& AudioInputName = GetAudioInputName(InputIndex, ChannelIndex);
+					if (InputData.IsVertexBound(AudioInputName))
+					{
+						InputBuffers.Add(InputData.GetDataReadReference<FAudioBuffer>(AudioInputName));
+					}
+					else
+					{
+						InputBuffers.AddDefaulted();
+					}
 				}
 
-				InputGains.Add(InputData.GetOrCreateDefaultDataReadReference<float>(GainInputNames[i], InParams.OperatorSettings));
+				InputGains.Add(InputData.GetOrCreateDefaultDataReadReference<float>(GainInputNames[InputIndex], InParams.OperatorSettings));
 			}
 
 			return MakeUnique<TAudioMixerNodeOperator<NumInputs, NumChannels>>(InParams, MoveTemp(InputBuffers), MoveTemp(InputGains));
@@ -167,14 +175,23 @@ namespace Metasound
 
 		virtual void BindInputs(FInputVertexInterfaceData& InOutVertexData) override
 		{
-			for (uint32 i = 0; i < NumInputs; ++i)
+			for (uint32 InputIndex = 0; InputIndex < NumInputs; ++InputIndex)
 			{
-				for (uint32 Chan = 0; Chan < NumChannels; ++Chan)
+				for (uint32 ChannelIndex = 0; ChannelIndex < NumChannels; ++ChannelIndex)
 				{
-					InOutVertexData.BindReadVertex(GetAudioInputName(i, Chan), Inputs[i * NumChannels + Chan]);
+					const FVertexName& AudioInputName = GetAudioInputName(InputIndex, ChannelIndex);
+					int32 InputBufferIndex = InputIndex * NumChannels + ChannelIndex;
+					if (FAudioBufferReadRef* InputPtr = Inputs[InputBufferIndex].GetPtrOrNull())
+					{
+						InOutVertexData.BindReadVertex(AudioInputName, *InputPtr);
+					}
+					else if (InOutVertexData.IsVertexBound(AudioInputName))
+					{
+						Inputs[InputBufferIndex] = InOutVertexData.GetDataReadReference<FAudioBuffer>(AudioInputName);
+					}
 				}
 
-				InOutVertexData.BindReadVertex(GainInputNames[i], Gains[i]);
+				InOutVertexData.BindReadVertex(GainInputNames[InputIndex], Gains[InputIndex]);
 			}
 		}
 
@@ -203,13 +220,19 @@ namespace Metasound
 			float PrevGain = PrevGains[0];
 
 			// initialize the output buffers w/ the first set of input buffers 
-			for (uint32 i = 0; i < NumChannels; ++i)
+			for (uint32 Idx = 0; Idx < NumChannels; ++Idx)
 			{
-				TArrayView<const float> InputView = *Inputs[i];
-				TArrayView<float> OutputView = *Outputs[i];
-	
-				check(InputView.Num() == OutputView.Num());
-				Audio::ArrayFade(InputView, PrevGain, NextGain, OutputView);
+				TArrayView<float> OutputView = *Outputs[Idx];
+				if (const FAudioBufferReadRef *InputPtr = Inputs[Idx].GetPtrOrNull())
+				{
+					TArrayView<const float> InputView = **InputPtr;
+					check(InputView.Num() == OutputView.Num());
+					Audio::ArrayFade(InputView, PrevGain, NextGain, OutputView);
+				}
+				else
+				{
+					FMemory::Memzero(OutputView.GetData(), OutputView.Num() * sizeof(float));
+				}
 			}
 
 			PrevGains[0] = NextGain;
@@ -223,8 +246,11 @@ namespace Metasound
 				// for each channel of audio
 				for (uint32 ChanIndex = 0; ChanIndex < NumChannels; ++ChanIndex)
 				{
-					// Outputs[Chan] += Gains[i] * Inputs[i][Chan]
-					Audio::ArrayMixIn(*Inputs[InputIndex * NumChannels + ChanIndex], *Outputs[ChanIndex], PrevGain, NextGain);
+					int32 InputBufferIndex = InputIndex * NumChannels + ChanIndex;
+					if (const FAudioBufferReadRef* InputPtr = Inputs[InputBufferIndex].GetPtrOrNull())
+					{
+						Audio::ArrayMixIn(**InputPtr, *Outputs[ChanIndex], PrevGain, NextGain);
+					}
 				}
 
 				PrevGains[InputIndex] = NextGain;
@@ -322,7 +348,7 @@ namespace Metasound
 		static inline const TArray<FVertexName> GainInputNames = InitializeGainInputNames();
 
 		TArray<FFloatReadRef> Gains;
-		TArray<FAudioBufferReadRef> Inputs;
+		TArray<TOptional<FAudioBufferReadRef>> Inputs;
 		TArray<FAudioBufferWriteRef> Outputs;
 
 		TArray<float> PrevGains;
