Speculative fix for HMX related server crash in dynamic metasound transactor
#rnx

diff --git a/Engine/Plugins/Runtime/Metasound/Source/MetasoundGraphCore/Private/MetasoundDynamicOperatorTransactor.cpp b/Engine/Plugins/Runtime/Metasound/Source/MetasoundGraphCore/Private/MetasoundDynamicOperatorTransactor.cpp
--- a/Engine/Plugins/Runtime/Metasound/Source/MetasoundGraphCore/Private/MetasoundDynamicOperatorTransactor.cpp
+++ b/Engine/Plugins/Runtime/Metasound/Source/MetasoundGraphCore/Private/MetasoundDynamicOperatorTransactor.cpp
@@ -2,6 +2,7 @@
 
 #include "MetasoundDynamicOperatorTransactor.h"
 
+#include "Algo/Transform.h"
 #include "Containers/SpscQueue.h"
 #include "Containers/UnrealString.h"
 #include "MetasoundAudioBuffer.h"
@@ -129,8 +130,13 @@ namespace Metasound
 						};
 					};
 
+					auto IsValidEdge = [](const FDataEdge& InDataEdge) -> bool
+					{
+						return (InDataEdge.From.Node != nullptr) && (InDataEdge.To.Node != nullptr);
+					};
+
 					TArray<FOperatorEdge> OperatorEdges;
-					Algo::Transform(InGraph.GetDataEdges(), OperatorEdges, DataEdgeToOperatorEdge);
+					Algo::TransformIf(InGraph.GetDataEdges(), OperatorEdges, IsValidEdge, DataEdgeToOperatorEdge);
 					return OperatorEdges;
 				}
 				
@@ -463,8 +469,15 @@ namespace Metasound
 				FOperatorID FromOperatorID = DirectedGraphAlgo::GetOperatorID(Edge.From.Node);
 				FOperatorID ToOperatorID = DirectedGraphAlgo::GetOperatorID(Edge.To.Node);
 
-				OperatorMap[FromOperatorID].Outputs.Add(ToOperatorID);
-				OperatorMap[ToOperatorID].Inputs.Add(FromOperatorID);
+				if (FIncrementalSortOperatorInfo* FromInfo = OperatorMap.Find(FromOperatorID))
+				{
+					if (FIncrementalSortOperatorInfo* ToInfo = OperatorMap.Find(ToOperatorID))
+					{
+						FromInfo->Outputs.Add(ToOperatorID);
+						ToInfo->Inputs.Add(FromOperatorID);
+					}
+				}
+
 			}
 
 			// Unconnected nodes are intentionally skipped by the operator builder
@@ -1090,13 +1103,16 @@ namespace Metasound
 				{
 					if ((Edge.From.Node == &InNode) || (Edge.To.Node == &InNode))
 					{
-						// Remove edge from graph sorter
-						GraphSorter.RemoveDataEdge(DirectedGraphAlgo::GetOperatorID(Edge.From.Node), DirectedGraphAlgo::GetOperatorID(Edge.To.Node));
-
-						if (Edge.To.Node == &InNode)
+						if ((Edge.From.Node != nullptr) && (Edge.To.Node != nullptr))
 						{
-							// Track all incoming edges to remove them from dynamic operator's model.
-							OperatorsConnectedToInput.AddUnique(DirectedGraphAlgo::GetOperatorID(Edge.From.Node));
+							// Remove edge from graph sorter
+							GraphSorter.RemoveDataEdge(DirectedGraphAlgo::GetOperatorID(Edge.From.Node), DirectedGraphAlgo::GetOperatorID(Edge.To.Node));
+
+							if (Edge.To.Node == &InNode)
+							{
+								// Track all incoming edges to remove them from dynamic operator's model.
+								OperatorsConnectedToInput.AddUnique(DirectedGraphAlgo::GetOperatorID(Edge.From.Node));
+							}
 						}
 					}
 				}
diff --git a/Engine/Plugins/Runtime/Metasound/Source/MetasoundGraphCore/Private/MetasoundGraph.cpp b/Engine/Plugins/Runtime/Metasound/Source/MetasoundGraphCore/Private/MetasoundGraph.cpp
--- a/Engine/Plugins/Runtime/Metasound/Source/MetasoundGraphCore/Private/MetasoundGraph.cpp
+++ b/Engine/Plugins/Runtime/Metasound/Source/MetasoundGraphCore/Private/MetasoundGraph.cpp
@@ -185,9 +185,7 @@ PRAGMA_ENABLE_DEPRECATION_WARNINGS
 			return false;
 		}
 
-		FDataEdge Edge(FOutputDataSource(FromNode, FromVertex), FInputDataDestination(ToNode, ToVertex));
-
-		AddDataEdge(Edge);
+		Edges.Add(FDataEdge{FOutputDataSource(FromNode, FromVertex), FInputDataDestination(ToNode, ToVertex)});
 
 		return true;
 	}
diff --git a/Engine/Plugins/Runtime/Metasound/Source/MetasoundGraphCore/Private/MetasoundOperatorBuilder.cpp b/Engine/Plugins/Runtime/Metasound/Source/MetasoundGraphCore/Private/MetasoundOperatorBuilder.cpp
--- a/Engine/Plugins/Runtime/Metasound/Source/MetasoundGraphCore/Private/MetasoundOperatorBuilder.cpp
+++ b/Engine/Plugins/Runtime/Metasound/Source/MetasoundGraphCore/Private/MetasoundOperatorBuilder.cpp
@@ -447,10 +447,13 @@ namespace Metasound
 			// Set the output destinations on operator infos
 			for (const FDataEdge& Edge : InGraph.GetDataEdges())
 			{
-				const FOperatorID FromOperatorID = GetOperatorID(Edge.From.Node);
-				FGraphOperatorData::FOperatorInfo& OperatorInfo = OperatorMap.FindChecked(FromOperatorID);
+				if ((Edge.From.Node != nullptr) && (Edge.To.Node != nullptr))
+				{
+					const FOperatorID FromOperatorID = GetOperatorID(Edge.From.Node);
+					FGraphOperatorData::FOperatorInfo& OperatorInfo = OperatorMap.FindChecked(FromOperatorID);
 
-				OperatorInfo.OutputConnections.FindOrAdd(Edge.From.Vertex.VertexName).Add(FGraphOperatorData::FVertexDestination{GetOperatorID(Edge.To.Node), Edge.To.Vertex.VertexName});
+					OperatorInfo.OutputConnections.FindOrAdd(Edge.From.Vertex.VertexName).Add(FGraphOperatorData::FVertexDestination{GetOperatorID(Edge.To.Node), Edge.To.Vertex.VertexName});
+				}
 			}
 		}
 	}
