Fix for preset creation throwing stale builder warnings
#rnx

diff --git a/Engine/Plugins/Runtime/Metasound/Source/MetasoundEditor/Private/MetasoundAssetDefinitions.cpp b/Engine/Plugins/Runtime/Metasound/Source/MetasoundEditor/Private/MetasoundAssetDefinitions.cpp
--- a/Engine/Plugins/Runtime/Metasound/Source/MetasoundEditor/Private/MetasoundAssetDefinitions.cpp
+++ b/Engine/Plugins/Runtime/Metasound/Source/MetasoundEditor/Private/MetasoundAssetDefinitions.cpp
@@ -140,12 +140,18 @@ namespace Metasound::Editor
 					{
 						// Duplicate asset to preserve properties from referenced asset (ex. quality settings, soundwave properties)
 						NewMetaSound = IAssetTools::Get().DuplicateAssetWithDialogAndTitle(AssetName, FPackageName::GetLongPackagePath(PackagePath), ReferencedMetaSound, LOCTEXT("CreateMetaSoundPresetTitle", "Create MetaSound Preset"));
+
+						// Duplicate needs to initialize the asset whereas creating an asset with a factory as above in restricted mode already initializes the asset.
+						// (calling twice is explicitly avoided as this would result in an assert due to the Builder Subsystem then being left with a stale class key
+						// left pointing to the new MetaSound's builder.
+						if (NewMetaSound)
+						{
+							UMetaSoundEditorSubsystem::GetChecked().InitAsset(*NewMetaSound, ReferencedMetaSound, /*bClearDocument=*/true);
+						}
 					}
 				
 					if (NewMetaSound)
 					{
-						UMetaSoundEditorSubsystem::GetChecked().InitAsset(*NewMetaSound, ReferencedMetaSound, /*bClearDocument=*/true);
-
 						FGraphBuilder::RegisterGraphWithFrontend(*NewMetaSound);
 						ObjectsToSync.Add(NewMetaSound);
 					}
diff --git a/Engine/Plugins/Runtime/Metasound/Source/MetasoundEditor/Private/MetasoundEditorSubsystem.cpp b/Engine/Plugins/Runtime/Metasound/Source/MetasoundEditor/Private/MetasoundEditorSubsystem.cpp
--- a/Engine/Plugins/Runtime/Metasound/Source/MetasoundEditor/Private/MetasoundEditorSubsystem.cpp
+++ b/Engine/Plugins/Runtime/Metasound/Source/MetasoundEditor/Private/MetasoundEditorSubsystem.cpp
@@ -316,7 +316,8 @@ void UMetaSoundEditorSubsystem::InitAsset(UObject& InNewMetaSound, UObject* InRe
 	using namespace Metasound::Frontend;
 
 	TScriptInterface<IMetaSoundDocumentInterface> DocInterface = &InNewMetaSound;
-	FMetaSoundFrontendDocumentBuilder Builder(DocInterface);
+
+	FMetaSoundFrontendDocumentBuilder& Builder = IDocumentBuilderRegistry::GetChecked().FindOrBeginBuilding(DocInterface);
 
 	if (bClearDocument)
 	{
diff --git a/Engine/Plugins/Runtime/Metasound/Source/MetasoundEngine/Private/MetasoundDocumentBuilderRegistry.cpp b/Engine/Plugins/Runtime/Metasound/Source/MetasoundEngine/Private/MetasoundDocumentBuilderRegistry.cpp
--- a/Engine/Plugins/Runtime/Metasound/Source/MetasoundEngine/Private/MetasoundDocumentBuilderRegistry.cpp
+++ b/Engine/Plugins/Runtime/Metasound/Source/MetasoundEngine/Private/MetasoundDocumentBuilderRegistry.cpp
@@ -236,6 +236,29 @@ namespace Metasound::Engine
 		return nullptr;
 	}
 
+	TArray<FMetaSoundFrontendDocumentBuilder*> FDocumentBuilderRegistry::FindBuilders(const FMetasoundFrontendClassName& ClassName) const
+	{
+		TArray<TWeakObjectPtr<UMetaSoundBuilderBase>> Entries;
+
+		{
+			FScopeLock Lock(&BuildersCriticalSection);
+			Builders.MultiFind(ClassName, Entries);
+		}
+
+		TArray<FMetaSoundFrontendDocumentBuilder*> ToReturn;
+		Algo::TransformIf(Entries, ToReturn,
+			[](const TWeakObjectPtr<UMetaSoundBuilderBase>& WeakBuilder)
+			{
+				return WeakBuilder.IsValid();
+			},
+			[](const TWeakObjectPtr<UMetaSoundBuilderBase>& WeakBuilder)
+			{
+				return &WeakBuilder.Get()->GetBuilder();
+			});
+
+		return ToReturn;
+	}
+
 	UMetaSoundBuilderBase* FDocumentBuilderRegistry::FindBuilderObject(TScriptInterface<const IMetaSoundDocumentInterface> MetaSound) const
 	{
 		UMetaSoundBuilderBase* FoundEntry = nullptr;
diff --git a/Engine/Plugins/Runtime/Metasound/Source/MetasoundEngine/Public/MetasoundDocumentBuilderRegistry.h b/Engine/Plugins/Runtime/Metasound/Source/MetasoundEngine/Public/MetasoundDocumentBuilderRegistry.h
--- a/Engine/Plugins/Runtime/Metasound/Source/MetasoundEngine/Public/MetasoundDocumentBuilderRegistry.h
+++ b/Engine/Plugins/Runtime/Metasound/Source/MetasoundEngine/Public/MetasoundDocumentBuilderRegistry.h
@@ -114,6 +114,8 @@ namespace Metasound::Engine
 		UE_API virtual FMetaSoundFrontendDocumentBuilder* FindBuilder(TScriptInterface<IMetaSoundDocumentInterface> MetaSound) const override;
 		UE_API virtual FMetaSoundFrontendDocumentBuilder* FindBuilder(const FMetasoundFrontendClassName& InClassName, const FTopLevelAssetPath& AssetPath) const override;
 
+		UE_API virtual TArray<FMetaSoundFrontendDocumentBuilder*> FindBuilders(const FMetasoundFrontendClassName& ClassName) const override;
+
 		UE_API virtual FMetaSoundFrontendDocumentBuilder* FindOutermostBuilder(const UObject& InSubObject) const override;
 
 		UE_API virtual bool FinishBuilding(const FMetasoundFrontendClassName& InClassName, bool bForceUnregisterNodeClass = false) const override;
diff --git a/Engine/Plugins/Runtime/Metasound/Source/MetasoundFrontend/Private/MetasoundFrontendDocumentBuilder.cpp b/Engine/Plugins/Runtime/Metasound/Source/MetasoundFrontend/Private/MetasoundFrontendDocumentBuilder.cpp
--- a/Engine/Plugins/Runtime/Metasound/Source/MetasoundFrontend/Private/MetasoundFrontendDocumentBuilder.cpp
+++ b/Engine/Plugins/Runtime/Metasound/Source/MetasoundFrontend/Private/MetasoundFrontendDocumentBuilder.cpp
@@ -4207,8 +4207,7 @@ void FMetaSoundFrontendDocumentBuilder::InitDocument(const FMetasoundFrontendDoc
 			DocumentDelegates->AddPageDelegates(Frontend::DefaultPageID);
 		}
 
-		FMetasoundFrontendClassMetadata& ClassMetadata = Document.RootGraph.Metadata;
-		InitGraphClassMetadata(Document.RootGraph.Metadata, bResetVersion, InNewClassName);
+		InitGraphClassMetadata(bResetVersion, InNewClassName);
 
 #if WITH_EDITORONLY_DATA
 		// Set default doc version Metadata
@@ -4268,7 +4267,29 @@ void FMetaSoundFrontendDocumentBuilder::InitGraphClassMetadata(FMetasoundFronten
 
 void FMetaSoundFrontendDocumentBuilder::InitGraphClassMetadata(bool bResetVersion, const FMetasoundFrontendClassName* NewClassName)
 {
-	InitGraphClassMetadata(GetDocumentChecked().RootGraph.Metadata, bResetVersion, NewClassName);
+	FMetasoundFrontendClassMetadata& ThisMetadata = GetDocumentChecked().RootGraph.Metadata;
+
+#if DO_CHECK
+	using namespace Metasound::Frontend;
+	if (IDocumentBuilderRegistry* BuilderRegistry = IDocumentBuilderRegistry::Get())
+	{
+		TArray<FMetaSoundFrontendDocumentBuilder*> RegisteredBuilders = BuilderRegistry->FindBuilders(ThisMetadata.GetClassName());
+		auto IsThis = [this](const FMetaSoundFrontendDocumentBuilder* Builder)
+		{
+			checkf(Builder, TEXT("FindBuilders should always return valid pointers to builders"));
+			return this == Builder;
+		};
+		checkf(!Algo::AnyOf(RegisteredBuilders, IsThis),
+			TEXT("Renaming graph class will result in duplicate builder entries in the builder registry. "
+				"Builder(s) '%s' must be unregistered prior to updating class name"),
+			*FString::JoinBy(RegisteredBuilders, TEXT(", "), [](const FMetaSoundFrontendDocumentBuilder* RegisteredBuilder)
+			{
+				return RegisteredBuilder->GetDebugName();
+			}));
+	}
+#endif // DO_CHECK
+
+	InitGraphClassMetadata(ThisMetadata, bResetVersion, NewClassName);
 }
 
 void FMetaSoundFrontendDocumentBuilder::InitNodeLocations()
