Fix for crash when renaming asset when re-initializing MetaSound asset via MetaSoundEditorSubsystem

#rnx

diff --git a/Engine/Plugins/Runtime/Metasound/Source/MetasoundEditor/Private/MetasoundEditorSubsystem.cpp b/Engine/Plugins/Runtime/Metasound/Source/MetasoundEditor/Private/MetasoundEditorSubsystem.cpp
--- a/Engine/Plugins/Runtime/Metasound/Source/MetasoundEditor/Private/MetasoundEditorSubsystem.cpp
+++ b/Engine/Plugins/Runtime/Metasound/Source/MetasoundEditor/Private/MetasoundEditorSubsystem.cpp
@@ -317,15 +317,20 @@ void UMetaSoundEditorSubsystem::InitAsset(UObject& InNewMetaSound, UObject* InRe
 
 	TScriptInterface<IMetaSoundDocumentInterface> DocInterface = &InNewMetaSound;
 
-	FMetaSoundFrontendDocumentBuilder& Builder = IDocumentBuilderRegistry::GetChecked().FindOrBeginBuilding(DocInterface);
+	// Must force 'FinishBuilding' and use local builder as re-initializing document via
+	// 'InitDocument' call below can lead to stale builder entries in the registry when renaming.
+	const FMetasoundFrontendClassName InitClassName = DocInterface->GetConstDocument().RootGraph.Metadata.GetClassName();
+	const FTopLevelAssetPath AssetPath = DocInterface->GetAssetPathChecked();
+	IDocumentBuilderRegistry::GetChecked().FinishBuilding(InitClassName, AssetPath);
 
+	FMetaSoundFrontendDocumentBuilder Builder = FMetaSoundFrontendDocumentBuilder(DocInterface);
 	if (bClearDocument)
 	{
 		TSharedRef<FDocumentModifyDelegates> ModifyDelegates = MakeShared<FDocumentModifyDelegates>(Builder.GetConstDocumentChecked());
 		Builder.ClearDocument(ModifyDelegates);
 	}
-	
 	Builder.InitDocument();
+
 	Builder.InitNodeLocations();
 
 	constexpr bool bForceNodeCreation = true;
