Fix MetaSound interface cache corruption when renaming inputs or outputs causing editor display issues

diff --git a/Engine/Plugins/Runtime/Metasound/Source/MetasoundFrontend/Private/MetasoundFrontendDocumentCache.cpp b/Engine/Plugins/Runtime/Metasound/Source/MetasoundFrontend/Private/MetasoundFrontendDocumentCache.cpp
--- a/Engine/Plugins/Runtime/Metasound/Source/MetasoundFrontend/Private/MetasoundFrontendDocumentCache.cpp
+++ b/Engine/Plugins/Runtime/Metasound/Source/MetasoundFrontend/Private/MetasoundFrontendDocumentCache.cpp
@@ -463,10 +463,10 @@ namespace Metasound::Frontend
 
 	void FDocumentGraphInterfaceCache::OnInputNameChanged(FName OldName, FName NewName)
 	{
-		int32* Index = InputNameToIndex.Find(OldName);
-		check(Index);
-		InputNameToIndex.Remove(OldName);
-		InputNameToIndex.Add(NewName, *Index);
+		int32 Index = INDEX_NONE;
+		const bool bRemoved = InputNameToIndex.RemoveAndCopyValue(OldName, Index);
+		check(bRemoved);
+		InputNameToIndex.Add(NewName, Index);
 
 		++TransactionCount;
 	}
@@ -536,10 +536,10 @@ namespace Metasound::Frontend
 
 	void FDocumentGraphInterfaceCache::OnOutputNameChanged(FName OldName, FName NewName)
 	{
-		int32* Index = OutputNameToIndex.Find(OldName);
-		check(Index);
-		OutputNameToIndex.Remove(OldName);
-		OutputNameToIndex.Add(NewName, *Index);
+		int32 Index = INDEX_NONE;
+		const bool bRemoved = OutputNameToIndex.RemoveAndCopyValue(OldName, Index);
+		check(bRemoved);
+		OutputNameToIndex.Add(NewName, Index);
 
 		++TransactionCount;
 	}
