[CIS] added missing PLATFORM_DESKTOP scope around WavePlayerNodeV2's use of TryGetLatest() on the sound proxy

diff --git a/Engine/Plugins/Runtime/Metasound/Source/MetasoundEngine/Private/MetasoundWavePlayerNodeV2.cpp b/Engine/Plugins/Runtime/Metasound/Source/MetasoundEngine/Private/MetasoundWavePlayerNodeV2.cpp
--- a/Engine/Plugins/Runtime/Metasound/Source/MetasoundEngine/Private/MetasoundWavePlayerNodeV2.cpp
+++ b/Engine/Plugins/Runtime/Metasound/Source/MetasoundEngine/Private/MetasoundWavePlayerNodeV2.cpp
@@ -404,13 +404,13 @@ namespace Metasound
 
 		void ExecuteSubblocks()
 		{
-#if !UE_BUILD_SHIPPING
+#if PLATFORM_DESKTOP && !UE_BUILD_SHIPPING
 			// hot reload sound
 			if (bIsPlaying && Inputs.WaveAsset->IsSoundWaveValid() && Inputs.WaveAsset->GetSoundWaveProxy()->TryGetLatest())
 			{
 				StartPlaying();
 			}
-#endif // !UE_BUILD_SHIPPING
+#endif // PLATFORM_DESKTOP && !UE_BUILD_SHIPPING
 			
 			// Parse triggers and render audio
 			int32 PlayTrigIndex = 0;
