- Move VersionAsset code out of serialize path to execute asynchronously from PostLoad when in editor (or defer if needing to wait to upgrade deprecated soft object path cache). This enables versioning to safely access any referenced MetaSound data while still benefitting from running asynchronously in editor for editor performance reasons
- Consolidate all versioning code into single VersioningManager
- Move all remaining code to version MetaSound soft object cache to only compile in builds with WITH_EDITORONLY_DATA enabled

#rnx
#Synced-CL 48328941

diff --git a/Engine/Plugins/Runtime/Metasound/Source/MetasoundEditor/Private/MetasoundEditor.cpp b/Engine/Plugins/Runtime/Metasound/Source/MetasoundEditor/Private/MetasoundEditor.cpp
--- a/Engine/Plugins/Runtime/Metasound/Source/MetasoundEditor/Private/MetasoundEditor.cpp
+++ b/Engine/Plugins/Runtime/Metasound/Source/MetasoundEditor/Private/MetasoundEditor.cpp
@@ -1179,14 +1179,8 @@ namespace Metasound
 
 			if (FMetasoundAssetBase* MetaSoundAsset = IMetasoundUObjectRegistry::Get().GetObjectAsAssetBase(ObjectToEdit))
 			{
-				FMetaSoundFrontendDocumentBuilder& DocBuilder = Builder->GetBuilder();
-				if (MetaSoundAsset->VersionAsset(DocBuilder))
-				{
-					MetaSoundAsset->SetVersionedOnLoad();
-				}
-
 				constexpr bool bForceNodeCreation = false;
-				FInputNodeTemplate::GetChecked().Inject(DocBuilder, bForceNodeCreation);
+				FInputNodeTemplate::GetChecked().Inject(Builder->GetBuilder(), bForceNodeCreation);
 
 				// Ensures validation is re-run on re-opening of the editor.
 				// This is needed to refresh errors potentially caused by unloading of
diff --git a/Engine/Plugins/Runtime/Metasound/Source/MetasoundEditor/Private/MetasoundEditorGraph.cpp b/Engine/Plugins/Runtime/Metasound/Source/MetasoundEditor/Private/MetasoundEditorGraph.cpp
--- a/Engine/Plugins/Runtime/Metasound/Source/MetasoundEditor/Private/MetasoundEditorGraph.cpp
+++ b/Engine/Plugins/Runtime/Metasound/Source/MetasoundEditor/Private/MetasoundEditorGraph.cpp
@@ -2613,8 +2613,7 @@ PRAGMA_ENABLE_DEPRECATION_WARNINGS
 		// Cache data to be used in edge swapping below, being careful to not reference the input node pointer
 		// or vertex reference within the lower inner loop as the underlying node array may be reallocated by
 		// template nodes being added.
-		const FName InputName = Input.GetMemberName();
-		const FMetasoundFrontendNode* InputNode = OutBuilder.FindGraphInputNode(InputName);
+		const FMetasoundFrontendNode* InputNode = OutBuilder.FindNode(Input.NodeID);
 
 		// Potentially not used input, which is perfectly valid so early out
 		if (!InputNode)
@@ -2622,6 +2621,13 @@ PRAGMA_ENABLE_DEPRECATION_WARNINGS
 			return;
 		}
 
+		const FName InputName = InputNode->Name;
+		if (InputNode->Interface.Outputs.IsEmpty())
+		{
+			UE_LOG(LogMetaSound, Error, TEXT("Failed to migrate input node '%s': Input has no outputs assigned on interface."), *InputName.ToString());
+			return;
+		}
+
 		const FMetasoundFrontendVertex InputNodeOutputVertex = InputNode->Interface.Outputs.Last();
 
 		const FGuid InputNodeID = InputNode->GetID();
diff --git a/Engine/Plugins/Runtime/Metasound/Source/MetasoundEngine/Private/Metasound.cpp b/Engine/Plugins/Runtime/Metasound/Source/MetasoundEngine/Private/Metasound.cpp
--- a/Engine/Plugins/Runtime/Metasound/Source/MetasoundEngine/Private/Metasound.cpp
+++ b/Engine/Plugins/Runtime/Metasound/Source/MetasoundEngine/Private/Metasound.cpp
@@ -227,7 +227,12 @@ FTopLevelAssetPath UMetaSoundPatch::GetAssetPathChecked() const
 void UMetaSoundPatch::PostLoad() 
 {
 	Super::PostLoad();
-	Metasound::Engine::FAssetHelper::PostLoad(*this);
+
+#if WITH_EDITORONLY_DATA
+PRAGMA_DISABLE_DEPRECATION_WARNINGS
+	Metasound::Engine::FAssetHelper::PostLoad(*this, ReferenceAssetClassCache);
+PRAGMA_ENABLE_DEPRECATION_WARNINGS
+#endif // WITH_EDITORONLY_DATA
 }
 
 #if WITH_EDITOR
@@ -248,20 +253,17 @@ TSharedPtr<Audio::IProxyData> UMetaSoundPatch::CreateProxyData(const Audio::FPro
 	return FMetasoundAssetBase::CreateProxyData(InitParams);
 }
 
-TArray<FMetasoundAssetBase*> UMetaSoundPatch::GetReferencedAssets()
+TArray<FMetasoundAssetBase*> UMetaSoundPatch::GetReferencedAssets() const
 {
 	return Metasound::Engine::FAssetHelper::GetReferencedAssets(*this);
 }
 
-const TSet<FSoftObjectPath>& UMetaSoundPatch::GetAsyncReferencedAssetClassPaths() const 
-{
-	return ReferenceAssetClassCache;
-}
-
+#if WITH_EDITORONLY_DATA
 void UMetaSoundPatch::OnAsyncReferencedAssetsLoaded(const TArray<FMetasoundAssetBase*>& InAsyncReferences)
 {
 	Metasound::Engine::FAssetHelper::OnAsyncReferencedAssetsLoaded(*this, InAsyncReferences);
 }
+#endif // WITH_EDITORONLY_DATA
 
 bool UMetaSoundPatch::IsActivelyBuilding() const
 {
diff --git a/Engine/Plugins/Runtime/Metasound/Source/MetasoundEngine/Private/MetasoundAssetSubsystem.cpp b/Engine/Plugins/Runtime/Metasound/Source/MetasoundEngine/Private/MetasoundAssetSubsystem.cpp
--- a/Engine/Plugins/Runtime/Metasound/Source/MetasoundEngine/Private/MetasoundAssetSubsystem.cpp
+++ b/Engine/Plugins/Runtime/Metasound/Source/MetasoundEngine/Private/MetasoundAssetSubsystem.cpp
@@ -168,10 +168,12 @@ namespace Metasound::Engine
 
 	void FMetaSoundAssetManager::AddReferencedObjects(FReferenceCollector& Collector)
 	{
+#if WITH_EDITORONLY_DATA
 		for (FMetaSoundAsyncAssetDependencies& Dependencies : LoadingDependencies)
 		{
 			Collector.AddReferencedObject(Dependencies.MetaSound);
 		}
+#endif // WITH_EDITORONLY_DATA
 	}
 
 #if WITH_EDITORONLY_DATA
@@ -285,8 +287,10 @@ namespace Metasound::Engine
 			return FMetaSoundAssetKey::GetInvalid();
 		}
 
-		const FMetasoundAssetBase* MetaSoundAsset = Metasound::IMetasoundUObjectRegistry::Get().GetObjectAsAssetBase(&InObject);
-		check(MetaSoundAsset);
+#if WITH_EDITORONLY_DATA
+		// DocInterface needs to be updated and ready for extracting class info, so wait until versioning is finished
+		FVersioningManager::Get().WaitUntilVersioningComplete(InObject);
+#endif // WITH_EDITORONLY_DATA
 
 		TScriptInterface<const IMetaSoundDocumentInterface> ScriptInterface(&InObject);
 		check(ScriptInterface.GetObject() != nullptr);
@@ -413,6 +417,7 @@ namespace Metasound::Engine
 		}
 	}
 
+#if WITH_EDITORONLY_DATA
 	FMetaSoundAsyncAssetDependencies* FMetaSoundAssetManager::FindLoadingDependencies(const UObject* InParentAsset)
 	{
 		auto IsEqualMetaSoundUObject = [InParentAsset](const FMetaSoundAsyncAssetDependencies& InDependencies) -> bool
@@ -432,6 +437,7 @@ namespace Metasound::Engine
 
 		return LoadingDependencies.FindByPredicate(IsEqualID);
 	}
+#endif // WITH_EDITORONLY_DATA
 
 	FMetasoundAssetBase* FMetaSoundAssetManager::FindAsset(const FMetaSoundAssetKey& InKey) const
 	{
@@ -589,6 +595,7 @@ namespace Metasound::Engine
 		return AssetSubsystemPrivate::InitialAssetScanComplete;
 	}
 
+#if WITH_EDITORONLY_DATA
 	void FMetaSoundAssetManager::OnReferencedAssetsLoaded(int32 InLoadID)
 	{
 		FMetaSoundAsyncAssetDependencies* LoadedDependencies = FindLoadingDependencies(InLoadID);
@@ -620,8 +627,10 @@ namespace Metasound::Engine
 							}
 						}
 
-						// Update parent asset with loaded assets. 
+PRAGMA_DISABLE_DEPRECATION_WARNINGS
+						// Update parent asset with loaded assets.
 						ParentAssetBase->OnAsyncReferencedAssetsLoaded(LoadedAssetBases);
+PRAGMA_ENABLE_DEPRECATION_WARNINGS
 					}
 				}
 			}
@@ -630,6 +639,7 @@ namespace Metasound::Engine
 			RemoveLoadingDependencies(InLoadID);
 		}
 	}
+#endif // WITH_EDITORONLY_DATA
 
 	void FMetaSoundAssetManager::OnPackageLoaded(const FPackageLoadedArgs& PackageLoadedArgs)
 	{
@@ -1005,6 +1015,7 @@ namespace Metasound::Engine
 		}
 	}
 
+#if WITH_EDITORONLY_DATA
 	void FMetaSoundAssetManager::RemoveLoadingDependencies(int32 InLoadID)
 	{
 		auto IsEqualID = [InLoadID](const FMetaSoundAsyncAssetDependencies& InDependencies) -> bool
@@ -1013,6 +1024,7 @@ namespace Metasound::Engine
 		};
 		LoadingDependencies.RemoveAllSwap(IsEqualID);
 	}
+#endif // WITH_EDITORONLY_DATA
 
 	void FMetaSoundAssetManager::RenameAsset(const FAssetData& InAssetData, const FString& InOldObjectPath)
 	{
@@ -1081,14 +1093,14 @@ namespace Metasound::Engine
 	}
 #endif // WITH_EDITOR
 
-	void FMetaSoundAssetManager::RequestAsyncLoadReferencedAssets(FMetasoundAssetBase& InAssetBase)
+#if WITH_EDITORONLY_DATA
+	void FMetaSoundAssetManager::RequestAsyncLoadReferencedAssets(FMetasoundAssetBase& InAssetBase, const TSet<FSoftObjectPath>* Paths)
 	{
-		const TSet<FSoftObjectPath>& AsyncReferences = InAssetBase.GetAsyncReferencedAssetClassPaths();
-		if (AsyncReferences.Num() > 0)
+		if (Paths && !Paths->IsEmpty())
 		{
 			if (UObject* OwningAsset = InAssetBase.GetOwningAsset())
 			{
-				TArray<FSoftObjectPath> PathsToLoad = AsyncReferences.Array();
+				TArray<FSoftObjectPath> PathsToLoad = Paths->Array();
 
 				// Protect against duplicate calls to async load assets. 
 				if (FMetaSoundAsyncAssetDependencies* ExistingAsyncLoad = FindLoadingDependencies(OwningAsset))
@@ -1102,9 +1114,12 @@ namespace Metasound::Engine
 
 				int32 AsyncLoadID = AsyncLoadIDCounter++;
 
-				auto AssetsLoadedDelegate = [this, AsyncLoadID]()
+				auto AssetsLoadedDelegate = [AsyncLoadID]()
 				{
-					this->OnReferencedAssetsLoaded(AsyncLoadID);
+					if (FMetaSoundAssetManager* ThisManager = FMetaSoundAssetManager::Get())
+					{
+						ThisManager->OnReferencedAssetsLoaded(AsyncLoadID);
+					}
 				};
 
 				// Store async loading data for use when async load is complete.
@@ -1122,6 +1137,7 @@ namespace Metasound::Engine
 		}
 
 	}
+#endif // WITH_EDITORONLY_DATA
 
 	void FMetaSoundAssetManager::ReloadMetaSoundAssets() const
 	{
@@ -1382,7 +1398,7 @@ namespace Metasound::Engine
 			const bool bWasLoaded = AssetData.IsAssetLoaded();
 			if (!bWasLoaded && ClassInfo.bIsValid)
 			{
-				if (ClassInfo.DocInfo.DocumentVersion >= FMetasoundFrontendDocument::GetMaxVersion())
+				if (ClassInfo.DocInfo.DocumentVersion >= FVersioningManager::GetMaxDocumentVersion())
 				{
 					Results.PackagesUpToDate.Add(PackagePath);
 					UE_LOG(LogMetaSound, Display, TEXT("MetaSound '%s' already versioned & contains valid asset tags. Skipping reserialization."), *AssetData.GetFullName());
@@ -1394,6 +1410,7 @@ namespace Metasound::Engine
 			FMetasoundAssetBase* MetaSound = IMetasoundUObjectRegistry::Get().GetObjectAsAssetBase(MetaSoundObject);
 			if (MetaSound)
 			{
+				FVersioningManager::Get().WaitUntilVersioningComplete(*MetaSoundObject);
 				bool bVersioned = MetaSound->GetVersionedOnLoad();
 				TScriptInterface<IMetaSoundDocumentInterface> DocInterface(MetaSoundObject);
 				const FMetasoundFrontendDocumentMetadata& DocMetadata = DocInterface->GetConstDocument().Metadata;
@@ -1404,13 +1421,9 @@ namespace Metasound::Engine
 				}
 				else
 				{
-					FMetaSoundFrontendDocumentBuilder& DocBuilder = FDocumentBuilderRegistry::GetChecked().FindOrBeginBuilding(MetaSoundObject);
-					bVersioned = MetaSound->VersionAsset(DocBuilder);
-					DocVersion = DocMetadata.Version.Number.ToString();
-					if (bVersioned)
-					{
-						MetaSound->SetVersionedOnLoad();
-					}
+					FVersioningManager::Get().VersionAssetAsync(*MetaSoundObject, false /* bIsDeterministic */);
+					FVersioningManager::Get().WaitUntilVersioningComplete(*MetaSoundObject);
+					bVersioned = MetaSound->GetVersionedOnLoad();
 				}
 
 				constexpr bool bAlwaysMarkDirty = true;
@@ -1439,16 +1452,17 @@ namespace Metasound::Engine
 	}
 #endif // WITH_EDITORONLY_DATA
 
-	void FMetaSoundAssetManager::WaitUntilAsyncLoadReferencedAssetsComplete(FMetasoundAssetBase& InAssetBase)
+#if WITH_EDITORONLY_DATA
+	void FMetaSoundAssetManager::WaitUntilAsyncLoadReferencedAssetsComplete(const FMetasoundAssetBase& InAssetBase)
 	{
-		TSet<FMetasoundAssetBase*> TransitiveReferences;
-		TArray<FMetasoundAssetBase*> TransitiveReferencesQueue;
+		TSet<const FMetasoundAssetBase*> TransitiveReferences;
+		TArray<const FMetasoundAssetBase*> TransitiveReferencesQueue;
 		TransitiveReferences.Add(&InAssetBase);
 		TransitiveReferencesQueue.Add(&InAssetBase);
 		while (!TransitiveReferencesQueue.IsEmpty())
 		{
-			FMetasoundAssetBase* Reference = TransitiveReferencesQueue.Pop();
-			UObject* OwningAsset = Reference->GetOwningAsset();
+			const FMetasoundAssetBase* Reference = TransitiveReferencesQueue.Pop();
+			const UObject* OwningAsset = Reference->GetOwningAsset();
 			if (!OwningAsset)
 			{
 				continue;
@@ -1479,7 +1493,7 @@ namespace Metasound::Engine
 				}
 			}
 
-			for (FMetasoundAssetBase* NextReference : Reference->GetReferencedAssets())
+			for (const FMetasoundAssetBase* NextReference : Reference->GetReferencedAssets())
 			{
 				bool bAlreadyInSet;
 				TransitiveReferences.Add(NextReference, &bAlreadyInSet);
@@ -1490,6 +1504,7 @@ namespace Metasound::Engine
 			}
 		}
 	}
+#endif // WITH_EDITORONLY_DATA
 
 	void DeinitializeAssetManager()
 	{
diff --git a/Engine/Plugins/Runtime/Metasound/Source/MetasoundEngine/Private/MetasoundBuilderBase.cpp b/Engine/Plugins/Runtime/Metasound/Source/MetasoundEngine/Private/MetasoundBuilderBase.cpp
--- a/Engine/Plugins/Runtime/Metasound/Source/MetasoundEngine/Private/MetasoundBuilderBase.cpp
+++ b/Engine/Plugins/Runtime/Metasound/Source/MetasoundEngine/Private/MetasoundBuilderBase.cpp
@@ -386,8 +386,7 @@ bool UMetaSoundBuilderBase::ClearMemberMetadata(const FGuid& InMemberID)
 
 bool UMetaSoundBuilderBase::ConformObjectToDocument()
 {
-	TScriptInterface<IMetaSoundDocumentInterface> DocInterface(&Builder.CastDocumentObjectChecked<UObject>());
-	return DocInterface->ConformObjectToDocument();
+	return Builder.ConformObjectToDocument();
 }
 
 bool UMetaSoundBuilderBase::ContainsNode(const FMetaSoundNodeHandle& NodeHandle) const
diff --git a/Engine/Plugins/Runtime/Metasound/Source/MetasoundEngine/Private/MetasoundBuilderSubsystem.cpp b/Engine/Plugins/Runtime/Metasound/Source/MetasoundEngine/Private/MetasoundBuilderSubsystem.cpp
--- a/Engine/Plugins/Runtime/Metasound/Source/MetasoundEngine/Private/MetasoundBuilderSubsystem.cpp
+++ b/Engine/Plugins/Runtime/Metasound/Source/MetasoundEngine/Private/MetasoundBuilderSubsystem.cpp
@@ -17,6 +17,7 @@
 #include "MetasoundDocumentBuilderRegistry.h"
 #include "MetasoundDocumentInterface.h"
 #include "MetasoundDynamicOperatorTransactor.h"
+#include "MetasoundEngineAsset.h"
 #include "MetasoundFrontendDataTypeRegistry.h"
 #include "MetasoundFrontendDocument.h"
 #include "MetasoundFrontendDocumentIdGenerator.h"
diff --git a/Engine/Plugins/Runtime/Metasound/Source/MetasoundEngine/Private/MetasoundDocumentBuilderRegistry.cpp b/Engine/Plugins/Runtime/Metasound/Source/MetasoundEngine/Private/MetasoundDocumentBuilderRegistry.cpp
--- a/Engine/Plugins/Runtime/Metasound/Source/MetasoundEngine/Private/MetasoundDocumentBuilderRegistry.cpp
+++ b/Engine/Plugins/Runtime/Metasound/Source/MetasoundEngine/Private/MetasoundDocumentBuilderRegistry.cpp
@@ -4,6 +4,8 @@
 
 #include "HAL/PlatformProperties.h"
 #include "MetasoundAssetManager.h"
+#include "MetasoundEngineAsset.h"
+#include "MetasoundFrontendDocumentVersioning.h"
 #include "MetasoundGlobals.h"
 #include "MetasoundSettings.h"
 #include "MetasoundTrace.h"
@@ -40,6 +42,14 @@ namespace Metasound::Engine
 		UE_CLOG(NumStale > 0, LogMetaSound, Display, TEXT("BuilderRegistry is shutting down with %i stale entries"), NumStale);
 	}
 
+#if WITH_EDITORONLY_DATA
+	void FDocumentBuilderRegistry::WaitUntilVersioningComplete(UObject& InMetaSoundObject) const
+	{
+		using namespace Frontend;
+		FVersioningManager::Get().WaitUntilVersioningComplete(InMetaSoundObject);
+	}
+#endif // WITH_EDITORONLY_DATA
+
 	void FDocumentBuilderRegistry::AddBuilderInternal(const FMetasoundFrontendClassName& InClassName, UMetaSoundBuilderBase* NewBuilder) const
 	{
 		FScopeLock Lock(&BuildersCriticalSection);
diff --git a/Engine/Plugins/Runtime/Metasound/Source/MetasoundEngine/Private/MetasoundEngineAsset.h b/Engine/Plugins/Runtime/Metasound/Source/MetasoundEngine/Private/MetasoundEngineAsset.h
--- a/Engine/Plugins/Runtime/Metasound/Source/MetasoundEngine/Private/MetasoundEngineAsset.h
+++ b/Engine/Plugins/Runtime/Metasound/Source/MetasoundEngine/Private/MetasoundEngineAsset.h
@@ -21,6 +21,7 @@
 #if WITH_EDITORONLY_DATA
 #include "Algo/Transform.h"
 #include "Interfaces/ITargetPlatform.h"
+#include "MetasoundFrontendDocumentVersioning.h"
 #include "MetasoundFrontendRegistryContainer.h"
 #include "Misc/DataValidation.h"
 #include "Serialization/JsonWriter.h"
@@ -41,11 +42,6 @@ namespace Metasound::Engine
 	 */
 	struct FAssetHelper
 	{
-		static bool SerializationRequiresDeterminism(bool bIsCooking)
-		{
-			return bIsCooking || IsRunningCookCommandlet();
-		}
-
 #if WITH_EDITOR
 		static void PreDuplicate(TScriptInterface<IMetaSoundDocumentInterface> MetaSound, FObjectDuplicationParameters& DupParams)
 		{
@@ -134,7 +130,9 @@ namespace Metasound::Engine
 				TArray<FText> PathStrings;
 				Algo::Transform(AssetPaths, PathStrings, [](const FTopLevelAssetPath& Path) { return FText::FromString(Path.ToString()); });
 				InOutContext.AddError(FText::Format(LOCTEXT("UniqueClassNameValidation",
-					"Multiple assets use the same class name which may result in unintended behavior. This may happen when an asset is moved, then the move is reverted in revision control without removing the newly created asset. Please remove the offending asset or duplicate it to automatically generate a new class name." \
+					"Multiple assets use the same class name which may result in unintended behavior. This may happen when an asset is moved, "
+					"then the move is reverted in revision control without removing the newly created asset. Please remove the offending asset "
+					"or duplicate it to automatically generate a new class name." \
 					"\nConflicting Asset Paths:\n{0}"), FText::Join(FText::FromString(TEXT("\n")), PathStrings)));
 			}
 
@@ -224,20 +222,22 @@ namespace Metasound::Engine
 		}
 
 		template <typename TMetaSoundObject>
-		static TArray<FMetasoundAssetBase*> GetReferencedAssets(TMetaSoundObject& InMetaSound)
+		static TArray<FMetasoundAssetBase*> GetReferencedAssets(const TMetaSoundObject& InMetaSound)
 		{
 			TArray<FMetasoundAssetBase*> ReferencedAssets;
 
 			IMetasoundUObjectRegistry& UObjectRegistry = IMetasoundUObjectRegistry::Get();
 
-			for (TObjectPtr<UObject>& Object : InMetaSound.ReferencedAssetClassObjects)
+			for (const TObjectPtr<UObject>& Object : InMetaSound.ReferencedAssetClassObjects)
 			{
-				if (FMetasoundAssetBase* Asset = UObjectRegistry.GetObjectAsAssetBase(Object))
-				{
-					ReferencedAssets.Add(Asset);
-				}
-				else
+				if (Object.Get() != nullptr)
 				{
+					if (FMetasoundAssetBase* Asset = UObjectRegistry.GetObjectAsAssetBase(Object))
+					{
+						ReferencedAssets.Add(Asset);
+						continue;
+					}
+
 					UE_LOG(LogMetaSound, Error, TEXT("Referenced asset \"%s\", referenced from \"%s\", is not convertible to FMetasoundAssetBase"), *Object->GetPathName(), *InMetaSound.GetPathName());
 				}
 			}
@@ -250,11 +250,6 @@ namespace Metasound::Engine
 #if WITH_EDITORONLY_DATA
 			using namespace Frontend;
 
-			if (IMetaSoundAssetManager* AssetManager = IMetaSoundAssetManager::Get())
-			{
-				AssetManager->WaitUntilAsyncLoadReferencedAssetsComplete(InMetaSound);
-			}
-
 			const bool bIsCooking = InSaveContext.IsCooking();
 			const bool bCanEverExecute = Metasound::CanEverExecuteGraph(bIsCooking);
 			if (!bCanEverExecute)
@@ -264,9 +259,9 @@ namespace Metasound::Engine
 				{
 					PlatformName = *TargetPlatform->IniPlatformName();
 				}
-				const bool bIsDeterministic = SerializationRequiresDeterminism(bIsCooking);
-				FDocumentIDGenerator::FScopeDeterminism DeterminismScope(bIsDeterministic);
 
+				const bool bIsDeterministic = bIsCooking || IsRunningCookCommandlet();
+				FDocumentIDGenerator::FScopeDeterminism DeterminismScope(bIsDeterministic);
 
 				FMetaSoundAssetCookOptions CookOptions;
 				if (const UMetaSoundSettings* Settings = GetDefault<UMetaSoundSettings>(); ensure(Settings))
@@ -308,60 +303,56 @@ namespace Metasound::Engine
 #endif // WITH_EDITORONLY_DATA
 		}
 
-		template <typename TMetaSoundObject>
-		static void SerializeToArchive(TMetaSoundObject& InMetaSound, FArchive& InArchive)
+		static void SerializeToArchive(FMetasoundAssetBase& InMetaSound, FArchive& InArchive)
 		{
+		}
+
 #if WITH_EDITORONLY_DATA
+		static void PostLoad(FMetasoundAssetBase& InMetaSound, TSet<FSoftObjectPath>& ReferenceAssetClassCache)
+		{
 			using namespace Frontend;
 
-			bool bVersionedAsset = false;
+			UObject* OwningObject = InMetaSound.GetOwningAsset();
+			check(OwningObject);
 
-			if (InArchive.IsLoading() && !InArchive.IsTransacting())
+			// Do not call asset manager on CDO objects which may be loaded before asset 
+			// manager is set.
+			const bool bIsCDO = OwningObject->HasAnyFlags(RF_ClassDefaultObject);
+			if (!bIsCDO)
 			{
-				const bool bIsTransacting = InArchive.IsTransacting();
-				TStrongObjectPtr<UMetaSoundBuilderBase> Builder;
+PRAGMA_DISABLE_DEPRECATION_WARNINGS
+				if (ReferenceAssetClassCache.IsEmpty())
 				{
-					FGCScopeGuard ScopeGuard;
-					Builder.Reset(&FDocumentBuilderRegistry::GetChecked().FindOrBeginBuilding(InMetaSound, bIsTransacting));
+					const bool bIsDeterministic = IsRunningCookCommandlet();
+					FVersioningManager::Get().VersionAssetAsync(*OwningObject, bIsDeterministic);
 				}
-
+				else
 				{
-					const bool bIsCooking = InArchive.IsCooking();
-					const bool bIsDeterministic = SerializationRequiresDeterminism(bIsCooking);
-					FDocumentIDGenerator::FScopeDeterminism DeterminismScope(bIsDeterministic);
-					check(Builder.IsValid());
-					bVersionedAsset = InMetaSound.VersionAsset(Builder->GetBuilder());
+					// Versioning asset is applied after async load of soft object
+					// references to ensure document properties are ready for versioning.
+					IMetaSoundAssetManager::GetChecked().RequestAsyncLoadReferencedAssets(InMetaSound, &ReferenceAssetClassCache);
 				}
 
-				Builder->ClearInternalFlags(EInternalObjectFlags::Async);
-			}
-
-			if (bVersionedAsset)
-			{
-				InMetaSound.SetVersionedOnLoad();
-			}
-#endif // WITH_EDITORONLY_DATA
-		}
-
-		template<typename TMetaSoundObject>
-		static void PostLoad(TMetaSoundObject& InMetaSound)
-		{
-			using namespace Frontend;
-			// Do not call asset manager on CDO objects which may be loaded before asset 
-			// manager is set.
-			const bool bIsCDO = InMetaSound.HasAnyFlags(RF_ClassDefaultObject);
-			if (!bIsCDO)
-			{
-				if (InMetaSound.GetAsyncReferencedAssetClassPaths().Num() > 0)
+				// There is no way to determine if a commandlet is going to read directly from a MetaSound object
+				// immediately after post load to perform tasks such as serialization diffing, reference scrubbing,
+				// etc. outside of the MetaSound Builder/Asset Manager API. Therefore, its best to just immediately
+				// wait until versioning is complete in this context.
+				if (IsRunningCommandlet())
 				{
-					IMetaSoundAssetManager::GetChecked().RequestAsyncLoadReferencedAssets(InMetaSound);
+					using namespace Frontend;
+					FVersioningManager::Get().WaitUntilVersioningComplete(*OwningObject);
 				}
 			}
+PRAGMA_ENABLE_DEPRECATION_WARNINGS
 		}
 
 		template<typename TMetaSoundObject>
 		static void OnAsyncReferencedAssetsLoaded(TMetaSoundObject& InMetaSound, const TArray<FMetasoundAssetBase*>& InAsyncReferences)
 		{
+PRAGMA_DISABLE_DEPRECATION_WARNINGS
+			const bool bVersioningDeferred = !InMetaSound.ReferenceAssetClassCache.IsEmpty();
+PRAGMA_ENABLE_DEPRECATION_WARNINGS
+
 			for (FMetasoundAssetBase* AssetBase : InAsyncReferences)
 			{
 				if (AssetBase)
@@ -369,11 +360,21 @@ namespace Metasound::Engine
 					if (UObject* OwningAsset = AssetBase->GetOwningAsset())
 					{
 						InMetaSound.ReferencedAssetClassObjects.Add(OwningAsset);
+
+PRAGMA_DISABLE_DEPRECATION_WARNINGS
 						InMetaSound.ReferenceAssetClassCache.Remove(FSoftObjectPath(OwningAsset));
+PRAGMA_ENABLE_DEPRECATION_WARNINGS
 					}
 				}
 			}
+
+			if (bVersioningDeferred)
+			{
+				const bool bIsDeterministic = IsRunningCookCommandlet();
+				Frontend::FVersioningManager::Get().VersionAssetAsync(InMetaSound, bIsDeterministic);
+			}
 		}
+#endif // WITH_EDITORONLY_DATA
 	};
 } // namespace Metasound::Engine
 #undef LOCTEXT_NAMESPACE // MetasoundEngine
diff --git a/Engine/Plugins/Runtime/Metasound/Source/MetasoundEngine/Private/MetasoundSource.cpp b/Engine/Plugins/Runtime/Metasound/Source/MetasoundEngine/Private/MetasoundSource.cpp
--- a/Engine/Plugins/Runtime/Metasound/Source/MetasoundEngine/Private/MetasoundSource.cpp
+++ b/Engine/Plugins/Runtime/Metasound/Source/MetasoundEngine/Private/MetasoundSource.cpp
@@ -719,22 +719,17 @@ void UMetaSoundSource::SetReferencedAssets(TSet<Metasound::Frontend::IMetaSoundA
 }
 #endif // WITH_EDITOR
 
-TArray<FMetasoundAssetBase*> UMetaSoundSource::GetReferencedAssets()
+TArray<FMetasoundAssetBase*> UMetaSoundSource::GetReferencedAssets() const
 {
 	return Metasound::Engine::FAssetHelper::GetReferencedAssets(*this);
 }
 
-const TSet<FSoftObjectPath>& UMetaSoundSource::GetAsyncReferencedAssetClassPaths() const 
-{
-	return ReferenceAssetClassCache;
-}
-
+#if WITH_EDITORONLY_DATA
 void UMetaSoundSource::OnAsyncReferencedAssetsLoaded(const TArray<FMetasoundAssetBase*>& InAsyncReferences)
 {
 	Metasound::Engine::FAssetHelper::OnAsyncReferencedAssetsLoaded(*this, InAsyncReferences);
 }
 
-#if WITH_EDITORONLY_DATA
 void UMetaSoundSource::MigrateEditorGraph(FMetaSoundFrontendDocumentBuilder& OutBuilder)
 {
 	PRAGMA_DISABLE_DEPRECATION_WARNINGS
@@ -767,7 +762,12 @@ FText UMetaSoundSource::GetDisplayName() const
 void UMetaSoundSource::PostLoad()
 {
 	Super::PostLoad();
-	Metasound::Engine::FAssetHelper::PostLoad(*this);
+
+#if WITH_EDITORONLY_DATA
+PRAGMA_DISABLE_DEPRECATION_WARNINGS
+	Metasound::Engine::FAssetHelper::PostLoad(*this, ReferenceAssetClassCache);
+PRAGMA_ENABLE_DEPRECATION_WARNINGS
+#endif // WITH_EDITORONLY_DATA
 
 	Duration = GetDuration();
 	bLooping = IsLooping();
diff --git a/Engine/Plugins/Runtime/Metasound/Source/MetasoundEngine/Public/Metasound.h b/Engine/Plugins/Runtime/Metasound/Source/MetasoundEngine/Public/Metasound.h
--- a/Engine/Plugins/Runtime/Metasound/Source/MetasoundEngine/Public/Metasound.h
+++ b/Engine/Plugins/Runtime/Metasound/Source/MetasoundEngine/Public/Metasound.h
@@ -87,10 +87,10 @@ protected:
 	UPROPERTY()
 	TSet<TObjectPtr<UObject>> ReferencedAssetClassObjects;
 
-	UPROPERTY()
+#if WITH_EDITORONLY_DATA
+	UPROPERTY(meta = (DeprecatedProperty, DeprecationMessage = "Soft reference cache is no longer supported"))
 	TSet<FSoftObjectPath> ReferenceAssetClassCache;
 
-#if WITH_EDITORONLY_DATA
 	UPROPERTY(meta=(DeprecatedProperty, DeprecationMessage = "Use EditorGraph instead as it is now transient and generated via the FrontendDocument dynamically."))
 	TObjectPtr<UMetasoundEditorGraphBase> Graph;
 
@@ -175,9 +175,12 @@ public:
 	{
 		return ReferencedAssetClassKeys;
 	}
-	UE_API virtual TArray<FMetasoundAssetBase*> GetReferencedAssets() override;
-	UE_API virtual const TSet<FSoftObjectPath>& GetAsyncReferencedAssetClassPaths() const override;
+	UE_API virtual TArray<FMetasoundAssetBase*> GetReferencedAssets() const override;
+
+#if WITH_EDITORONLY_DATA
+	UE_DEPRECATED(5.8, "This is deprecated behavior and will become UE_INTERNAL API in a future release for back compat.")
 	UE_API virtual void OnAsyncReferencedAssetsLoaded(const TArray<FMetasoundAssetBase*>& InAsyncReferences) override;
+#endif // WITH_EDITORONLY_DATA
 
 	UObject* GetOwningAsset() override
 	{
diff --git a/Engine/Plugins/Runtime/Metasound/Source/MetasoundEngine/Public/MetasoundAssetSubsystem.h b/Engine/Plugins/Runtime/Metasound/Source/MetasoundEngine/Public/MetasoundAssetSubsystem.h
--- a/Engine/Plugins/Runtime/Metasound/Source/MetasoundEngine/Public/MetasoundAssetSubsystem.h
+++ b/Engine/Plugins/Runtime/Metasound/Source/MetasoundEngine/Public/MetasoundAssetSubsystem.h
@@ -98,10 +98,13 @@ PRAGMA_ENABLE_DEPRECATION_WARNINGS
 #if WITH_EDITOR
 		UE_API bool ReplaceReferencesInDirectory(const TArray<FMetaSoundAssetDirectory>& InDirectories, const FNodeRegistryKey& OldClassKey, const FNodeRegistryKey& NewClassKey) const;
 #endif // WITH_EDITOR
-		UE_API void RequestAsyncLoadReferencedAssets(FMetasoundAssetBase& InAssetBase);
 
-		UE_DEPRECATED(5.6, "Implementation has been privatized")
-		void OnAssetScanComplete() { }
+#if WITH_EDITORONLY_DATA
+		UE_DEPRECATED(5.8,
+			"This is a deprecated feature, and exists to port old soft object references to hard references. "
+			"This will become UE_INTERNAL API in a future release.")
+		UE_API virtual void RequestAsyncLoadReferencedAssets(FMetasoundAssetBase& InAssetBase, const TSet<FSoftObjectPath>* Paths = nullptr) override;
+#endif // WITH_EDITORONLY_DATA
 
 		UE_API void SearchAndIterateDirectoryAssets(const TArray<FDirectoryPath>& InDirectories, TFunctionRef<void(const FAssetData&)> InFunction) const;
 
@@ -156,10 +159,9 @@ PRAGMA_ENABLE_DEPRECATION_WARNINGS
 #if WITH_EDITORONLY_DATA
 		UE_API FMetaSoundAssetManager::FVersionAssetResults SetAccessFlagsOnAssetsInFolders(const TArray<FString>& FolderPaths, EMetasoundFrontendClassAccessFlags Flags, bool bRecursePaths) const override;
 		UE_API virtual FVersionAssetResults VersionAssetsInFolders(const TArray<FString>& FolderPaths, bool bRecursePaths = true) const override;
+		UE_API virtual void WaitUntilAsyncLoadReferencedAssetsComplete(const FMetasoundAssetBase& InAssetBase) override;
 #endif // WITH_EDITORONLY_DATA
 
-		UE_API virtual void WaitUntilAsyncLoadReferencedAssetsComplete(FMetasoundAssetBase& InAssetBase) override;
-
 		/* FGCObject */
 		UE_API virtual void AddReferencedObjects(FReferenceCollector& Collector) override;
 		virtual FString GetReferencerName() const override { return TEXT("FMetaSoundAssetManager"); }
@@ -169,7 +171,9 @@ PRAGMA_ENABLE_DEPRECATION_WARNINGS
 	private:
 		FMetaSoundAssetKey AddOrUpdateFromObjectInternal(const UObject& InObject);
 
+#if WITH_EDITORONLY_DATA
 		TArray<FMetaSoundAsyncAssetDependencies> LoadingDependencies;
+#endif // WITH_EDITORONLY_DATA
 
 		static UE_API void DepthFirstVisit_AssetKey(const FMetaSoundAssetKey& InKey, TFunctionRef<TSet<FMetaSoundAssetKey>(const FMetaSoundAssetKey&)> VisitFunction);
 
@@ -179,10 +183,12 @@ PRAGMA_ENABLE_DEPRECATION_WARNINGS
 		// or may be out-of-date if edited after last registration attempt via AddOrUpdate... calls.
 		TArray<Frontend::FMetaSoundAssetClassInfo> FindAssetClassInfoInternal(const FMetaSoundAssetKey& InKey) const;
 
+#if WITH_EDITORONLY_DATA
 		FMetaSoundAsyncAssetDependencies* FindLoadingDependencies(const UObject* InParentAsset);
 		FMetaSoundAsyncAssetDependencies* FindLoadingDependencies(int32 InLoadID);
 		void RemoveLoadingDependencies(int32 InLoadID);
 		void OnReferencedAssetsLoaded(int32 InLoadID);
+#endif // WITH_EDITORONLY_DATA
 
 		bool GetReferencedPresetHierarchyInternal(FMetasoundAssetBase& InAsset, TArray<FMetasoundAssetBase*>& OutReferencedAssets) const;
 
@@ -229,6 +235,7 @@ struct FMetaSoundAsyncAssetDependencies
 {
 	GENERATED_BODY()
 
+#if WITH_EDITORONLY_DATA
 	// ID of the async load
 	int32 LoadID = 0;
 
@@ -241,6 +248,7 @@ struct FMetaSoundAsyncAssetDependencies
 
 	// Handle to in-flight streaming request
 	TSharedPtr<FStreamableHandle> StreamableHandle;
+#endif // WITH_EDITORONLY_DATA
 };
 
 
diff --git a/Engine/Plugins/Runtime/Metasound/Source/MetasoundEngine/Public/MetasoundDocumentBuilderRegistry.h b/Engine/Plugins/Runtime/Metasound/Source/MetasoundEngine/Public/MetasoundDocumentBuilderRegistry.h
--- a/Engine/Plugins/Runtime/Metasound/Source/MetasoundEngine/Public/MetasoundDocumentBuilderRegistry.h
+++ b/Engine/Plugins/Runtime/Metasound/Source/MetasoundEngine/Public/MetasoundDocumentBuilderRegistry.h
@@ -80,6 +80,10 @@ namespace Metasound::Engine
 				return *CastChecked<BuilderClass>(Builder);
 			}
 
+#if WITH_EDITORONLY_DATA
+			WaitUntilVersioningComplete(InMetaSoundObject);
+#endif // WITH_EDITORONLY_DATA
+
 			FNameBuilder BuilderName;
 			BuilderName.Append(InMetaSoundObject.GetName());
 			BuilderName.Append(TEXT("_Builder"));
@@ -166,6 +170,10 @@ PRAGMA_ENABLE_DEPRECATION_WARNINGS
 		UE_API void SetEventLogVerbosity(ELogEvent Event, ELogVerbosity::Type Verbosity);
 
 	private:
+#if WITH_EDITORONLY_DATA
+		UE_API void WaitUntilVersioningComplete(UObject& InMetaSoundObject) const;
+#endif // WITH_EDITORONLY_DATA
+
 		UE_API void AddBuilderInternal(const FMetasoundFrontendClassName& InClassName, UMetaSoundBuilderBase* NewBuilder) const;
 		UE_API bool CanPostEventLog(ELogEvent Event, ELogVerbosity::Type Verbosity) const;
 		UE_API void FinishBuildingInternal(UMetaSoundBuilderBase& Builder, bool bForceUnregisterNodeClass) const;
diff --git a/Engine/Plugins/Runtime/Metasound/Source/MetasoundEngine/Public/MetasoundSource.h b/Engine/Plugins/Runtime/Metasound/Source/MetasoundEngine/Public/MetasoundSource.h
--- a/Engine/Plugins/Runtime/Metasound/Source/MetasoundEngine/Public/MetasoundSource.h
+++ b/Engine/Plugins/Runtime/Metasound/Source/MetasoundEngine/Public/MetasoundSource.h
@@ -126,10 +126,10 @@ protected:
 	UPROPERTY()
 	TSet<TObjectPtr<UObject>> ReferencedAssetClassObjects;
 
-	UPROPERTY()
+#if WITH_EDITORONLY_DATA
+	UPROPERTY(meta = (DeprecatedProperty, DeprecationMessage = "Soft reference cache is no longer supported"))
 	TSet<FSoftObjectPath> ReferenceAssetClassCache;
 
-#if WITH_EDITORONLY_DATA
 	UPROPERTY(meta=(DeprecatedProperty, DeprecationMessage = "Use EditorGraph instead as it is now transient and generated via the FrontendDocument dynamically."))
 	TObjectPtr<UMetasoundEditorGraphBase> Graph;
 
@@ -240,9 +240,12 @@ public:
 	{
 		return ReferencedAssetClassKeys;
 	}
-	UE_API virtual TArray<FMetasoundAssetBase*> GetReferencedAssets() override;
-	UE_API virtual const TSet<FSoftObjectPath>& GetAsyncReferencedAssetClassPaths() const override;
+	UE_API virtual TArray<FMetasoundAssetBase*> GetReferencedAssets() const override;
+
+#if WITH_EDITORONLY_DATA
+	UE_DEPRECATED(5.8, "This is deprecated behavior and will become UE_INTERNAL API in a future release for back compat.")
 	UE_API virtual void OnAsyncReferencedAssetsLoaded(const TArray<FMetasoundAssetBase*>& InAsyncReferences) override;
+#endif // WITH_EDITORONLY_DATA
 
 	UE_API virtual void BeginDestroy() override;
 	UE_API virtual void PreSave(FObjectPreSaveContext InSaveContext) override;
diff --git a/Engine/Plugins/Runtime/Metasound/Source/MetasoundEngineTest/Private/MetasoundInterfaceUpdateTests.cpp b/Engine/Plugins/Runtime/Metasound/Source/MetasoundEngineTest/Private/MetasoundInterfaceUpdateTests.cpp
--- a/Engine/Plugins/Runtime/Metasound/Source/MetasoundEngineTest/Private/MetasoundInterfaceUpdateTests.cpp
+++ b/Engine/Plugins/Runtime/Metasound/Source/MetasoundEngineTest/Private/MetasoundInterfaceUpdateTests.cpp
@@ -3,6 +3,7 @@
 #include "IAudioParameterInterfaceRegistry.h"
 #include "Interfaces/MetasoundTestInterfaces.h"
 #include "MetasoundBuilderSubsystem.h"
+#include "../../MetasoundFrontend/Public/MetasoundFrontendDocumentVersioning.h"
 #include "MetasoundSource.h"
 #include "MetasoundStandardNodesNames.h"
 #include "Misc/AutomationTest.h"
@@ -163,7 +164,7 @@ namespace Metasound::Test::InterfaceUpdate
 		}
 
 		// This will internally throw an error if there are problems when versioning the asset
-		MetasoundSource->VersionAsset(MetasoundFrontendBuilder);
+		FVersioningManager::Get().VersionInterfaces(MetasoundFrontendBuilder);
 
 		return true;
 	}
diff --git a/Engine/Plugins/Runtime/Metasound/Source/MetasoundFrontend/Private/MetasoundAssetBase.cpp b/Engine/Plugins/Runtime/Metasound/Source/MetasoundFrontend/Private/MetasoundAssetBase.cpp
--- a/Engine/Plugins/Runtime/Metasound/Source/MetasoundFrontend/Private/MetasoundAssetBase.cpp
+++ b/Engine/Plugins/Runtime/Metasound/Source/MetasoundFrontend/Private/MetasoundAssetBase.cpp
@@ -130,70 +130,6 @@ namespace Metasound
 				FMetasoundFrontendClass FrontendClass;
 				FTopLevelAssetPath AssetPath;
 			};
-
-			void GetUpdatePathForDocument(const FMetasoundFrontendVersion& InCurrentVersion, const FMetasoundFrontendVersion& InTargetVersion, TArray<TSharedPtr<IBuilderVersionTransform>>& OutUpgradePath)
-			{
-				OutUpgradePath.Reset();
-
-				if (InCurrentVersion.Name == InTargetVersion.Name)
-				{
-					// Get all associated registered interfaces
-					TArray<FMetasoundFrontendVersion> RegisteredVersions = ISearchEngine::Get().FindAllRegisteredInterfacesWithName(InTargetVersion.Name);
-
-					// Filter registry entries that exist between current version and target version
-					auto FilterRegistryEntries = [&InCurrentVersion, &InTargetVersion](const FMetasoundFrontendVersion& InVersion)
-					{
-						const bool bIsGreaterThanCurrent = InVersion.Number > InCurrentVersion.Number;
-						const bool bIsLessThanOrEqualToTarget = InVersion.Number <= InTargetVersion.Number;
-
-						return bIsGreaterThanCurrent && bIsLessThanOrEqualToTarget;
-					};
-					RegisteredVersions = RegisteredVersions.FilterByPredicate(FilterRegistryEntries);
-
-					// sort registry entries to create an ordered upgrade path.
-					RegisteredVersions.Sort();
-
-					// Get registry entries from registry keys.
-					auto GetUpdateTransform = [](const FMetasoundFrontendVersion& InVersion)
-					{
-						return IInterfaceRegistry::Get().FindInterfaceUpdateTransform(InVersion);
-					};
-					Algo::Transform(RegisteredVersions, OutUpgradePath, GetUpdateTransform);
-				}
-			}
-
-			bool UpdateDocumentInterface(const TArray<TSharedPtr<IBuilderVersionTransform>>& InUpgradePath, const FMetasoundFrontendVersion& InterfaceVersion, FMetaSoundFrontendDocumentBuilder& Builder)
-			{
-				const FMetasoundFrontendVersionNumber* LastVersionUpdated = nullptr;
-				for (const TSharedPtr<IBuilderVersionTransform>& Transform : InUpgradePath)
-				{
-					if (ensureAlwaysMsgf(Transform.IsValid(), TEXT("Transform should always be provided, even if referencing the default transform, by this point in updating document interfaces")))
-					{
-						if (Transform->Transform(Builder))
-						{
-							LastVersionUpdated = &Transform->GetVersion().Number;
-						}
-					}
-				}
-
-				if (LastVersionUpdated)
-				{
-					const FMetasoundFrontendClassMetadata& Metadata = Builder.GetMetasoundAsset().GetDocumentChecked().RootGraph.Metadata;
-#if WITH_EDITOR
-					const FString AssetName = *Metadata.GetDisplayName().ToString();
-#else
-					const FString AssetName = *Metadata.GetClassName().ToString();
-#endif // !WITH_EDITOR
-					UE_LOG(LogMetaSound, Display, TEXT("Asset '%s' interface '%s' updated: '%s' --> '%s'"),
-						*AssetName,
-						*InterfaceVersion.Name.ToString(),
-						*InterfaceVersion.Number.ToString(),
-						*LastVersionUpdated->ToString());
-					return true;
-				}
-
-				return false;
-			}
 		} // namespace AssetBasePrivate
 
 		FConsoleVariableMulticastDelegate CVarMetaSoundBlockRateChanged;
@@ -460,6 +396,7 @@ void FMetasoundAssetBase::UpdateAndRegisterForSerialization(const Metasound::Fro
 
 	UObject* Owner = GetOwningAsset();
 	check(Owner);
+
 	IMetaSoundAssetManager::GetChecked().AddOrUpdateFromObject(*Owner);
 
 	bool bDidUpdate = false;
@@ -591,44 +528,7 @@ bool FMetasoundAssetBase::IsInterfaceDeclared(const FMetasoundFrontendVersion& I
 #if WITH_EDITORONLY_DATA
 bool FMetasoundAssetBase::VersionAsset(FMetaSoundFrontendDocumentBuilder& Builder)
 {
-	using namespace Metasound;
-
-	METASOUND_TRACE_CPUPROFILER_EVENT_SCOPE(MetaSoundAssetBase::VersionAsset);
-
-	bool bDidEdit = Frontend::VersionDocument(Builder);
-
-	// TODO: Move this logic to builder API above, which will require rewriting update transforms to
-	// take in builder instead of DocumentHandle.
-	{
-		const FMetasoundFrontendDocument& Document = Builder.GetConstDocumentChecked();
-		bool bInterfaceUpdated = false;
-		bool bPassUpdated = true;
-
-		// Has to be re-run until no pass reports an update in case versions
-		// fork (ex. an interface splits into two newly named interfaces).
-		while (bPassUpdated)
-		{
-			bPassUpdated = false;
-
-			const TArray<FMetasoundFrontendVersion> Versions = Document.Interfaces.Array();
-
-			for (const FMetasoundFrontendVersion& Version : Versions)
-			{
-				bPassUpdated |= TryUpdateInterfaceFromVersion(Version, Builder);
-			}
-
-			bInterfaceUpdated |= bPassUpdated;
-		}
-
-		if (bInterfaceUpdated)
-		{
-			TScriptInterface<IMetaSoundDocumentInterface> Interface(GetOwningAsset());
-			Interface->ConformObjectToDocument();
-		}
-		bDidEdit |= bInterfaceUpdated;
-	}
-
-	return bDidEdit;
+	return false;
 }
 #endif // WITH_EDITORONLY_DATA
 
@@ -898,52 +798,7 @@ void FMetasoundAssetBase::UpdateAndRegisterReferencesForSerialization(const Meta
 		}
 	}
 }
-#endif // WITH_EDITORONLY_DATA
-
-bool FMetasoundAssetBase::TryUpdateInterfaceFromVersion(const FMetasoundFrontendVersion& Version, FMetaSoundFrontendDocumentBuilder& Builder)
-{
-	using namespace Metasound::Frontend;
-	using namespace AssetBasePrivate;
-
-	const FMetasoundFrontendVersion TargetInterfaceVersion = GetTargetInterfaceVersion(Version);
-	if (TargetInterfaceVersion.IsValid())
-	{
-		TArray<TSharedPtr<IBuilderVersionTransform>> UpgradePath;
-		GetUpdatePathForDocument(Version, TargetInterfaceVersion, UpgradePath);
-		const bool bUpdated = UpdateDocumentInterface(UpgradePath, Version, Builder);
-		ensureMsgf(bUpdated, TEXT("Target interface '%s' was out-of-date but interface failed to be updated"), *TargetInterfaceVersion.ToString());
-		return bUpdated;
-	}
-
-	return false;
-}
-
-FMetasoundFrontendVersion FMetasoundAssetBase::GetTargetInterfaceVersion(const FMetasoundFrontendVersion& InterfaceVersion) const
-{
-	using namespace Metasound::Frontend;
-
-	// Find registered target interface.
-	FMetasoundFrontendVersion TargetInterfaceVersion;
-	bool bFoundTargetInterface = ISearchEngine::Get().FindHighestInterfaceVersion(InterfaceVersion.Name, TargetInterfaceVersion);
-	if (!bFoundTargetInterface)
-	{
-		UE_LOG(LogMetaSound, Warning,
-			TEXT("Could not check for interface updates. Target interface is not registered [InterfaceVersion:%s] when attempting to update root graph of asset (%s). "
-				"Ensure that the module which registers the interface has been loaded before the asset is loaded."),
-			*InterfaceVersion.ToString(),
-			*GetOwningAssetName());
-		return { };
-	}
-
-	if (TargetInterfaceVersion == InterfaceVersion)
-	{
-		return { };
-	}
-
-	return TargetInterfaceVersion;
-}
 
-#if WITH_EDITORONLY_DATA
 bool FMetasoundAssetBase::GetVersionedOnLoad() const
 {
 	return bVersionedOnLoad;
diff --git a/Engine/Plugins/Runtime/Metasound/Source/MetasoundFrontend/Private/MetasoundFrontendDocument.cpp b/Engine/Plugins/Runtime/Metasound/Source/MetasoundFrontend/Private/MetasoundFrontendDocument.cpp
--- a/Engine/Plugins/Runtime/Metasound/Source/MetasoundFrontend/Private/MetasoundFrontendDocument.cpp
+++ b/Engine/Plugins/Runtime/Metasound/Source/MetasoundFrontend/Private/MetasoundFrontendDocument.cpp
@@ -1695,6 +1695,28 @@ void FMetasoundFrontendGraphClass::IterateGraphPages(TFunctionRef<void(const FMe
 	}
 }
 
+#if WITH_EDITORONLY_DATA
+bool FMetasoundFrontendDocument::MigrateProperties()
+{
+	bool bMigrated = false;
+	if (ArchetypeVersion.IsValid())
+	{
+		bMigrated = true;
+		Interfaces.Add(ArchetypeVersion);
+		ArchetypeVersion = FMetasoundFrontendVersion::GetInvalid();
+	}
+
+	if (!InterfaceVersions.IsEmpty())
+	{
+		bMigrated = true;
+		Interfaces.Append(InterfaceVersions);
+		InterfaceVersions.Reset();
+	}
+
+	return bMigrated;
+}
+#endif // WITH_EDITORONLY_DATA
+
 void FMetasoundFrontendGraphClass::ResetGraphs()
 {
 	PagedGraphs.Empty();
@@ -1708,7 +1730,8 @@ TArray<FMetasoundFrontendGraph>& FMetasoundFrontendGraphClass::IPropertyVersionT
 
 FMetasoundFrontendVersionNumber FMetasoundFrontendDocument::GetMaxVersion()
 {
-	return Metasound::Frontend::GetMaxDocumentVersion();
+	using namespace Metasound::Frontend;
+	return FVersioningManager::GetMaxDocumentVersion();
 }
 #endif // WITH_EDITORONLY_DATA
 
diff --git a/Engine/Plugins/Runtime/Metasound/Source/MetasoundFrontend/Private/MetasoundFrontendDocumentBuilder.cpp b/Engine/Plugins/Runtime/Metasound/Source/MetasoundFrontend/Private/MetasoundFrontendDocumentBuilder.cpp
--- a/Engine/Plugins/Runtime/Metasound/Source/MetasoundFrontend/Private/MetasoundFrontendDocumentBuilder.cpp
+++ b/Engine/Plugins/Runtime/Metasound/Source/MetasoundFrontend/Private/MetasoundFrontendDocumentBuilder.cpp
@@ -2283,7 +2283,7 @@ Metasound::Frontend::EAutoUpdateEligibility FMetaSoundFrontendDocumentBuilder::C
 	// to ensure a change wasn't created that didn't contain a version change.
 	if (RegistryClassInfo.Version == NodeClassMetadata.GetVersion())
 	{
-		if (ChangeIDComparisonEnabledInAutoUpdate())
+		if (FVersioningManager::ChangeIDComparisonEnabledInAutoUpdate())
 		{
 			// WIP: Merging these paths.  Should no longer require using different
 			// logic to determine changes in native vs asset class definitions.
@@ -2510,6 +2510,11 @@ bool FMetaSoundFrontendDocumentBuilder::ConformGraphOutputNodeToClass(const FMet
 	return false;
 }
 
+bool FMetaSoundFrontendDocumentBuilder::ConformObjectToDocument()
+{
+	return GetDocumentInterfaceChecked().ConformObjectToDocument();
+}
+
 #if WITH_EDITORONLY_DATA
 bool FMetaSoundFrontendDocumentBuilder::ConfigureDocument(bool bLogFailure)
 {
@@ -4213,7 +4218,7 @@ void FMetaSoundFrontendDocumentBuilder::InitDocument(const FMetasoundFrontendDoc
 		// Set default doc version Metadata
 		{
 			FMetasoundFrontendDocumentMetadata& DocMetadata = Document.Metadata;
-			DocMetadata.Version.Number = GetMaxDocumentVersion();
+			DocMetadata.Version.Number = FVersioningManager::GetMaxDocumentVersion();
 		}
 #endif // WITH_EDITORONLY_DATA
 
diff --git a/Engine/Plugins/Runtime/Metasound/Source/MetasoundFrontend/Private/MetasoundFrontendDocumentModifyDelegates.cpp b/Engine/Plugins/Runtime/Metasound/Source/MetasoundFrontend/Private/MetasoundFrontendDocumentModifyDelegates.cpp
--- a/Engine/Plugins/Runtime/Metasound/Source/MetasoundFrontend/Private/MetasoundFrontendDocumentModifyDelegates.cpp
+++ b/Engine/Plugins/Runtime/Metasound/Source/MetasoundFrontend/Private/MetasoundFrontendDocumentModifyDelegates.cpp
@@ -29,7 +29,7 @@ PRAGMA_ENABLE_DEPRECATION_WARNINGS
 		// Modify delegates may register when a builder is constructed prior to a document
 		// being versioned. In this legacy case, add the default page ID delegates as the
 		// default page is the only version that exists on legacy documents.
-		if (Document.Metadata.Version.Number < Frontend::GetPageMigrationVersion())
+		if (Document.Metadata.Version.Number < FVersioningManager::GetPageMigrationVersion())
 		{
 			if (Document.RootGraph.GetConstGraphPages().IsEmpty())
 			{
diff --git a/Engine/Plugins/Runtime/Metasound/Source/MetasoundFrontend/Private/MetasoundFrontendDocumentVersioning.cpp b/Engine/Plugins/Runtime/Metasound/Source/MetasoundFrontend/Private/MetasoundFrontendDocumentVersioning.cpp
--- a/Engine/Plugins/Runtime/Metasound/Source/MetasoundFrontend/Private/MetasoundFrontendDocumentVersioning.cpp
+++ b/Engine/Plugins/Runtime/Metasound/Source/MetasoundFrontend/Private/MetasoundFrontendDocumentVersioning.cpp
@@ -1,6 +1,8 @@
 // Copyright Epic Games, Inc. All Rights Reserved.
 #include "MetasoundFrontendDocumentVersioning.h"
 
+#include "Misc/CoreMiscDefines.h"
+
 #if WITH_EDITORONLY_DATA
 #include "Algo/Transform.h"
 #include "Algo/Unique.h"
@@ -10,9 +12,12 @@
 #include "Interfaces/MetasoundFrontendInterface.h"
 #include "Interfaces/MetasoundFrontendInterfaceRegistry.h"
 #include "MetasoundAccessPtr.h"
+#include "MetasoundAssetManager.h"
 #include "MetasoundDocumentInterface.h"
+#include "MetasoundFrontendDocument.h"
 #include "MetasoundFrontendDocumentBuilder.h"
 #include "MetasoundFrontendDocumentController.h"
+#include "MetasoundFrontendDocumentIdGenerator.h"
 #include "MetasoundFrontendRegistries.h"
 #include "MetasoundFrontendSearchEngine.h"
 #include "MetasoundFrontendTransform.h"
@@ -20,11 +25,11 @@
 #include "MetasoundTrace.h"
 #include "Misc/App.h"
 
+#define METASOUND_VERSIONING_LOG(Verbosity, Format, ...) if (DocumentTransform::GetVersioningLoggingEnabled()) { UE_LOG(LogMetaSound, Verbosity, Format, ##__VA_ARGS__); }
+
 
 namespace Metasound::Frontend
 {
-#define METASOUND_VERSIONING_LOG(Verbosity, Format, ...) if (DocumentTransform::GetVersioningLoggingEnabled()) { UE_LOG(LogMetaSound, Verbosity, Format, ##__VA_ARGS__); }
-
 	namespace VersioningPrivate
 	{
 		static int32 MetaSoundAutoUpdateUseAssetChangeIDsCVar = 0;
@@ -36,6 +41,113 @@ namespace Metasound::Frontend
 			TEXT("0: Don't use Change Ids (default), !0: Use ChangeIDs"),
 			ECVF_Default);
 
+		FMetasoundFrontendVersion GetTargetInterfaceVersion(const FMetaSoundFrontendDocumentBuilder& Builder, const FMetasoundFrontendVersion& InterfaceVersion)
+		{
+			using namespace Metasound::Frontend;
+
+			// Find registered target interface.
+			FMetasoundFrontendVersion TargetInterfaceVersion;
+			bool bFoundTargetInterface = ISearchEngine::Get().FindHighestInterfaceVersion(InterfaceVersion.Name, TargetInterfaceVersion);
+			if (!bFoundTargetInterface)
+			{
+				METASOUND_VERSIONING_LOG(Warning,
+					TEXT("Could not check for interface updates. Target interface is not registered [InterfaceVersion:%s] when attempting to update root graph of asset (%s). "
+						"Ensure that the module which registers the interface has been loaded before the asset is loaded."),
+					*InterfaceVersion.ToString(),
+					*Builder.GetDebugName());
+				return { };
+			}
+
+			if (TargetInterfaceVersion == InterfaceVersion)
+			{
+				return { };
+			}
+
+			return TargetInterfaceVersion;
+		}
+
+		void GetUpdatePathForDocument(const FMetasoundFrontendVersion& InCurrentVersion, const FMetasoundFrontendVersion& InTargetVersion, TArray<TSharedPtr<IBuilderVersionTransform>>& OutUpgradePath)
+		{
+			OutUpgradePath.Reset();
+
+			if (InCurrentVersion.Name == InTargetVersion.Name)
+			{
+				// Get all associated registered interfaces
+				TArray<FMetasoundFrontendVersion> RegisteredVersions = ISearchEngine::Get().FindAllRegisteredInterfacesWithName(InTargetVersion.Name);
+
+				// Filter registry entries that exist between current version and target version
+				auto FilterRegistryEntries = [&InCurrentVersion, &InTargetVersion](const FMetasoundFrontendVersion& InVersion)
+				{
+					const bool bIsGreaterThanCurrent = InVersion.Number > InCurrentVersion.Number;
+					const bool bIsLessThanOrEqualToTarget = InVersion.Number <= InTargetVersion.Number;
+
+					return bIsGreaterThanCurrent && bIsLessThanOrEqualToTarget;
+				};
+				RegisteredVersions = RegisteredVersions.FilterByPredicate(FilterRegistryEntries);
+
+				// sort registry entries to create an ordered upgrade path.
+				RegisteredVersions.Sort();
+
+				// Get registry entries from registry keys.
+				auto GetUpdateTransform = [](const FMetasoundFrontendVersion& InVersion)
+				{
+					return IInterfaceRegistry::Get().FindInterfaceUpdateTransform(InVersion);
+				};
+				Algo::Transform(RegisteredVersions, OutUpgradePath, GetUpdateTransform);
+			}
+		}
+
+		bool UpdateDocumentInterface(const TArray<TSharedPtr<IBuilderVersionTransform>>& InUpgradePath, const FMetasoundFrontendVersion& InterfaceVersion, FMetaSoundFrontendDocumentBuilder& Builder)
+		{
+			const FMetasoundFrontendVersionNumber* LastVersionUpdated = nullptr;
+			for (const TSharedPtr<IBuilderVersionTransform>& Transform : InUpgradePath)
+			{
+				if (ensureAlwaysMsgf(Transform.IsValid(), TEXT("Transform should always be provided, even if referencing the default transform, by this point in updating document interfaces")))
+				{
+					if (Transform->Transform(Builder))
+					{
+						LastVersionUpdated = &Transform->GetVersion().Number;
+					}
+				}
+			}
+
+			if (LastVersionUpdated)
+			{
+				const FMetasoundFrontendClassMetadata& Metadata = Builder.GetMetasoundAsset().GetDocumentChecked().RootGraph.Metadata;
+#if WITH_EDITOR
+				const FString AssetName = *Metadata.GetDisplayName().ToString();
+#else
+				const FString AssetName = *Metadata.GetClassName().ToString();
+#endif // !WITH_EDITOR
+				UE_LOG(LogMetaSound, Display, TEXT("Asset '%s' interface '%s' updated: '%s' --> '%s'"),
+					*AssetName,
+					*InterfaceVersion.Name.ToString(),
+					*InterfaceVersion.Number.ToString(),
+					*LastVersionUpdated->ToString());
+				return true;
+			}
+
+			return false;
+		}
+
+		// Checks if version is up-to-date. If so, returns true. If false, updates the interfaces within the given asset's document to the most recent version.
+		bool TryUpdateInterfaceFromVersion(const FMetasoundFrontendVersion& Version, FMetaSoundFrontendDocumentBuilder& Builder)
+		{
+			using namespace Metasound::Frontend;
+
+			const FMetasoundFrontendVersion TargetInterfaceVersion = GetTargetInterfaceVersion(Builder, Version);
+			if (TargetInterfaceVersion.IsValid())
+			{
+				TArray<TSharedPtr<IBuilderVersionTransform>> UpgradePath;
+				GetUpdatePathForDocument(Version, TargetInterfaceVersion, UpgradePath);
+				const bool bUpdated = UpdateDocumentInterface(UpgradePath, Version, Builder);
+				ensureMsgf(bUpdated, TEXT("Target interface '%s' was out-of-date but interface failed to be updated"), *TargetInterfaceVersion.ToString());
+				return bUpdated;
+			}
+
+			return false;
+		}
+
 		class FMigratePagePropertiesTransform : public FMetaSoundFrontendDocumentBuilder::IPropertyVersionTransform
 		{
 		public:
@@ -99,21 +211,15 @@ PRAGMA_ENABLE_DEPRECATION_WARNINGS
 			}
 		};
 
-		class FVersionDocumentInterfacesTransform : public FMetaSoundFrontendDocumentBuilder::IPropertyVersionTransform
+		class FMigratePropertiesTransform : public FMetaSoundFrontendDocumentBuilder::IPropertyVersionTransform
 		{
 		public:
-			virtual ~FVersionDocumentInterfacesTransform() = default;
+			virtual ~FMigratePropertiesTransform() = default;
 
 			bool Transform(FMetaSoundFrontendDocumentBuilder& OutBuilder) const override
 			{
-				FMetasoundFrontendDocument& Document = GetDocumentUnsafe(OutBuilder);
-				if (Document.RequiresInterfaceVersioning())
-				{
-					Document.VersionInterfaces();
-					return true;
-				}
-
-				return false;
+				const bool bMigrated = GetDocumentUnsafe(OutBuilder).MigrateProperties();
+				return bMigrated;
 			}
 		};
 
@@ -870,28 +976,174 @@ PRAGMA_ENABLE_DEPRECATION_WARNINGS
 		}
 	} // namespace VersioningPrivate
 
-	bool ChangeIDComparisonEnabledInAutoUpdate()
+	TUniquePtr<FVersioningManager> FVersioningManager::Instance;
+
+	FVersioningManager::~FVersioningManager()
+	{
+		using namespace UE::Tasks;
+
+		TArray<FTask> PendingTasks;
+		Algo::Transform(ActiveVersionTaskData, PendingTasks, [](const TPair<FUniqueId, FTaskData>& Pair) { return Pair.Value.Task; });
+		UE::Tasks::Wait(PendingTasks);
+		ActiveVersionTaskData.Reset();
+	}
+
+	FVersioningManager& FVersioningManager::Get()
+	{
+		// Its possible some commandlets load MetaSounds prior
+		// to the MetaSoundEngine module loading (ex. generating
+		// asset manifests), so attempt to generate one on the fly
+		// when accessing here.
+		if (!Instance.IsValid())
+		{
+			Instance = MakeUnique<FVersioningManager>();
+		}
+		return *Instance;
+	}
+
+	void FVersioningManager::Initialize()
+	{
+		if (!Instance.IsValid())
+		{
+			Instance = MakeUnique<FVersioningManager>();
+		}
+	}
+
+	void FVersioningManager::Shutdown()
+	{
+		Instance.Reset();
+	}
+
+	bool FVersioningManager::ChangeIDComparisonEnabledInAutoUpdate()
 	{
 		return VersioningPrivate::MetaSoundAutoUpdateUseAssetChangeIDsCVar != 0;
 	}
 
-	bool VersionDocument(FMetaSoundFrontendDocumentBuilder& Builder)
+	FMetasoundFrontendVersionNumber FVersioningManager::GetMaxDocumentVersion()
+	{
+		return FMetasoundFrontendVersionNumber { 1, 16 };
+	}
+
+	FMetasoundFrontendVersionNumber FVersioningManager::GetPageMigrationVersion()
+	{
+		return FMetasoundFrontendVersionNumber { 1, 13 };
+	}
+
+	void FVersioningManager::VersionAssetAsync(UObject& InMetaSound, bool bIsDeterministic)
+	{
+		TScriptInterface<IMetaSoundDocumentInterface> DocInterface(&InMetaSound);
+		if (!ensureAlwaysMsgf(DocInterface.GetObject(), TEXT("Failed to version MetaSound '%s': Object does not implement MetaSound Document Interface."), *InMetaSound.GetName()))
+		{
+			return;
+		}
+
+		if (!ensureAlwaysMsgf(InMetaSound.IsAsset(), TEXT("Failed to version MetaSound '%s': Only serialized MetaSound assets require versioning."), *InMetaSound.GetName()))
+		{
+			return;
+		}
+
+		const FUniqueId AssetId = InMetaSound.GetUniqueID();
+		if (!ActiveVersionTaskData.Contains(AssetId))
+		{
+			InMetaSound.SetInternalFlags(EInternalObjectFlags::Async);
+			TStrongObjectPtr<UObject> AssetPtr(&InMetaSound);
+
+			check(DocInterface.GetObject());
+			const FMetasoundFrontendDocument& Document = DocInterface->GetConstDocument();
+			const FMetasoundFrontendClassName& ClassName = Document.RootGraph.Metadata.GetClassName();
+
+			// Ensure nothing is currently holding on to a builder and use a local builder for versioning. All other call sites
+			// (such as a request to register and grabbing a new builder) should wait until versioning is complete for safety.
+			IDocumentBuilderRegistry::GetChecked().FinishBuilding(ClassName, DocInterface->GetAssetPathChecked());
+
+			FScopeLock Lock(&ActiveVersionCritSection);
+			ActiveVersionTaskData.Add(AssetId, FTaskData
+			{
+				.AssetPtr = MoveTemp(AssetPtr),
+				.Task = UE::Tasks::Launch(UE_SOURCE_LOCATION, [this, AssetId, bIsDeterministic]()
+				{
+					FScopeLock VersionLock(&ActiveVersionCritSection);
+					FTaskData& TaskData = ActiveVersionTaskData.FindChecked(AssetId);
+					FDocumentIDGenerator::FScopeDeterminism DeterminismScope(bIsDeterministic);
+					FMetaSoundFrontendDocumentBuilder DocBuilder(TaskData.AssetPtr.Get());
+
+					const bool bVersionedDoc = VersionDocument(DocBuilder);
+					const bool bVersionedInterfaces = VersionInterfacesInternal(DocBuilder);
+
+					if (bVersionedDoc || bVersionedInterfaces)
+					{
+						DocBuilder.GetMetasoundAsset().SetVersionedOnLoad();
+					}
+
+					TaskData.AssetPtr->ClearInternalFlags(EInternalObjectFlags::Async);
+					TaskData.AssetPtr.Reset();
+					ActiveVersionTaskData.Remove(AssetId);
+				})
+			});
+		}
+	}
+
+	bool FVersioningManager::VersionInterfaces(FMetaSoundFrontendDocumentBuilder& DocBuilder) const
+	{
+		WaitUntilVersioningComplete(DocBuilder.CastDocumentObjectChecked<UObject>());
+		return VersionInterfacesInternal(DocBuilder);
+	}
+
+	bool FVersioningManager::VersionInterfacesInternal(FMetaSoundFrontendDocumentBuilder& DocBuilder) const
+	{
+		using namespace Metasound;
+		using namespace Metasound::Frontend;
+
+		bool bDidEdit = false;
+		{
+			const FMetasoundFrontendDocument& Document = DocBuilder.GetConstDocumentChecked();
+			bool bInterfaceUpdated = false;
+			bool bPassUpdated = true;
+
+			// Has to be re-run until no pass reports an update in case versions
+			// fork (ex. an interface splits into two newly named interfaces).
+			while (bPassUpdated)
+			{
+				bPassUpdated = false;
+
+				const TArray<FMetasoundFrontendVersion> Versions = Document.Interfaces.Array();
+
+				for (const FMetasoundFrontendVersion& Version : Versions)
+				{
+					bPassUpdated |= VersioningPrivate::TryUpdateInterfaceFromVersion(Version, DocBuilder);
+				}
+
+				bInterfaceUpdated |= bPassUpdated;
+			}
+
+			if (bInterfaceUpdated)
+			{
+				DocBuilder.ConformObjectToDocument();
+			}
+			bDidEdit |= bInterfaceUpdated;
+
+			return bDidEdit;
+		}
+	}
+
+	bool FVersioningManager::VersionDocument(FMetaSoundFrontendDocumentBuilder& DocBuilder) const
 	{
 		using namespace VersioningPrivate;
 
-		bool bWasUpdated = false;
+		METASOUND_TRACE_CPUPROFILER_EVENT_SCOPE(FVersioningManager::VersionDocument);
 
-		UObject& MetaSoundAsset = Builder.CastDocumentObjectChecked<UObject>();
-		const FName Name(*MetaSoundAsset.GetName());
-		const FString Path = MetaSoundAsset.GetPathName();
+		// Applies all property migration for data prior to versioning schema.
+		// Only runs if data has yet to be migrated.
+		bool bWasUpdated = FMigratePropertiesTransform().Transform(DocBuilder);
 
-		// Copied as value will be mutated with each applicable transform below
-		const FMetasoundFrontendVersionNumber InitVersionNumber = Builder.GetConstDocumentChecked().Metadata.Version.Number;
+		UObject& MetaSoundObject = DocBuilder.CastDocumentObjectChecked<UObject>();
+		const FName Name(*MetaSoundObject.GetName());
+		const FString Path = MetaSoundObject.GetPathName();
 
-		// Old manual property transform that was applied prior to versioning schema being added.
-		// Only runs if internal logic finds necessary.
-		bWasUpdated = FVersionDocumentInterfacesTransform().Transform(Builder);
+		// Copied as value will be mutated with each applicable transform below
+		const FMetasoundFrontendVersionNumber InitVersionNumber = DocBuilder.GetConstDocumentChecked().Metadata.Version.Number;
 
+		using namespace Metasound::Frontend;
 		if (InitVersionNumber < GetMaxDocumentVersion())
 		{
 			// Controller (Soft Deprecated) Transforms
@@ -899,9 +1151,9 @@ PRAGMA_ENABLE_DEPRECATION_WARNINGS
 			{
 				// Page Graph migration must be completed for graph accessor back
 				// compat prior to all controller versioning, so just do it here.
-				FMigratePagePropertiesTransform().Transform(Builder);
+				FMigratePagePropertiesTransform().Transform(DocBuilder);
 
-				FDocumentHandle DocHandle = Builder.GetMetasoundAsset().GetDocumentHandle();
+				FDocumentHandle DocHandle = DocBuilder.GetMetasoundAsset().GetDocumentHandle();
 
 				bWasUpdated |= FVersionDocument_1_1(Name, Path).Transform(DocHandle);
 				bWasUpdated |= FVersionDocument_1_2(Name, Path).Transform(DocHandle);
@@ -917,16 +1169,35 @@ PRAGMA_ENABLE_DEPRECATION_WARNINGS
 				// No longer supported, new versions should go in VersioningPrivate::VersionBuilderDocument
 			}
 
-			bWasUpdated |= VersionBuilderDocument(Builder);
+			bWasUpdated |= VersionBuilderDocument(DocBuilder);
 			if (bWasUpdated)
 			{
-				const FMetasoundFrontendVersionNumber& NewVersionNumber = Builder.GetConstDocumentChecked().Metadata.Version.Number;
+				const FMetasoundFrontendVersionNumber& NewVersionNumber = DocBuilder.GetConstDocumentChecked().Metadata.Version.Number;
 				METASOUND_VERSIONING_LOG(Verbose, TEXT("MetaSound at '%s' Document Versioned: '%s' --> '%s'"), *Path, *InitVersionNumber.ToString(), *NewVersionNumber.ToString());
 			}
 		}
 
 		return bWasUpdated;
 	}
-#undef METASOUND_VERSIONING_LOG
+
+	void FVersioningManager::WaitUntilVersioningComplete(const UObject& InMetaSound) const
+	{
+		// Async reference load must be waited for prior to versioning wait as this is a
+		// required step to kick off versioning for older assets with soft object references.
+		if (IMetaSoundAssetManager* AssetManager = IMetaSoundAssetManager::Get())
+		{
+			if (const FMetasoundAssetBase* MetaSoundAsset = AssetManager->GetAsAsset(InMetaSound))
+			{
+				AssetManager->WaitUntilAsyncLoadReferencedAssetsComplete(*MetaSoundAsset);
+			}
+		}
+
+		if (const FTaskData* TaskData = ActiveVersionTaskData.Find(InMetaSound.GetUniqueID()))
+		{
+			TaskData->Task.Wait();
+		}
+	}
 } // namespace Metasound::Frontend
+
+#undef METASOUND_VERSIONING_LOG
 #endif // WITH_EDITORONLY_DATA
diff --git a/Engine/Plugins/Runtime/Metasound/Source/MetasoundFrontend/Private/MetasoundFrontendDocumentVersioning.h b/Engine/Plugins/Runtime/Metasound/Source/MetasoundFrontend/Private/MetasoundFrontendDocumentVersioning.h
deleted file mode 100644
diff --git a/Engine/Plugins/Runtime/Metasound/Source/MetasoundFrontend/Private/MetasoundFrontendNodeController.cpp b/Engine/Plugins/Runtime/Metasound/Source/MetasoundFrontend/Private/MetasoundFrontendNodeController.cpp
--- a/Engine/Plugins/Runtime/Metasound/Source/MetasoundFrontend/Private/MetasoundFrontendNodeController.cpp
+++ b/Engine/Plugins/Runtime/Metasound/Source/MetasoundFrontend/Private/MetasoundFrontendNodeController.cpp
@@ -1244,7 +1244,7 @@ PRAGMA_ENABLE_DEPRECATION_WARNINGS
 			// to ensure a change wasn't created that didn't contain a version change.
 			if (RegistryClass.Metadata.GetVersion() == NodeClassMetadata.GetVersion())
 			{
-				if (Frontend::ChangeIDComparisonEnabledInAutoUpdate())
+				if (FVersioningManager::ChangeIDComparisonEnabledInAutoUpdate())
 				{
 					// WIP: Merging these paths.  Should no longer require using different
 					// logic to determine changes in native vs asset class definitions.
diff --git a/Engine/Plugins/Runtime/Metasound/Source/MetasoundFrontend/Public/MetasoundAssetBase.h b/Engine/Plugins/Runtime/Metasound/Source/MetasoundFrontend/Public/MetasoundAssetBase.h
--- a/Engine/Plugins/Runtime/Metasound/Source/MetasoundFrontend/Public/MetasoundAssetBase.h
+++ b/Engine/Plugins/Runtime/Metasound/Source/MetasoundFrontend/Public/MetasoundAssetBase.h
@@ -121,13 +121,21 @@ public:
 
 	// Returns set of class references set call to serialize in the editor
 	// Used at runtime load register referenced classes.
-	virtual TArray<FMetasoundAssetBase*> GetReferencedAssets() = 0;
+	virtual TArray<FMetasoundAssetBase*> GetReferencedAssets() const = 0;
 
-	// Return all dependent asset paths to load asynchronously
-	virtual const TSet<FSoftObjectPath>& GetAsyncReferencedAssetClassPaths() const = 0;
+#if WITH_EDITORONLY_DATA
+	UE_DEPRECATED(5.8, "Dynamic loading of referenced assets is deprecated")
+	virtual const TSet<FSoftObjectPath>& GetAsyncReferencedAssetClassPaths() const
+	{
+		static TSet<FSoftObjectPath> Empty;
+		return Empty;
+	}
+#endif // WITH_EDITORONLY_DATA
 
-	// Called when async assets have finished loading.
-	virtual void OnAsyncReferencedAssetsLoaded(const TArray<FMetasoundAssetBase*>& InAsyncReferences) = 0;
+#if WITH_EDITORONLY_DATA
+	UE_DEPRECATED(5.8, "This is deprecated behavior and will become UE_INTERNAL API in a future release for back compat.")
+	virtual void OnAsyncReferencedAssetsLoaded(const TArray<FMetasoundAssetBase*>& InAsyncReferences) { }
+#endif // WITH_EDITORONLY_DATA
 
 	UE_API bool AddingReferenceCausesLoop(const FMetasoundAssetBase& InMetaSound) const;
 
@@ -158,6 +166,7 @@ public:
 	UE_API const Metasound::Frontend::FGraphRegistryKey& GetGraphRegistryKey() const;
 
 #if WITH_EDITORONLY_DATA
+	UE_DEPRECATED(5.8, "Implementation is now internal and moved to VersioningManager.")
 	UE_API bool VersionAsset(FMetaSoundFrontendDocumentBuilder& Builder);
 #endif // WITH_EDITORONLY_DATA
 
@@ -223,16 +232,7 @@ PRAGMA_ENABLE_DEPRECATION_WARNINGS
 private:
 #if WITH_EDITORONLY_DATA
 	UE_API void UpdateAndRegisterReferencesForSerialization(const Metasound::Frontend::FMetaSoundAssetCookOptions& InCookOptions);
-#endif // WITH_EDITORONLY_DATA
-
-	// Checks if version is up-to-date. If so, returns true. If false, updates the interfaces within the given asset's document to the most recent version.
-	UE_API bool TryUpdateInterfaceFromVersion(const FMetasoundFrontendVersion& Version, FMetaSoundFrontendDocumentBuilder& Builder);
 
-	// Returns new interface to be versioned to from the given version. If no interface versioning is
-	// required, returns invalid interface version (no name and invalid version number).
-	FMetasoundFrontendVersion GetTargetInterfaceVersion(const FMetasoundFrontendVersion& InterfaceVersion) const;
-
-#if WITH_EDITORONLY_DATA
 	bool bVersionedOnLoad = false;
 #endif // WITH_EDITORONLY_DATA
 
diff --git a/Engine/Plugins/Runtime/Metasound/Source/MetasoundFrontend/Public/MetasoundAssetManager.h b/Engine/Plugins/Runtime/Metasound/Source/MetasoundFrontend/Public/MetasoundAssetManager.h
--- a/Engine/Plugins/Runtime/Metasound/Source/MetasoundFrontend/Public/MetasoundAssetManager.h
+++ b/Engine/Plugins/Runtime/Metasound/Source/MetasoundFrontend/Public/MetasoundAssetManager.h
@@ -320,8 +320,10 @@ PRAGMA_ENABLE_DEPRECATION_WARNINGS
 		virtual bool ReassignClassName(TScriptInterface<IMetaSoundDocumentInterface> DocInterface) = 0;
 #endif // WITH_EDITOR
 
-		// Requests an async load of all async referenced assets of the input asset.
-		virtual void RequestAsyncLoadReferencedAssets(FMetasoundAssetBase& InAssetBase) = 0;
+#if WITH_EDITORONLY_DATA
+		UE_DEPRECATED(5.8, "This is a deprecated feature, and exists to port old soft object references to hard references.")
+		virtual void RequestAsyncLoadReferencedAssets(FMetasoundAssetBase& InAssetBase, const TSet<FSoftObjectPath>* Paths = nullptr) { }
+#endif // WITH_EDITORONLY_DATA
 
 		// Synchronously requests unregister and re-register of all loaded MetaSound assets node class entries.
 		virtual void ReloadMetaSoundAssets() const = 0;
@@ -345,10 +347,10 @@ PRAGMA_ENABLE_DEPRECATION_WARNINGS
 		// Populates the provided array with MetaSound packages whose document versions were out-of-date and consequently updated.
 		// Returns resulting versioned and unversioned paths.
 		virtual FVersionAssetResults VersionAssetsInFolders(const TArray<FString>& FolderPaths, bool bRecursePaths = true) const = 0;
-#endif // WITH_EDITORONLY_DATA
 
 		// Waits until all async load requests related to this asset are complete.
- 		virtual void WaitUntilAsyncLoadReferencedAssetsComplete(FMetasoundAssetBase& InAssetBase) = 0;
+		virtual void WaitUntilAsyncLoadReferencedAssetsComplete(const FMetasoundAssetBase& InAssetBase) = 0;
+#endif // WITH_EDITORONLY_DATA
 	};
 
 	class FMetaSoundAssetManagerBase : public IMetaSoundAssetManager
diff --git a/Engine/Plugins/Runtime/Metasound/Source/MetasoundFrontend/Public/MetasoundFrontendDocument.h b/Engine/Plugins/Runtime/Metasound/Source/MetasoundFrontend/Public/MetasoundFrontendDocument.h
--- a/Engine/Plugins/Runtime/Metasound/Source/MetasoundFrontend/Public/MetasoundFrontendDocument.h
+++ b/Engine/Plugins/Runtime/Metasound/Source/MetasoundFrontend/Public/MetasoundFrontendDocument.h
@@ -2039,6 +2039,7 @@ struct FMetasoundFrontendDocument
 
 public:
 #if WITH_EDITORONLY_DATA
+	UE_DEPRECATED(5.8, "MaxVersion is now externally accessible via FVersioningManager::GetMaxDocumentVersion()")
 	static UE_API FMetasoundFrontendVersionNumber GetMaxVersion();
 #endif // WITH_EDITORONLY_DATA
 
@@ -2088,27 +2089,19 @@ private:
 
 public:
 #if WITH_EDITORONLY_DATA
+	UE_DEPRECATED(5.8, "Moved to internal implementation. See 'MigrateProperties'")
 	bool RequiresInterfaceVersioning() const
 	{
-		return ArchetypeVersion.IsValid() || !InterfaceVersions.IsEmpty();
+		return false;
 	}
 
-	// Data migration for 5.0 Early Access data. ArchetypeVersion/InterfaceVersions properties can be removed post 5.0 release
-	// and this fix-up can be removed post 5.0 release.
+	UE_DEPRECATED(5.8, "Moved to internal implementation. See 'MigrateProperties'")
 	void VersionInterfaces()
 	{
-		if (ArchetypeVersion.IsValid())
-		{
-			Interfaces.Add(ArchetypeVersion);
-			ArchetypeVersion = FMetasoundFrontendVersion::GetInvalid();
-		}
-
-		if (!InterfaceVersions.IsEmpty())
-		{
-			Interfaces.Append(InterfaceVersions);
-			InterfaceVersions.Reset();
-		}
 	}
+
+	// Migrates data from deprecated properties to new properties
+	UE_INTERNAL bool MigrateProperties();
 #endif // WITH_EDITORONLY_DATA
 };
 
diff --git a/Engine/Plugins/Runtime/Metasound/Source/MetasoundFrontend/Public/MetasoundFrontendDocumentBuilder.h b/Engine/Plugins/Runtime/Metasound/Source/MetasoundFrontend/Public/MetasoundFrontendDocumentBuilder.h
--- a/Engine/Plugins/Runtime/Metasound/Source/MetasoundFrontend/Public/MetasoundFrontendDocumentBuilder.h
+++ b/Engine/Plugins/Runtime/Metasound/Source/MetasoundFrontend/Public/MetasoundFrontendDocumentBuilder.h
@@ -285,6 +285,8 @@ PRAGMA_ENABLE_EXPERIMENTAL_WARNINGS
 	UE_API bool ConfigureDocument(bool bLogFailure = false);
 #endif // WITH_EDITORONLY_DATA
 
+	UE_API bool ConformObjectToDocument();
+
 	UE_API bool ContainsDependencyOfType(EMetasoundFrontendClassType ClassType) const;
 	UE_API bool ContainsEdge(const FMetasoundFrontendEdge& InEdge, const FGuid* InPageID = nullptr) const;
 	UE_API bool ContainsNode(const FGuid& InNodeID, const FGuid* InPageID = nullptr) const;
diff --git a/Engine/Plugins/Runtime/Metasound/Source/MetasoundFrontend/Public/MetasoundFrontendDocumentVersioning.h b/Engine/Plugins/Runtime/Metasound/Source/MetasoundFrontend/Public/MetasoundFrontendDocumentVersioning.h
new file mode 100644
--- /dev/null
+++ b/Engine/Plugins/Runtime/Metasound/Source/MetasoundFrontend/Public/MetasoundFrontendDocumentVersioning.h
@@ -0,0 +1,91 @@
+// Copyright Epic Games, Inc. All Rights Reserved.
+#pragma once
+
+#include "Misc/CoreMiscDefines.h"
+
+#if WITH_EDITORONLY_DATA
+#include <atomic>
+
+#include "MetasoundFrontendDocument.h"
+#include "Misc/CoreMiscDefines.h"
+#include "HAL/CriticalSection.h"
+#include "UObject/GarbageCollection.h"
+#include "UObject/ObjectMacros.h"
+#include "UObject/Object.h"
+#include "UObject/StrongObjectPtrTemplates.h"
+#include "Tasks/Task.h"
+#include "Templates/UniquePtr.h"
+
+
+// Forward Declarations
+class FMetasoundAssetBase;
+class IMetaSoundDocumentInterface;
+struct FMetaSoundFrontendDocumentBuilder;
+
+#define UE_API METASOUNDFRONTEND_API
+
+namespace Metasound::Frontend
+{
+	// Asynchronously manages versioning of provided MetaSounds asset data.
+	// Only required for serialized assets and passed data for managing
+	// immediately upon loading in builds with editor-only data enabled.
+	class FVersioningManager
+	{
+	public:
+		FVersioningManager() = default;
+		~FVersioningManager();
+
+		UE_API static FVersioningManager& Get();
+		UE_API static void Initialize();
+		UE_API static bool IsActive();
+		UE_API static void Shutdown();
+
+		// Current, max document version
+		UE_API static FMetasoundFrontendVersionNumber GetMaxDocumentVersion();
+
+		// Version where page data was migrated from a singleton graph to
+		// multiple paged graphs.  Some legacy behavior requires access to
+		// this version to initialize a builder prior to versioning document
+		// data.
+		UE_INTERNAL static FMetasoundFrontendVersionNumber GetPageMigrationVersion();
+
+		// Feature to check if an ID is changed when Document Metadata is updated.
+		// Is actively being deprecated in favor of the Builder API where transaction
+		// ids are centrally tracked.
+		UE_INTERNAL static bool ChangeIDComparisonEnabledInAutoUpdate();
+
+		// Versions the given MetaSound asynchronously. Ensures provided MetaSound is a serialized asset.
+		UE_API void VersionAssetAsync(UObject& InMetaSound, bool bIsDeterministic);
+
+		// Versions document interfaces on the given builder's object. Primarily exposed
+		// for testing purposes from other modules. Waits for all other active async versioning
+		// requests for safety (is and should not be called from versioning async task).
+		UE_API bool VersionInterfaces(FMetaSoundFrontendDocumentBuilder& DocBuilder) const;
+
+		// Waits until the given MetaSound is finished asynchronously versioning.
+		UE_API void WaitUntilVersioningComplete(const UObject& InMetaSound) const;
+
+	private:
+		// Versions Frontend Document. Called asynchronously while versioning asset.
+		bool VersionDocument(FMetaSoundFrontendDocumentBuilder& DocBuilder) const;
+
+		// Interface versioning implementation (does not guard against thread access
+		// or impose any timing restrictions as the public API call does).
+		bool VersionInterfacesInternal(FMetaSoundFrontendDocumentBuilder& DocBuilder) const;
+
+		using FUniqueId = uint32;
+
+		static TUniquePtr<FVersioningManager> Instance;
+
+		mutable FCriticalSection ActiveVersionCritSection;
+
+		struct FTaskData
+		{
+			TStrongObjectPtr<UObject> AssetPtr;
+			UE::Tasks::FTask Task;
+		};
+		TMap<FUniqueId, FTaskData> ActiveVersionTaskData;
+	};
+} // namespace Metasound::Frontend
+#endif // WITH_EDITORONLY_DATA
+#undef UE_API
