Fix for strong pointers being held in certain automated build contexts too long during shutdown due to missing VersioningManager Shutdown request (and initialization for that matter)
#rnx

#Synced-CL 48425006

diff --git a/Engine/Plugins/Runtime/Metasound/Source/MetasoundFrontend/Private/MetasoundFrontend.cpp b/Engine/Plugins/Runtime/Metasound/Source/MetasoundFrontend/Private/MetasoundFrontend.cpp
--- a/Engine/Plugins/Runtime/Metasound/Source/MetasoundFrontend/Private/MetasoundFrontend.cpp
+++ b/Engine/Plugins/Runtime/Metasound/Source/MetasoundFrontend/Private/MetasoundFrontend.cpp
@@ -8,6 +8,7 @@
 #include "MetasoundAudioBuffer.h"
 #include "MetasoundDataTypeRegistrationMacro.h"
 #include "MetasoundFrontendDocument.h"
+#include "MetasoundFrontendDocumentVersioning.h"
 #include "MetasoundFrontendGraph.h"
 #include "MetasoundFrontendNodeTemplateRegistry.h"
 #include "MetasoundFrontendRegistryContainer.h"
@@ -193,12 +194,20 @@ class FMetasoundFrontendModule : public IModuleInterface
 		RegisterNodeTemplate(MakeUnique<FInputNodeTemplate>());
 		RegisterNodeTemplate(MakeUnique<FRerouteNodeTemplate>());
 		METASOUND_REGISTER_ITEMS_IN_MODULE
+
+#if WITH_EDITORONLY_DATA
+		FVersioningManager::Initialize();
+#endif // WITH_EDITORONLY_DATA
 	}
 
 	virtual void ShutdownModule() override
 	{
 		using namespace Metasound::Frontend;
 
+#if WITH_EDITORONLY_DATA
+		FVersioningManager::Shutdown();
+#endif // WITH_EDITORONLY_DATA
+
 		METASOUND_UNREGISTER_ITEMS_IN_MODULE
 		UnregisterNodeTemplate(FAudioAnalyzerNodeTemplate::ClassName, FAudioAnalyzerNodeTemplate::VersionNumber);
 		UnregisterNodeTemplate(FInputNodeTemplate::ClassName, FInputNodeTemplate::VersionNumber);
