Fix thread safety issues accessing MetaSound while versioning async

#rnx

diff --git a/Engine/Plugins/Runtime/Metasound/Source/MetasoundEditor/Private/MetasoundEditorGraph.cpp b/Engine/Plugins/Runtime/Metasound/Source/MetasoundEditor/Private/MetasoundEditorGraph.cpp
--- a/Engine/Plugins/Runtime/Metasound/Source/MetasoundEditor/Private/MetasoundEditorGraph.cpp
+++ b/Engine/Plugins/Runtime/Metasound/Source/MetasoundEditor/Private/MetasoundEditorGraph.cpp
@@ -2591,11 +2591,7 @@ void UMetasoundEditorGraph::MigrateEditorDocumentData(FMetaSoundFrontendDocument
 				if (UMetasoundEditorGraphMemberDefaultLiteral* DefaultLiteral = Input.GetLiteral())
 				{
 					TSubclassOf<UMetasoundEditorGraphMemberDefaultLiteral> SubClass = DefaultLiteral->GetClass();
-
-					// Migration can occur on async thread, and bind can create a new literal
-					FGCScopeGuard ScopeGuard;
 					EditorSubsystem->BindMemberMetadata(OutBuilder, Input, SubClass, DefaultLiteral);
-					DefaultLiteral->ClearInternalFlags(EInternalObjectFlags::Async);
 				}
 			}
 		}
diff --git a/Engine/Plugins/Runtime/Metasound/Source/MetasoundEngine/Private/Metasound.cpp b/Engine/Plugins/Runtime/Metasound/Source/MetasoundEngine/Private/Metasound.cpp
--- a/Engine/Plugins/Runtime/Metasound/Source/MetasoundEngine/Private/Metasound.cpp
+++ b/Engine/Plugins/Runtime/Metasound/Source/MetasoundEngine/Private/Metasound.cpp
@@ -94,15 +94,18 @@ UMetaSoundPatch::UMetaSoundPatch(const FObjectInitializer& ObjectInitializer)
 	RootMetaSoundDocument.RootGraph.ID = FGuid(0x4d657461, 0x536f756e, 0x64506174, 0x63680000);
 }
 
-Metasound::Frontend::FDocumentAccessPtr UMetaSoundPatch::GetDocumentAccessPtr()
+Metasound::Frontend::FDocumentAccessPtr UMetaSoundPatch::GetDocumentAccessPtr(bool bReloadBuilder)
 {
 	using namespace Metasound::Frontend;
 
 	// Mutation of a document via the soft deprecated access ptr/controller system is not tracked by
 	// the builder registry, so the document cache is invalidated here.
-	if (IDocumentBuilderRegistry* BuilderRegistry = IDocumentBuilderRegistry::Get())
+	if (bReloadBuilder)
 	{
-		BuilderRegistry->ReloadBuilder(RootMetaSoundDocument.RootGraph.Metadata.GetClassName());
+		if (IDocumentBuilderRegistry* BuilderRegistry = IDocumentBuilderRegistry::Get())
+		{
+			BuilderRegistry->ReloadBuilder(RootMetaSoundDocument.RootGraph.Metadata.GetClassName());
+		}
 	}
 
 	// Return document using FAccessPoint to inform the TAccessPtr when the 
diff --git a/Engine/Plugins/Runtime/Metasound/Source/MetasoundEngine/Private/MetasoundSource.cpp b/Engine/Plugins/Runtime/Metasound/Source/MetasoundEngine/Private/MetasoundSource.cpp
--- a/Engine/Plugins/Runtime/Metasound/Source/MetasoundEngine/Private/MetasoundSource.cpp
+++ b/Engine/Plugins/Runtime/Metasound/Source/MetasoundEngine/Private/MetasoundSource.cpp
@@ -991,15 +991,18 @@ float UMetaSoundSource::GetDuration() const
 	return IsOneShot() ? INDEFINITELY_LOOPING_DURATION - 1.0f : INDEFINITELY_LOOPING_DURATION;
 }
 
-Metasound::Frontend::FDocumentAccessPtr UMetaSoundSource::GetDocumentAccessPtr()
+Metasound::Frontend::FDocumentAccessPtr UMetaSoundSource::GetDocumentAccessPtr(bool bReloadBuilder)
 {
 	using namespace Metasound::Frontend;
 
 	// Mutation of a document via the soft deprecated access ptr/controller system is not tracked by
 	// the builder registry, so the document cache is invalidated here.
-	if (IDocumentBuilderRegistry* BuilderRegistry = IDocumentBuilderRegistry::Get())
+	if (bReloadBuilder)
 	{
-		BuilderRegistry->ReloadBuilder(RootMetasoundDocument.RootGraph.Metadata.GetClassName());
+		if (IDocumentBuilderRegistry* BuilderRegistry = IDocumentBuilderRegistry::Get())
+		{
+			BuilderRegistry->ReloadBuilder(RootMetasoundDocument.RootGraph.Metadata.GetClassName());
+		}
 	}
 
 	// Return document using FAccessPoint to inform the TAccessPtr when the 
diff --git a/Engine/Plugins/Runtime/Metasound/Source/MetasoundEngine/Public/Metasound.h b/Engine/Plugins/Runtime/Metasound/Source/MetasoundEngine/Public/Metasound.h
--- a/Engine/Plugins/Runtime/Metasound/Source/MetasoundEngine/Public/Metasound.h
+++ b/Engine/Plugins/Runtime/Metasound/Source/MetasoundEngine/Public/Metasound.h
@@ -202,7 +202,7 @@ protected:
 	UE_DEPRECATED(5.6, "AccessPtrs are actively being deprecated. Writable access outside of the builder API "
 		"is particularly problematic as in so accessing, the builder's caches are reset which can cause major "
 		"editor performance regressions.")
-	UE_API Metasound::Frontend::FDocumentAccessPtr GetDocumentAccessPtr() override;
+	UE_API Metasound::Frontend::FDocumentAccessPtr GetDocumentAccessPtr(bool bReloadBuilder = true) override;
 	UE_API Metasound::Frontend::FConstDocumentAccessPtr GetDocumentConstAccessPtr() const override;
 
 private:
diff --git a/Engine/Plugins/Runtime/Metasound/Source/MetasoundEngine/Public/MetasoundSource.h b/Engine/Plugins/Runtime/Metasound/Source/MetasoundEngine/Public/MetasoundSource.h
--- a/Engine/Plugins/Runtime/Metasound/Source/MetasoundEngine/Public/MetasoundSource.h
+++ b/Engine/Plugins/Runtime/Metasound/Source/MetasoundEngine/Public/MetasoundSource.h
@@ -328,7 +328,7 @@ protected:
 	UE_DEPRECATED(5.6, "AccessPtrs are actively being deprecated. Writable access outside of the builder API "
 		"is particularly problematic as in so accessing, the builder's caches are reset which can cause major "
 		"editor performance regressions.")
-	UE_API Metasound::Frontend::FDocumentAccessPtr GetDocumentAccessPtr() override;
+	UE_API Metasound::Frontend::FDocumentAccessPtr GetDocumentAccessPtr(bool bReloadBuilder = true) override;
 	UE_API Metasound::Frontend::FConstDocumentAccessPtr GetDocumentConstAccessPtr() const override;
 
 	/** Gets all the default parameters for this Asset.  */
diff --git a/Engine/Plugins/Runtime/Metasound/Source/MetasoundFrontend/Private/MetasoundAssetBase.cpp b/Engine/Plugins/Runtime/Metasound/Source/MetasoundFrontend/Private/MetasoundAssetBase.cpp
--- a/Engine/Plugins/Runtime/Metasound/Source/MetasoundFrontend/Private/MetasoundAssetBase.cpp
+++ b/Engine/Plugins/Runtime/Metasound/Source/MetasoundFrontend/Private/MetasoundAssetBase.cpp
@@ -656,10 +656,10 @@ bool FMetasoundAssetBase::MarkMetasoundDocumentDirty() const
 	return false;
 }
 
-Metasound::Frontend::FDocumentHandle FMetasoundAssetBase::GetDocumentHandle()
+Metasound::Frontend::FDocumentHandle FMetasoundAssetBase::GetDocumentHandle(bool bReloadBuilder)
 {
 PRAGMA_DISABLE_DEPRECATION_WARNINGS
-	return Metasound::Frontend::IDocumentController::CreateDocumentHandle(GetDocumentAccessPtr());
+	return Metasound::Frontend::IDocumentController::CreateDocumentHandle(GetDocumentAccessPtr(bReloadBuilder), bReloadBuilder);
 PRAGMA_ENABLE_DEPRECATION_WARNINGS
 }
 
@@ -680,49 +680,11 @@ Metasound::Frontend::FConstGraphHandle FMetasoundAssetBase::GetRootGraphHandle()
 
 bool FMetasoundAssetBase::ImportFromJSON(const FString& InJSON)
 {
-	METASOUND_TRACE_CPUPROFILER_EVENT_SCOPE(MetaSoundAssetBase::ImportFromJSON);
-
-PRAGMA_DISABLE_DEPRECATION_WARNINGS
-	FMetasoundFrontendDocument* Document = GetDocumentAccessPtr().Get();
-PRAGMA_ENABLE_DEPRECATION_WARNINGS
-
-	if (ensure(nullptr != Document))
-	{
-		bool bSuccess = Metasound::Frontend::ImportJSONToMetasound(InJSON, *Document);
-
-		if (bSuccess)
-		{
-			UObject* OwningAsset = GetOwningAsset();
-			check(OwningAsset);
-			ensure(OwningAsset->MarkPackageDirty());
-		}
-
-		return bSuccess;
-	}
 	return false;
 }
 
 bool FMetasoundAssetBase::ImportFromJSONAsset(const FString& InAbsolutePath)
 {
-	METASOUND_TRACE_CPUPROFILER_EVENT_SCOPE(MetaSoundAssetBase::ImportFromJSONAsset);
-
-PRAGMA_DISABLE_DEPRECATION_WARNINGS
-	Metasound::Frontend::FDocumentAccessPtr DocumentPtr = GetDocumentAccessPtr();
-PRAGMA_ENABLE_DEPRECATION_WARNINGS
-
-	if (FMetasoundFrontendDocument* Document = DocumentPtr.Get())
-	{
-		bool bSuccess = Metasound::Frontend::ImportJSONAssetToMetasound(InAbsolutePath, *Document);
-
-		if (bSuccess)
-		{
-			UObject* OwningAsset = GetOwningAsset();
-			check(OwningAsset);
-			ensure(OwningAsset->MarkPackageDirty());
-		}
-
-		return bSuccess;
-	}
 	return false;
 }
 
diff --git a/Engine/Plugins/Runtime/Metasound/Source/MetasoundFrontend/Private/MetasoundFrontendController.cpp b/Engine/Plugins/Runtime/Metasound/Source/MetasoundFrontend/Private/MetasoundFrontendController.cpp
--- a/Engine/Plugins/Runtime/Metasound/Source/MetasoundFrontend/Private/MetasoundFrontendController.cpp
+++ b/Engine/Plugins/Runtime/Metasound/Source/MetasoundFrontend/Private/MetasoundFrontendController.cpp
@@ -59,10 +59,10 @@ namespace Metasound
 			return Invalid;
 		}
 			
-		FDocumentHandle IDocumentController::CreateDocumentHandle(FDocumentAccessPtr InDocument)
+		FDocumentHandle IDocumentController::CreateDocumentHandle(FDocumentAccessPtr InDocument, bool bReloadBuilder)
 		{
 			// Create using standard document controller.
-			return FDocumentController::CreateDocumentHandle(InDocument);
+			return FDocumentController::CreateDocumentHandle(InDocument, bReloadBuilder);
 		}
 
 		FDocumentHandle IDocumentController::CreateDocumentHandle(FMetasoundFrontendDocument& InDocument)
diff --git a/Engine/Plugins/Runtime/Metasound/Source/MetasoundFrontend/Private/MetasoundFrontendDocumentController.h b/Engine/Plugins/Runtime/Metasound/Source/MetasoundFrontend/Private/MetasoundFrontendDocumentController.h
--- a/Engine/Plugins/Runtime/Metasound/Source/MetasoundFrontend/Private/MetasoundFrontendDocumentController.h
+++ b/Engine/Plugins/Runtime/Metasound/Source/MetasoundFrontend/Private/MetasoundFrontendDocumentController.h
@@ -22,15 +22,18 @@ namespace Metasound
 			 */
 			FDocumentController(FDocumentAccessPtr InDocumentPtr);
 
-			static FDocumentHandle CreateDocumentHandle(FDocumentAccessPtr InDocument)
+			static FDocumentHandle CreateDocumentHandle(FDocumentAccessPtr InDocument, bool bReloadBuilder = true)
 			{
 				// Unit test builds may not load the builder registry (i.e. via the engine module).  Creating and
 				// manipulating documents via controllers/handles must be supported for backward compat in this context,
 				// so registry is not required to exist.
-				if (IDocumentBuilderRegistry* Registry = IDocumentBuilderRegistry::Get())
+				if (bReloadBuilder)
 				{
-					const FMetasoundFrontendClassName& ClassName = InDocument.Get()->RootGraph.Metadata.GetClassName();
-					Registry->ReloadBuilder(ClassName);
+					if (IDocumentBuilderRegistry* Registry = IDocumentBuilderRegistry::Get())
+					{
+						const FMetasoundFrontendClassName& ClassName = InDocument.Get()->RootGraph.Metadata.GetClassName();
+						Registry->ReloadBuilder(ClassName);
+					}
 				}
 				return MakeShared<FDocumentController>(InDocument);
 			}
diff --git a/Engine/Plugins/Runtime/Metasound/Source/MetasoundFrontend/Private/MetasoundFrontendDocumentVersioning.cpp b/Engine/Plugins/Runtime/Metasound/Source/MetasoundFrontend/Private/MetasoundFrontendDocumentVersioning.cpp
--- a/Engine/Plugins/Runtime/Metasound/Source/MetasoundFrontend/Private/MetasoundFrontendDocumentVersioning.cpp
+++ b/Engine/Plugins/Runtime/Metasound/Source/MetasoundFrontend/Private/MetasoundFrontendDocumentVersioning.cpp
@@ -861,11 +861,10 @@ PRAGMA_ENABLE_DEPRECATION_WARNINGS
 					}
 				}
 
-				// Safeguards against prior fix-up corrupting any cached data
-				if (IDocumentBuilderRegistry* BuilderRegistry = IDocumentBuilderRegistry::Get())
-				{
-					BuilderRegistry->ReloadBuilder(Document.RootGraph.Metadata.GetClassName());
-				}
+				// Safeguards against prior fix-up corrupting any cached data.
+				// Safe to do without providing delegates as versioning is done
+				// with unregistered builder and prior to delegate registration.
+				OutBuilder.Reload();
 			}
 		};
 
@@ -926,11 +925,10 @@ PRAGMA_ENABLE_DEPRECATION_WARNINGS
 				Algo::Sort(Document.Dependencies, CompareForSort);
 				Document.Dependencies.SetNum(Algo::Unique(Document.Dependencies, CompareForUnique));
 
-				// Safeguards against prior fix-up corrupting any cached data
-				if (IDocumentBuilderRegistry* BuilderRegistry = IDocumentBuilderRegistry::Get())
-				{
-					BuilderRegistry->ReloadBuilder(Document.RootGraph.Metadata.GetClassName());
-				}
+				// Safeguards against prior fix-up corrupting any cached data.
+				// Safe to do without providing delegates as versioning is done
+				// with unregistered builder and prior to delegate registration.
+				OutBuilder.Reload();
 			}
 		};
 
@@ -1045,8 +1043,7 @@ PRAGMA_ENABLE_DEPRECATION_WARNINGS
 		const FUniqueId AssetId = InMetaSound.GetUniqueID();
 		if (!ActiveVersionTaskData.Contains(AssetId))
 		{
-			InMetaSound.SetInternalFlags(EInternalObjectFlags::Async);
-			TStrongObjectPtr<UObject> AssetPtr(&InMetaSound);
+			TWeakObjectPtr<UObject> AssetPtr(&InMetaSound);
 
 			check(DocInterface.GetObject());
 			const FMetasoundFrontendDocument& Document = DocInterface->GetConstDocument();
@@ -1064,19 +1061,23 @@ PRAGMA_ENABLE_DEPRECATION_WARNINGS
 				{
 					FScopeLock VersionLock(&ActiveVersionCritSection);
 					FTaskData& TaskData = ActiveVersionTaskData.FindChecked(AssetId);
-					FDocumentIDGenerator::FScopeDeterminism DeterminismScope(bIsDeterministic);
-					FMetaSoundFrontendDocumentBuilder DocBuilder(TaskData.AssetPtr.Get());
-
-					const bool bVersionedDoc = VersionDocument(DocBuilder);
-					const bool bVersionedInterfaces = VersionInterfacesInternal(DocBuilder);
 
-					if (bVersionedDoc || bVersionedInterfaces)
 					{
-						DocBuilder.GetMetasoundAsset().SetVersionedOnLoad();
+						FGCScopeGuard ScopeGuard;
+						FDocumentIDGenerator::FScopeDeterminism DeterminismScope(bIsDeterministic);
+						if (UObject* MetaSoundObject = TaskData.AssetPtr.Get())
+						{
+							FMetaSoundFrontendDocumentBuilder DocBuilder(MetaSoundObject);
+							const bool bVersionedDoc = VersionDocument(DocBuilder);
+							const bool bVersionedInterfaces = VersionInterfacesInternal(DocBuilder);
+
+							if (bVersionedDoc || bVersionedInterfaces)
+							{
+								DocBuilder.GetMetasoundAsset().SetVersionedOnLoad();
+							}
+						}
 					}
 
-					TaskData.AssetPtr->ClearInternalFlags(EInternalObjectFlags::Async);
-					TaskData.AssetPtr.Reset();
 					ActiveVersionTaskData.Remove(AssetId);
 				})
 			});
@@ -1153,7 +1154,10 @@ PRAGMA_ENABLE_DEPRECATION_WARNINGS
 				// compat prior to all controller versioning, so just do it here.
 				FMigratePagePropertiesTransform().Transform(DocBuilder);
 
-				FDocumentHandle DocHandle = DocBuilder.GetMetasoundAsset().GetDocumentHandle();
+				// Builder is not registered during versioning, so don't attempt to reload builder
+				// (which can cause delegates to be fired in invalid scenarios like reloading assets
+				// from disk in editor).
+				FDocumentHandle DocHandle = DocBuilder.GetMetasoundAsset().GetDocumentHandle(false);
 
 				bWasUpdated |= FVersionDocument_1_1(Name, Path).Transform(DocHandle);
 				bWasUpdated |= FVersionDocument_1_2(Name, Path).Transform(DocHandle);
diff --git a/Engine/Plugins/Runtime/Metasound/Source/MetasoundFrontend/Public/MetasoundAssetBase.h b/Engine/Plugins/Runtime/Metasound/Source/MetasoundFrontend/Public/MetasoundAssetBase.h
--- a/Engine/Plugins/Runtime/Metasound/Source/MetasoundFrontend/Public/MetasoundAssetBase.h
+++ b/Engine/Plugins/Runtime/Metasound/Source/MetasoundFrontend/Public/MetasoundAssetBase.h
@@ -143,14 +143,14 @@ public:
 
 	UE_API bool IsRegistered() const;
 
-	// Imports data from a JSON string directly
+	UE_DEPRECATED(5.8, "Feature no longer supported")
 	UE_API bool ImportFromJSON(const FString& InJSON);
 
-	// Imports the asset from a JSON file at provided path
+	UE_DEPRECATED(5.8, "Feature no longer supported")
 	UE_API bool ImportFromJSONAsset(const FString& InAbsolutePath);
 
 	// Soft Deprecated in favor of DocumentBuilder API. Returns handle for the root metasound graph of this asset.
-	UE_API Metasound::Frontend::FDocumentHandle GetDocumentHandle();
+	UE_API Metasound::Frontend::FDocumentHandle GetDocumentHandle(bool bReloadBuilder = true);
 	UE_API Metasound::Frontend::FConstDocumentHandle GetDocumentHandle() const;
 
 	// Soft Deprecated in favor of DocumentBuilder API. Returns handle for the root metasound graph of this asset.
@@ -221,7 +221,7 @@ PRAGMA_ENABLE_DEPRECATION_WARNINGS
 	UE_DEPRECATED(5.6, "AccessPtrs are actively being deprecated. Writable access outside of the builder API "
 		" is particularly problematic as in so accessing, the builder's caches are reset which can cause major "
 		"editor performance regressions.")
-	virtual Metasound::Frontend::FDocumentAccessPtr GetDocumentAccessPtr() = 0;
+	virtual Metasound::Frontend::FDocumentAccessPtr GetDocumentAccessPtr(bool bReloadBuilder = true) = 0;
 
 	// Returns an access pointer to the document.
 	virtual Metasound::Frontend::FConstDocumentAccessPtr GetDocumentConstAccessPtr() const = 0;
diff --git a/Engine/Plugins/Runtime/Metasound/Source/MetasoundFrontend/Public/MetasoundFrontendController.h b/Engine/Plugins/Runtime/Metasound/Source/MetasoundFrontend/Public/MetasoundFrontendController.h
--- a/Engine/Plugins/Runtime/Metasound/Source/MetasoundFrontend/Public/MetasoundFrontendController.h
+++ b/Engine/Plugins/Runtime/Metasound/Source/MetasoundFrontend/Public/MetasoundFrontendController.h
@@ -991,7 +991,7 @@ namespace Metasound
 			static UE_API FDocumentHandle GetInvalidHandle();
 
 			/** Create a document from FMetasoundFrontendDocument description pointer. */
-			static UE_API FDocumentHandle CreateDocumentHandle(FDocumentAccessPtr InDocument);
+			static UE_API FDocumentHandle CreateDocumentHandle(FDocumentAccessPtr InDocument, bool bReloadBuilder = true);
 			static UE_API FDocumentHandle CreateDocumentHandle(FMetasoundFrontendDocument& InDocument);
 
 			/** Create a document from FMetasoundFrontendDocument description pointer. */
diff --git a/Engine/Plugins/Runtime/Metasound/Source/MetasoundFrontend/Public/MetasoundFrontendDocumentVersioning.h b/Engine/Plugins/Runtime/Metasound/Source/MetasoundFrontend/Public/MetasoundFrontendDocumentVersioning.h
--- a/Engine/Plugins/Runtime/Metasound/Source/MetasoundFrontend/Public/MetasoundFrontendDocumentVersioning.h
+++ b/Engine/Plugins/Runtime/Metasound/Source/MetasoundFrontend/Public/MetasoundFrontendDocumentVersioning.h
@@ -81,7 +81,7 @@ namespace Metasound::Frontend
 
 		struct FTaskData
 		{
-			TStrongObjectPtr<UObject> AssetPtr;
+			TWeakObjectPtr<UObject> AssetPtr;
 			UE::Tasks::FTask Task;
 		};
 		TMap<FUniqueId, FTaskData> ActiveVersionTaskData;
