- Fix for async flag not clearing when versioning MetaSounds (at least one offender when upgrading to v1.12 due to member metadata generation)
- Reduce spam from FMetasoundCommentNodeIntVector mismatched tag serialization to VeryVerbose log level

#rnx
#Synced-CL 48462405

diff --git a/Engine/Plugins/Runtime/Metasound/Source/MetasoundFrontend/Private/MetasoundFrontendDocument.cpp b/Engine/Plugins/Runtime/Metasound/Source/MetasoundFrontend/Private/MetasoundFrontendDocument.cpp
--- a/Engine/Plugins/Runtime/Metasound/Source/MetasoundFrontend/Private/MetasoundFrontendDocument.cpp
+++ b/Engine/Plugins/Runtime/Metasound/Source/MetasoundFrontend/Private/MetasoundFrontendDocument.cpp
@@ -382,7 +382,7 @@ bool FMetasoundCommentNodeIntVector::SerializeFromMismatchedTag(const struct FPr
 			X = static_cast<int32>(OldVector.X);
 			Y = static_cast<int32>(OldVector.Y);
 
-			UE_LOG(LogMetaSound, Display, TEXT("FMetasoundCommentNodeIntVector::SerializeFromMismatchedTag - DeprecateSlateVector2D Type found: Resolving Mismatch"));
+			UE_LOG(LogMetaSound, VeryVerbose, TEXT("FMetasoundCommentNodeIntVector::SerializeFromMismatchedTag - DeprecateSlateVector2D Type found: Resolving Mismatch"));
 			return true;
 		}
 	}
@@ -396,7 +396,7 @@ bool FMetasoundCommentNodeIntVector::SerializeFromMismatchedTag(const struct FPr
 		X = 0;
 		Y = 0;
 
-		UE_LOG(LogMetaSound, Display, TEXT("FMetasoundCommentNodeIntVector::SerializeFromMismatchedTag - DeprecateSlateVector2D Type not loaded: Ignoring Mismatch"));
+		UE_LOG(LogMetaSound, VeryVerbose, TEXT("FMetasoundCommentNodeIntVector::SerializeFromMismatchedTag - DeprecateSlateVector2D Type not loaded: Ignoring Mismatch"));
 		return bIsCookCommandlet;
 	}
 
@@ -441,11 +441,11 @@ bool FMetasoundCommentNodeIntVector::SerializeFromMismatchedTag(const struct FPr
 	// so don't bother reporting if old data was not translated.
 	if (bIsCookCommandlet)
 	{
-		UE_LOG(LogMetaSound, Display, TEXT("FMetasoundCommentNodeIntVector::SerializeFromMismatchedTag - Did not resolve. Ignoring value (cooking content, value not necessary)."));
+		UE_LOG(LogMetaSound, VeryVerbose, TEXT("FMetasoundCommentNodeIntVector::SerializeFromMismatchedTag - Did not resolve. Ignoring value (cooking content, value not necessary)."));
 	}
 	else
 	{
-		UE_LOG(LogMetaSound, Display, TEXT("FMetasoundCommentNodeIntVector::SerializeFromMismatchedTag - Did not resolve. Returning failure state."));
+		UE_LOG(LogMetaSound, VeryVerbose, TEXT("FMetasoundCommentNodeIntVector::SerializeFromMismatchedTag - Did not resolve. Returning failure state."));
 	}
 
 	return bIsCookCommandlet;
diff --git a/Engine/Plugins/Runtime/Metasound/Source/MetasoundFrontend/Private/MetasoundFrontendDocumentVersioning.cpp b/Engine/Plugins/Runtime/Metasound/Source/MetasoundFrontend/Private/MetasoundFrontendDocumentVersioning.cpp
--- a/Engine/Plugins/Runtime/Metasound/Source/MetasoundFrontend/Private/MetasoundFrontendDocumentVersioning.cpp
+++ b/Engine/Plugins/Runtime/Metasound/Source/MetasoundFrontend/Private/MetasoundFrontendDocumentVersioning.cpp
@@ -24,6 +24,7 @@
 #include "MetasoundLog.h"
 #include "MetasoundTrace.h"
 #include "Misc/App.h"
+#include "UObject/UObjectHash.h"
 
 #define METASOUND_VERSIONING_LOG(Verbosity, Format, ...) if (DocumentTransform::GetVersioningLoggingEnabled()) { UE_LOG(LogMetaSound, Verbosity, Format, ##__VA_ARGS__); }
 
@@ -119,7 +120,7 @@ namespace Metasound::Frontend
 #else
 				const FString AssetName = *Metadata.GetClassName().ToString();
 #endif // !WITH_EDITOR
-				UE_LOG(LogMetaSound, Display, TEXT("Asset '%s' interface '%s' updated: '%s' --> '%s'"),
+				METASOUND_VERSIONING_LOG(Display, TEXT("Asset '%s' interface '%s' updated: '%s' --> '%s'"),
 					*AssetName,
 					*InterfaceVersion.Name.ToString(),
 					*InterfaceVersion.Number.ToString(),
@@ -1075,6 +1076,15 @@ PRAGMA_ENABLE_DEPRECATION_WARNINGS
 							{
 								DocBuilder.GetMetasoundAsset().SetVersionedOnLoad();
 							}
+
+							ForEachObjectWithOuter(MetaSoundObject, [](UObject* Object)
+							{
+								if (Object->HasAnyInternalFlags(EInternalObjectFlags::Async))
+								{
+									METASOUND_VERSIONING_LOG(Display, TEXT("Clearing object '%s' async flag (object likely generated while MetaSound versioning)"), *Object->GetName());
+									Object->ClearInternalFlags(EInternalObjectFlags::Async);
+								}
+							});
 						}
 					}
 
