Minor feedback & reduce log spam from async flag fix

#rnx
#Synced-CL 48462405

diff --git a/Engine/Plugins/Runtime/Metasound/Source/MetasoundFrontend/Private/MetasoundFrontendDocumentVersioning.cpp b/Engine/Plugins/Runtime/Metasound/Source/MetasoundFrontend/Private/MetasoundFrontendDocumentVersioning.cpp
--- a/Engine/Plugins/Runtime/Metasound/Source/MetasoundFrontend/Private/MetasoundFrontendDocumentVersioning.cpp
+++ b/Engine/Plugins/Runtime/Metasound/Source/MetasoundFrontend/Private/MetasoundFrontendDocumentVersioning.cpp
@@ -34,6 +34,7 @@ namespace Metasound::Frontend
 	namespace VersioningPrivate
 	{
 		static int32 MetaSoundAutoUpdateUseAssetChangeIDsCVar = 0;
+		static TUniquePtr<FVersioningManager> VersioningManager;
 
 		FAutoConsoleVariableRef CVarMetaSoundAutoUpdateUseAssetChangeIDs(
 			TEXT("au.MetaSound.AutoUpdate.UseChangeIDs"),
@@ -975,8 +976,6 @@ PRAGMA_ENABLE_DEPRECATION_WARNINGS
 		}
 	} // namespace VersioningPrivate
 
-	TUniquePtr<FVersioningManager> FVersioningManager::Instance;
-
 	FVersioningManager::~FVersioningManager()
 	{
 		using namespace UE::Tasks;
@@ -989,28 +988,32 @@ PRAGMA_ENABLE_DEPRECATION_WARNINGS
 
 	FVersioningManager& FVersioningManager::Get()
 	{
+		using namespace VersioningPrivate;
+
 		// Its possible some commandlets load MetaSounds prior
 		// to the MetaSoundEngine module loading (ex. generating
 		// asset manifests), so attempt to generate one on the fly
 		// when accessing here.
-		if (!Instance.IsValid())
+		if (!VersioningManager.IsValid())
 		{
-			Instance = MakeUnique<FVersioningManager>();
+			VersioningManager = MakeUnique<FVersioningManager>();
 		}
-		return *Instance;
+		return *VersioningManager;
 	}
 
 	void FVersioningManager::Initialize()
 	{
-		if (!Instance.IsValid())
+		using namespace VersioningPrivate;
+		if (!VersioningManager.IsValid())
 		{
-			Instance = MakeUnique<FVersioningManager>();
+			VersioningManager = MakeUnique<FVersioningManager>();
 		}
 	}
 
 	void FVersioningManager::Shutdown()
 	{
-		Instance.Reset();
+		using namespace VersioningPrivate;
+		VersioningManager.Reset();
 	}
 
 	bool FVersioningManager::ChangeIDComparisonEnabledInAutoUpdate()
@@ -1081,7 +1084,7 @@ PRAGMA_ENABLE_DEPRECATION_WARNINGS
 							{
 								if (Object->HasAnyInternalFlags(EInternalObjectFlags::Async))
 								{
-									METASOUND_VERSIONING_LOG(Display, TEXT("Clearing object '%s' async flag (object likely generated while MetaSound versioning)"), *Object->GetName());
+									METASOUND_VERSIONING_LOG(VeryVerbose, TEXT("Clearing object '%s' async flag (object likely generated while MetaSound versioning)"), *Object->GetName());
 									Object->ClearInternalFlags(EInternalObjectFlags::Async);
 								}
 							});
diff --git a/Engine/Plugins/Runtime/Metasound/Source/MetasoundFrontend/Public/MetasoundFrontendDocumentVersioning.h b/Engine/Plugins/Runtime/Metasound/Source/MetasoundFrontend/Public/MetasoundFrontendDocumentVersioning.h
--- a/Engine/Plugins/Runtime/Metasound/Source/MetasoundFrontend/Public/MetasoundFrontendDocumentVersioning.h
+++ b/Engine/Plugins/Runtime/Metasound/Source/MetasoundFrontend/Public/MetasoundFrontendDocumentVersioning.h
@@ -75,8 +75,6 @@ namespace Metasound::Frontend
 
 		using FUniqueId = uint32;
 
-		static TUniquePtr<FVersioningManager> Instance;
-
 		mutable FCriticalSection ActiveVersionCritSection;
 
 		struct FTaskData
