- Fix for reloading MetaSound assets when they're already loaded, causing builders to become invalidated when not desired.
- Minor optimization to AutoUpdate to just cast builder's object instead of attempting to load via AssetManager only to return local MetaSound object

#redirectsaresafe
#tests -PIE -dpcvars=s.AllowMultithreadedLoading=1 -asyncloadingthread
#rnx
#Synced-CL 48675545

diff --git a/Engine/Plugins/Runtime/Metasound/Source/MetasoundEngine/Private/MetasoundAssetSubsystem.cpp b/Engine/Plugins/Runtime/Metasound/Source/MetasoundEngine/Private/MetasoundAssetSubsystem.cpp
--- a/Engine/Plugins/Runtime/Metasound/Source/MetasoundEngine/Private/MetasoundAssetSubsystem.cpp
+++ b/Engine/Plugins/Runtime/Metasound/Source/MetasoundEngine/Private/MetasoundAssetSubsystem.cpp
@@ -1231,7 +1231,12 @@ PRAGMA_ENABLE_DEPRECATION_WARNINGS
 
 	FMetasoundAssetBase* FMetaSoundAssetManager::TryLoadAsset(const FSoftObjectPath& InObjectPath) const
 	{
-		return Metasound::IMetasoundUObjectRegistry::Get().GetObjectAsAssetBase(InObjectPath.TryLoad());
+		UObject* MetaSoundObject = InObjectPath.ResolveObject();
+		if (!MetaSoundObject)
+		{
+			MetaSoundObject = InObjectPath.TryLoad();
+		}
+		return Metasound::IMetasoundUObjectRegistry::Get().GetObjectAsAssetBase(MetaSoundObject);
 	}
 
 	bool FMetaSoundAssetManager::TryLoadReferencedAssets(const FMetasoundAssetBase& InAssetBase, TArray<FMetasoundAssetBase*>& OutReferencedAssets) const
diff --git a/Engine/Plugins/Runtime/Metasound/Source/MetasoundFrontend/Private/MetasoundFrontendTransform.cpp b/Engine/Plugins/Runtime/Metasound/Source/MetasoundFrontend/Private/MetasoundFrontendTransform.cpp
--- a/Engine/Plugins/Runtime/Metasound/Source/MetasoundFrontend/Private/MetasoundFrontendTransform.cpp
+++ b/Engine/Plugins/Runtime/Metasound/Source/MetasoundFrontend/Private/MetasoundFrontendTransform.cpp
@@ -539,16 +539,8 @@ namespace Metasound
 				bDidEdit |= FRebuildPresetRootGraph(ReferenceBuilder).Transform(InOutBuilderToTransform);
 				if (bDidEdit)
 				{
-					FMetasoundFrontendClassMetadata PresetMetadata = InOutBuilderToTransform.GetConstDocumentChecked().RootGraph.Metadata;
-					PresetMetadata.SetType(EMetasoundFrontendClassType::External);
-					const FNodeRegistryKey RegistryKey(PresetMetadata);
-					FMetasoundAssetBase* PresetMetaSoundAsset = IMetaSoundAssetManager::GetChecked().TryLoadAssetFromKey(RegistryKey);
-					if (ensure(PresetMetaSoundAsset))
-					{
-						TScriptInterface<IMetaSoundDocumentInterface> PresetInterface = PresetMetaSoundAsset->GetOwningAsset();
-						check(PresetInterface);
-						PresetInterface->ConformObjectToDocument();
-					}
+					TScriptInterface<IMetaSoundDocumentInterface> PresetInterface(&InOutBuilderToTransform.CastDocumentObjectChecked<UObject>());
+					PresetInterface->ConformObjectToDocument();
 
 					InOutBuilderToTransform.RemoveUnusedDependencies();
 					InOutBuilderToTransform.SynchronizeDependencyMetadata();
