MetaSound Document Configuration: Vertical Checkpoint 2
- Config execution code gen
- UX improvements around selection
- Preset adoption of DocConfig API
- Document Version for Preset --> DocConfig
- Auto-Asset Action generation for configs adjusted to appear on appropriate selection

[RN] MetaSounds now support document configurations, a customizable and automated way to rapidly author MetaSounds through a simplified set of code-defined data that procedurally generates the graph.

#Synced-CL 48741549

diff --git a/Engine/Plugins/Experimental/TechAudioTools/Source/TechAudioToolsMetaSound/Private/ViewModels/MetaSoundViewModel.cpp b/Engine/Plugins/Experimental/TechAudioTools/Source/TechAudioToolsMetaSound/Private/ViewModels/MetaSoundViewModel.cpp
--- a/Engine/Plugins/Experimental/TechAudioTools/Source/TechAudioToolsMetaSound/Private/ViewModels/MetaSoundViewModel.cpp
+++ b/Engine/Plugins/Experimental/TechAudioTools/Source/TechAudioToolsMetaSound/Private/ViewModels/MetaSoundViewModel.cpp
@@ -2,7 +2,7 @@
 
 #include "ViewModels/MetaSoundViewModel.h"
 
-#include "DocumentConfigurations/MetasoundPresetDocumentConfiguration.h"
+#include "DocumentConfigurations/MetasoundFrontendPresetDocumentConfiguration.h"
 #include "Logging/StructuredLog.h"
 #include "MetasoundDocumentBuilderRegistry.h"
 #include "MetasoundFrontendPages.h"
@@ -50,6 +50,8 @@ void UMetaSoundViewModel::InitializeMetaSound(const TScriptInterface<IMetaSoundD
 
 void UMetaSoundViewModel::Initialize(UMetaSoundBuilderBase* InBuilder)
 {
+	using namespace Metasound::Frontend;
+
 	Reset();
 
 	if (!InBuilder)
@@ -61,9 +63,7 @@ void UMetaSoundViewModel::Initialize(UMetaSoundBuilderBase* InBuilder)
 	Builder = InBuilder;
 
 	const FMetaSoundFrontendDocumentBuilder& DocumentBuilder = InBuilder->GetConstBuilder();
-	const FMetasoundFrontendDocument& FrontendDocument = DocumentBuilder.GetConstDocumentChecked();
-
-	UE_MVVM_SET_PROPERTY_VALUE(bIsPreset, FrontendDocument.RootGraph.PresetOptions.bIsPreset);
+	UE_MVVM_SET_PROPERTY_VALUE(bIsPreset, DocumentBuilder.IsDocumentConfiguredBy<FPresetDocumentConfiguration>());
 
 	CreateMemberViewModels();
 	SetIsInitialized(true);
@@ -139,11 +139,19 @@ void UMetaSoundViewModel::SetInputDefaultLiteral(const FName& InputName, const F
 
 void UMetaSoundViewModel::SetInputOverridesDefault(const FName& InputName, const bool bOverridesDefault) const
 {
+	using namespace Metasound::Frontend;
+
 	if (Builder)
 	{
-		const bool bInheritsDefaultValue = !bOverridesDefault;
 		FMetaSoundFrontendDocumentBuilder& DocumentBuilder = Builder->GetBuilder();
-		DocumentBuilder.SetGraphInputInheritsDefault(InputName, bInheritsDefaultValue);
+
+		if (const FMetasoundFrontendClassInput* Input = DocumentBuilder.FindGraphInput(InputName))
+		{
+			if (FPresetVertexMetadata* PresetMetadata = DocumentBuilder.FindVertexMetadata<FPresetVertexMetadata>(Input->NodeID))
+			{
+				PresetMetadata->bOverrideInheritedDefault = bOverridesDefault;
+			}
+		}
 	}
 }
 
@@ -384,6 +392,8 @@ void UMetaSoundViewModel::CreateMemberViewModels()
 
 UMetaSoundInputViewModel* UMetaSoundViewModel::CreateInputViewModel(const FMetasoundFrontendClassInput& InInput)
 {
+	using namespace Metasound::Frontend;
+
 	if (UMetaSoundInputViewModel* InputViewModel = NewObject<UMetaSoundInputViewModel>(this, GetInputViewModelClass()))
 	{
 		FName PageName;
@@ -405,8 +415,11 @@ UMetaSoundInputViewModel* UMetaSoundViewModel::CreateInputViewModel(const FMetas
 		if (Builder)
 		{
 			const FMetaSoundFrontendDocumentBuilder& DocumentBuilder = Builder->GetConstBuilder();
-			const FMetasoundFrontendDocument& FrontendDocument = DocumentBuilder.GetConstDocumentChecked();
-			bInheritsDefault = FrontendDocument.RootGraph.PresetOptions.InputsInheritingDefault.Contains(InInput.Name);
+			bInheritsDefault = false;
+			if (const FPresetVertexMetadata* VertexMetadata = DocumentBuilder.FindConstVertexMetadata<FPresetVertexMetadata>(InInput.NodeID))
+			{
+				bInheritsDefault = !VertexMetadata->bOverrideInheritedDefault;
+			}
 		}
 
 		InputViewModel->OnInputNameChanged(InInput.Name);
diff --git a/Engine/Plugins/Runtime/Metasound/Source/MetasoundEditor/Private/FindInMetasound.cpp b/Engine/Plugins/Runtime/Metasound/Source/MetasoundEditor/Private/FindInMetasound.cpp
--- a/Engine/Plugins/Runtime/Metasound/Source/MetasoundEditor/Private/FindInMetasound.cpp
+++ b/Engine/Plugins/Runtime/Metasound/Source/MetasoundEditor/Private/FindInMetasound.cpp
@@ -1,6 +1,8 @@
 // Copyright Epic Games, Inc. All Rights Reserved.
 
 #include "FindInMetasound.h"
+
+#include "DocumentConfigurations/MetasoundFrontendPresetDocumentConfiguration.h"
 #include "Framework/Application/SlateApplication.h"
 #include "Layout/WidgetPath.h"
 #include "MetasoundEditorGraph.h"
@@ -305,6 +307,8 @@ namespace Metasound::Editor
 
 	FText FFindInMetasoundResult::GetValueText()
 	{
+		using namespace Frontend;
+
 		if (!ValueText.IsEmpty())
 		{
 			return ValueText;
@@ -335,14 +339,9 @@ namespace Metasound::Editor
 				// Check if input is an overridden default value (for presets) 
 				if (const UMetasoundEditorGraphInput* Input = Cast<UMetasoundEditorGraphInput>(Member))
 				{
-					if (DocBuilder.IsPreset())
+					if (const FPresetVertexMetadata* VertexMetadata = DocBuilder.FindConstVertexMetadata<FPresetVertexMetadata>(Input->NodeID))
 					{
-						const TSet<FName>* InputsInheritingDefault = DocBuilder.GetGraphInputsInheritingDefault();
-						FName NodeName = Input->GetMemberName();
-						if (InputsInheritingDefault && !InputsInheritingDefault->Contains(NodeName))
-						{
-							bIsOverriddenLiteral = true;
-						}
+						bIsOverriddenLiteral = VertexMetadata->bOverrideInheritedDefault;
 					}
 				}
 
diff --git a/Engine/Plugins/Runtime/Metasound/Source/MetasoundEditor/Private/MetasoundAssetDefinitions.cpp b/Engine/Plugins/Runtime/Metasound/Source/MetasoundEditor/Private/MetasoundAssetDefinitions.cpp
--- a/Engine/Plugins/Runtime/Metasound/Source/MetasoundEditor/Private/MetasoundAssetDefinitions.cpp
+++ b/Engine/Plugins/Runtime/Metasound/Source/MetasoundEditor/Private/MetasoundAssetDefinitions.cpp
@@ -8,6 +8,7 @@
 #include "Components/AudioComponent.h"
 #include "ContentBrowserMenuContexts.h"
 #include "ContentBrowserModule.h"
+#include "DocumentConfigurations/MetasoundFrontendPresetDocumentConfiguration.h"
 #include "Editor/EditorPerProjectUserSettings.h"
 #include "FileHelpers.h"
 #include "IAssetTools.h"
@@ -25,11 +26,13 @@
 #include "MetasoundEditorSettings.h"
 #include "MetasoundEditorSubsystem.h"
 #include "MetasoundFactory.h"
+#include "MetasoundFrontendDocumentBuilder.h"
 #include "MetasoundSource.h"
 #include "MetasoundUObjectRegistry.h"
 #include "Misc/AssetRegistryInterface.h"
 #include "Modules/ModuleManager.h"
 #include "ObjectEditorUtils.h"
+#include "Sound/SoundWave.h"
 #include "Styling/AppStyle.h"
 #include "Styling/CoreStyle.h"
 #include "Styling/ISlateStyle.h"
@@ -39,6 +42,8 @@
 #include "ToolMenuSection.h"
 #include "Widgets/Input/SButton.h"
 #include "Widgets/SOverlay.h"
+#include "UObject/Class.h"
+#include "UObject/UObjectIterator.h"
 
 #include UE_INLINE_GENERATED_CPP_BY_NAME(MetasoundAssetDefinitions)
 
@@ -74,18 +79,18 @@ namespace Metasound::Editor
 
 		void ExecuteBrowseToPresetParentAsset(const FToolMenuContext& InContext)
 		{
-			using namespace Metasound;
+			using namespace Frontend;
 
 			if (UObject* MetaSound = UContentBrowserAssetContextMenuContext::LoadSingleSelectedAsset<UObject>(InContext))
 			{
 				if (UAssetEditorSubsystem* AssetSubsystem = GEditor->GetEditorSubsystem<UAssetEditorSubsystem>())
 				{
 					FMetaSoundFrontendDocumentBuilder& DocumentBuilder = Frontend::IDocumentBuilderRegistry::GetChecked().FindOrBeginBuilding(MetaSound);
-					if (const FMetasoundAssetBase* MetaSoundAsset = DocumentBuilder.GetReferencedPresetAsset())
+					if (const FPresetDocumentConfiguration* PresetConfig = DocumentBuilder.GetConstDocumentConfigurationAs<FPresetDocumentConfiguration>())
 					{
-						if (GEditor)
+						if (UObject* ParentObject = PresetConfig->Parent.GetObject(); GEditor && ParentObject)
 						{
-							GEditor->SyncBrowserToObjects({ MetaSoundAsset->GetOwningAsset() });
+							GEditor->SyncBrowserToObjects(TArray<UObject*>{ ParentObject });
 						}
 					}
 				}
@@ -94,65 +99,138 @@ namespace Metasound::Editor
 
 		void ExecuteOpenPresetParentAsset(const FToolMenuContext& InContext)
 		{
-			using namespace Metasound;
+			using namespace Frontend;
 
 			if (UObject* MetaSound = UContentBrowserAssetContextMenuContext::LoadSingleSelectedAsset<UObject>(InContext))
 			{
 				if (UAssetEditorSubsystem* AssetSubsystem = GEditor->GetEditorSubsystem<UAssetEditorSubsystem>())
 				{
 					FMetaSoundFrontendDocumentBuilder& DocumentBuilder = Frontend::IDocumentBuilderRegistry::GetChecked().FindOrBeginBuilding(MetaSound);
-					if (const FMetasoundAssetBase* MetaSoundAsset = DocumentBuilder.GetReferencedPresetAsset())
+					if (const FPresetDocumentConfiguration* PresetConfig = DocumentBuilder.GetConstDocumentConfigurationAs<FPresetDocumentConfiguration>())
 					{
-						AssetSubsystem->OpenEditorForAsset(MetaSoundAsset->GetOwningAsset());
+						if (UObject* ParentObject = PresetConfig->Parent.GetObject())
+						{
+							AssetSubsystem->OpenEditorForAsset(ParentObject);
+						}
 					}
 				}
 			}
 		}
-		
+
+		template <typename TClass, typename TFactoryClass>
+		void ExecuteCreateMetaSoundWithDocConfig(const FToolMenuContext& MenuContext, TInstancedStruct<FMetaSoundFrontendDocumentConfiguration> DocConfig)
+		{
+			if (const UContentBrowserAssetContextMenuContext* Context = UContentBrowserAssetContextMenuContext::FindContextWithAssets(MenuContext))
+			{
+				const UScriptStruct* DocConfigStruct = DocConfig.GetScriptStruct();
+				check(DocConfigStruct);
+
+				const FText DisplayName = DocConfigStruct->GetDisplayNameText();
+
+				UMetaSoundBaseFactory* Factory = NewObject<TFactoryClass>();
+				check(Factory);
+				Factory->DocumentConfiguration = MoveTemp(DocConfig);
+				Factory->SelectedObjects = Context->LoadSelectedObjects<UObject>();
+
+				if (Factory->SelectedObjects.IsEmpty())
+				{
+					return;
+				}
+
+				check(Factory->SelectedObjects.Last());
+				UObject* Outermost = Factory->SelectedObjects.Last()->GetOutermost();
+				check(Outermost);
+
+				FString PackagePath;
+				FString AssetName;
+				IAssetTools::Get().CreateUniqueAssetName(Outermost->GetName(), FString::Format(TEXT("_{0}"), { DisplayName.ToString() }), PackagePath, AssetName);
+
+				if (UObject* NewMetaSound = IAssetTools::Get().CreateAssetWithDialog(AssetName, FPackageName::GetLongPackagePath(PackagePath), Factory->GetSupportedClass(), Factory))
+				{
+					FContentBrowserModule& ContentBrowserModule = FModuleManager::LoadModuleChecked<FContentBrowserModule>("ContentBrowser");
+					ContentBrowserModule.Get().SyncBrowserToAssets(TArray<UObject*>({ NewMetaSound }));
+
+					// Save happened in Create call above, but need to save again after
+					// doc configuration executes to capture initial procedural changes.
+					if (ISourceControlModule::Get().IsEnabled())
+					{
+						constexpr bool bCheckDirty = false;
+						constexpr bool bPromptToSave = false;
+						FEditorFileUtils::PromptForCheckoutAndSave({ NewMetaSound->GetPackage() }, bCheckDirty, bPromptToSave);
+					}
+				}
+				else
+				{
+					UE_LOG(LogMetaSound, Display, TEXT("Error creating new asset when creating '%s' or creation of MetaSound with configuration was canceled by user."), *AssetName);
+				}
+			}
+		}
+
+		void RegisterAndSetObjectSettings(IMetasoundEditorModule& MetaSoundEditorModule, UMetaSoundEditorSubsystem& EditorSubsystem, UMetaSoundPatch& NewMetaSound, UMetaSoundPatch& ParentMetaSound)
+		{
+			if (MetaSoundEditorModule.IsRestrictedMode())
+			{
+				FGraphBuilder::RegisterGraphWithFrontend(NewMetaSound);
+			}
+		}
+
+		void RegisterAndSetObjectSettings(IMetasoundEditorModule& MetaSoundEditorModule, UMetaSoundEditorSubsystem& EditorSubsystem, UMetaSoundSource& NewMetaSound, UMetaSoundSource& ParentMetaSound)
+		{
+			if (MetaSoundEditorModule.IsRestrictedMode())
+			{
+				FGraphBuilder::RegisterGraphWithFrontend(NewMetaSound);
+				EditorSubsystem.SetSoundWaveSettingsFromTemplate(NewMetaSound, ParentMetaSound);
+			}
+		}
+
 		template <typename TClass, typename TFactoryClass>
 		void ExecuteCreateMetaSoundPreset(const FToolMenuContext& MenuContext)
 		{
+			using namespace Frontend;
+
 			if (const UContentBrowserAssetContextMenuContext* Context = UContentBrowserAssetContextMenuContext::FindContextWithAssets(MenuContext))
 			{
-				using namespace Metasound::Editor;
 				TArray<UObject*> ObjectsToSync;
 
-				for (TClass* ReferencedMetaSound : Context->LoadSelectedObjects<TClass>())
+				for (TClass* ParentMetaSound : Context->LoadSelectedObjects<TClass>())
 				{
 					FString PackagePath;
 					FString AssetName;
 					UObject* NewMetaSound = nullptr;
 
-					IAssetTools::Get().CreateUniqueAssetName(ReferencedMetaSound->GetOutermost()->GetName(), TEXT("_Preset"), PackagePath, AssetName);
+					IAssetTools::Get().CreateUniqueAssetName(ParentMetaSound->GetOutermost()->GetName(), TEXT("_Preset"), PackagePath, AssetName);
+
+					TInstancedStruct<FPresetDocumentConfiguration> PresetConfig = TInstancedStruct<FPresetDocumentConfiguration>::Make();
+					PresetConfig.GetMutable().Parent = ParentMetaSound;
 
-					Metasound::Editor::IMetasoundEditorModule& MetaSoundEditorModule = FModuleManager::GetModuleChecked<Metasound::Editor::IMetasoundEditorModule>("MetaSoundEditor");
+					IMetasoundEditorModule& MetaSoundEditorModule = FModuleManager::GetModuleChecked<IMetasoundEditorModule>("MetaSoundEditor");
 					if (MetaSoundEditorModule.IsRestrictedMode())
 					{
 						// Cannot duplicate cooked assets in restricted mode, so create new object and copy over properties
 						// for sources in InitAsset below. Since copying properties is done manually, 
 						// SetSoundWaveSettingsFromTemplate may need to be updated with properties to be copied. 
 						UMetaSoundBaseFactory* Factory = NewObject<TFactoryClass>();
-						Factory->ReferencedMetaSoundObject = ReferencedMetaSound;
+						check(Factory);
 
 						NewMetaSound = IAssetTools::Get().CreateAssetWithDialog(AssetName, FPackageName::GetLongPackagePath(PackagePath), Factory->GetSupportedClass(), Factory);
 					}
 					else
 					{
-						// Duplicate asset to preserve properties from referenced asset (ex. quality settings, soundwave properties)
-						NewMetaSound = IAssetTools::Get().DuplicateAssetWithDialogAndTitle(AssetName, FPackageName::GetLongPackagePath(PackagePath), ReferencedMetaSound, LOCTEXT("CreateMetaSoundPresetTitle", "Create MetaSound Preset"));
-
-						// Duplicate needs to initialize the asset whereas creating an asset with a factory as above in restricted mode already initializes the asset.
-						// (calling twice is explicitly avoided as this would result in an assert due to the Builder Subsystem then being left with a stale class key
-						// left pointing to the new MetaSound's builder.
-						if (NewMetaSound)
-						{
-							UMetaSoundEditorSubsystem::GetChecked().InitAsset(*NewMetaSound, ReferencedMetaSound, /*bClearDocument=*/true);
-						}
+						// Duplicate asset to preserve properties from parent asset (ex. quality settings, soundwave properties)
+						NewMetaSound = IAssetTools::Get().DuplicateAssetWithDialogAndTitle(AssetName, FPackageName::GetLongPackagePath(PackagePath), ParentMetaSound, LOCTEXT("CreateMetaSoundPresetTitle", "Create MetaSound Preset"));
 					}
 				
 					if (NewMetaSound)
 					{
-						FGraphBuilder::RegisterGraphWithFrontend(*NewMetaSound);
+						UMetaSoundEditorSubsystem& EditorSubsystem = UMetaSoundEditorSubsystem::GetChecked();
+						EditorSubsystem.InitAsset(*NewMetaSound, UMetaSoundEditorSubsystem::FInitAssetArgs
+						{
+							.DocumentConfiguration = MoveTemp(PresetConfig),
+							.SelectedObjects = { ParentMetaSound }
+						});
+
+						RegisterAndSetObjectSettings(MetaSoundEditorModule, EditorSubsystem, *CastChecked<TClass>(NewMetaSound), *ParentMetaSound);
+
 						ObjectsToSync.Add(NewMetaSound);
 					}
 					else
@@ -206,27 +284,29 @@ namespace Metasound::Editor
 			return false;
 		}
 
-		void AddMetaSoundActions(const FSlateIcon& AssetIcon, const UClass& Class, FToolUIAction CreatePresetAction)
+		void AddPresetActions(const FSlateIcon& AssetIcon, const UClass& MetaSoundClass, const UClass& Class, FToolUIAction CreateFromConfigAction, const FText& ConfigDisplayName)
 		{
 			const FText ClassName = Class.GetDisplayNameText();
 
+			const FName MetaSoundClassName = MetaSoundClass.GetFName();
 			const UAssetDefinition* AssetDefinition = UAssetDefinitionRegistry::Get()->GetAssetDefinitionForClass(&Class);
-			FToolMenuSection* Section = CastChecked<UAssetDefinition_SoundAssetBase>(AssetDefinition)->FindSoundContextMenuSection("MetaSound");
+			FToolMenuSection* Section = CastChecked<UAssetDefinition_SoundAssetBase>(AssetDefinition)->FindSoundContextMenuSection(MetaSoundClassName);
 			check(Section);
 
-			Section->AddDynamicEntry(NAME_None, FNewToolMenuSectionDelegate::CreateLambda([CreateAction = MoveTemp(CreatePresetAction), ClassName, AssetIcon](FToolMenuSection& InSection)
+			Section->AddDynamicEntry(NAME_None, FNewToolMenuSectionDelegate::CreateLambda([CreateAction = MoveTemp(CreateFromConfigAction), MetaSoundClassName, ConfigName = ConfigDisplayName, ClassName, AssetIcon](FToolMenuSection& InSection)
 			{
 				using namespace Metasound::Editor;
 
 				{
-					const TAttribute<FText> Label = FText::Format(LOCTEXT("MetaSound_CreatePreset", "Create Preset..."), ClassName);
-					const TAttribute<FText> ToolTip = FText::Format(LOCTEXT("MetaSound_CreatePresetToolTipFormat", "Creates a {0} Preset using the selected MetaSound as a reference (parent)."), ClassName);
+					const TAttribute<FText> Label = FText::Format(LOCTEXT("MetaSound_CreateFromPreset", "Create {0}..."), ConfigName);
+					const TAttribute<FText> ToolTip = FText::Format(LOCTEXT("MetaSound_CreateFromPresetToolTipFormat", "Creates a {0} Preset using the selected MetaSound as a parent."), ClassName);
 
-					InSection.AddMenuEntry("MetaSound_CreatePreset", Label, ToolTip, AssetIcon, CreateAction);
+					const FName MenuEntryName = *FString::Format(TEXT("{0}_CreateFromConfig_{1}"), { MetaSoundClassName.ToString(), ConfigName.ToString()});
+					InSection.AddMenuEntry(MenuEntryName, Label, ToolTip, AssetIcon, CreateAction);
 				}
 				{
 					const TAttribute<FText> Label = LOCTEXT("MetaSound_BrowseToPresetParentFormat", "Browse To Preset Parent");
-					const TAttribute<FText> ToolTip = FText::Format(LOCTEXT("MetaSound_BrowseToPresetParentToolTipFormat", "Browses to the selected {0} preset's referenced parent asset in the content browser."), ClassName);
+					const TAttribute<FText> ToolTip = FText::Format(LOCTEXT("MetaSound_BrowseToPresetParentToolTipFormat", "Browses to the selected {0} preset's parent asset in the content browser."), ClassName);
 					const FSlateIcon FindInContentBrowserIcon(FAppStyle::GetAppStyleSetName(), "SystemWideCommands.FindInContentBrowser");
 
 					FToolUIAction UIAction;
@@ -251,7 +331,7 @@ namespace Metasound::Editor
 			}));
 		}
 
-		void AddMetaSoundContextMenuAction(const UClass& MetaSoundClass, TArray<FToolMenuSection*>& OutSections)
+		FToolMenuSection* AddMetaSoundContextMenuAction(const UClass& MetaSoundClass, const UClass& ActionClass, TArray<FToolMenuSection*>& OutSections)
 		{
 			const UAudioEditorSettings* EdSettings = GetDefault<UAudioEditorSettings>();
 			check(EdSettings);
@@ -274,14 +354,15 @@ namespace Metasound::Editor
 			FToolMenuSection* PlaybackSection = GetSection("Playback");
 			FToolMenuSection* SoundSection = GetSection("Sound");
 
-			UToolMenu* Menu = UE::ContentBrowser::ExtendToolMenu_AssetContextMenu(&MetaSoundClass);
+			UToolMenu* Menu = UE::ContentBrowser::ExtendToolMenu_AssetContextMenu(&ActionClass);
 			check(Menu);
 
 			const UEnum* EnumClass = StaticEnum<EToolMenuInsertType>();
 			check(EnumClass);
 			const EToolMenuInsertType InsertType = static_cast<EToolMenuInsertType>(EnumClass->GetValueByName(EdSettings->MenuPosition));
 
-			FToolMenuSection& MetaSoundSection = Menu->FindOrAddSection("MetaSound");
+			const FName MetaSoundClassName = MetaSoundClass.GetFName();
+			FToolMenuSection& MetaSoundSection = Menu->FindOrAddSection(MetaSoundClassName);
 			MetaSoundSection.Label = MetaSoundClass.GetDisplayNameText();
 
 			if (SoundSection)
@@ -297,8 +378,8 @@ namespace Metasound::Editor
 					SoundSection->InsertPosition = FToolMenuInsert(
 						MetaSoundSection.Name,
 						InsertType == EToolMenuInsertType::Last
-							? EToolMenuInsertType::Before
-							: EToolMenuInsertType::After);
+						? EToolMenuInsertType::Before
+						: EToolMenuInsertType::After);
 				}
 				else
 				{
@@ -315,7 +396,52 @@ namespace Metasound::Editor
 				MetaSoundSection.InsertPosition = FToolMenuInsert({ }, InsertType);
 			}
 
-			OutSections.Add(&MetaSoundSection);
+			return OutSections.Add_GetRef(&MetaSoundSection);
+		}
+
+		void AddConfigActions(
+			const FSlateIcon& AssetIcon,
+			const UClass& Class,
+			const UClass& MetaSoundClass,
+			FToolUIAction CreateFromConfigAction,
+			TConstStructView<FMetaSoundFrontendDocumentConfiguration> Config)
+		{
+			const FText ClassName = Class.GetDisplayNameText();
+			const UScriptStruct* ScriptStruct = Config.GetScriptStruct();
+			check(ScriptStruct);
+			const UAssetDefinition* AssetDefinition = UAssetDefinitionRegistry::Get()->GetAssetDefinitionForClass(&Class);
+
+			FToolMenuSection* Section = nullptr;
+			if (const UAssetDefinition_SoundAssetBase* AudioAssetDef = Cast<const UAssetDefinition_SoundAssetBase>(AssetDefinition))
+			{
+				Section = AudioAssetDef->FindSoundContextMenuSection(MetaSoundClass.GetFName());
+			}
+
+			if (!Section)
+			{
+				TArray<FToolMenuSection*> Sections;
+				Section = AddMetaSoundContextMenuAction(MetaSoundClass, Class, Sections);
+			}
+
+			if (Section)
+			{
+				const FText ConfigToolTip = Config.GetPtr()->GetToolTip();
+				auto AddEntryLambda = [CreateAction = MoveTemp(CreateFromConfigAction), ScriptStruct, ConfigToolTip, ClassName, AssetIcon](FToolMenuSection& InSection)
+				{
+					using namespace Metasound::Editor;
+			
+					const FText ConfigName = ScriptStruct->GetDisplayNameText();
+					const TAttribute<FText> Label = FText::Format(LOCTEXT("MetaSound_CreateFromConfig", "Create {0}..."), ConfigName);
+					const TAttribute<FText> ToolTip = FText::Format(
+						LOCTEXT("MetaSound_CreateFromConfigToolTipFormat", "Creates a {0} with {1} configuration, initialized using the selected assets. {2}"),
+						ClassName, ConfigName, ConfigToolTip);
+			
+					const FName MenuEntryName = *FString::Format(TEXT("MetaSound_CreateFromConfig_{0}"), { ConfigName.ToString() });
+					InSection.AddMenuEntry(MenuEntryName, Label, ToolTip, AssetIcon, CreateAction);
+				};
+
+				Section->AddDynamicEntry(NAME_None, FNewToolMenuSectionDelegate::CreateLambda(AddEntryLambda));
+			}
 		}
 	} // namespace AssetDefinitionsPrivate
 } // namespace Metasound::Editor
@@ -383,7 +509,7 @@ TArray<FToolMenuSection*> UAssetDefinition_MetaSoundPatch::RebuildSoundContextMe
 	TArray<FToolMenuSection*> Sections = Super::RebuildSoundContextMenuSections();
 	if (const UClass* MetaSoundClass = GetAssetClass().Get())
 	{
-		AssetDefinitionsPrivate::AddMetaSoundContextMenuAction(*MetaSoundClass, Sections);
+		AssetDefinitionsPrivate::AddMetaSoundContextMenuAction(*MetaSoundClass, *MetaSoundClass, Sections);
 	}
 	return Sections;
 }
@@ -579,9 +705,9 @@ TSharedPtr<SWidget> UAssetDefinition_MetaSoundSource::GetThumbnailOverlay(const
 bool UAssetDefinition_MetaSoundSource::GetThumbnailActionOverlay(const FAssetData& InAssetData, FAssetActionThumbnailOverlayInfo& OutActionOverlayInfo) const
 {
 	OutActionOverlayInfo.IsActionPlayingDelegate = FIsActionPlaying::CreateLambda([InAssetData] ()
-		{
-			return UE::AudioEditor::IsSoundPlaying(InAssetData);
-		});
+	{
+		return UE::AudioEditor::IsSoundPlaying(InAssetData);
+	});
 
 	auto OnToolTipTextLambda = [InAssetData]() -> FText
 	{
@@ -679,14 +805,88 @@ TArray<FToolMenuSection*> UAssetDefinition_MetaSoundSource::RebuildSoundContextM
 	TArray<FToolMenuSection*> Sections = Super::RebuildSoundContextMenuSections();
 	if (const UClass* MetaSoundClass = GetAssetClass().Get())
 	{
-		AssetDefinitionsPrivate::AddMetaSoundContextMenuAction(*MetaSoundClass, Sections);
+		AssetDefinitionsPrivate::AddMetaSoundContextMenuAction(*MetaSoundClass, *MetaSoundClass, Sections);
 	}
 	return Sections;
 }
 
 
-namespace MenuExtension_MetaSoundSourceTemplate
+namespace MenuExtension_MetaSoundTemplate
 {
+	template <typename TClass, typename TFactoryClass>
+	void RegisterDocumentConfigurationActions(const FSlateIcon& AssetIcon)
+	{
+		using namespace Metasound;
+		using namespace Metasound::Editor;
+		using namespace Metasound::Frontend;
+
+		const UClass* MetaSoundClass = TClass::StaticClass();
+		check(MetaSoundClass);
+
+		for (TObjectIterator<UStruct> It; It; ++It)
+		{
+			UScriptStruct* Struct = Cast<UScriptStruct>(*It);
+			if (!Struct)
+			{
+				continue;
+			}
+
+			if (!Struct->IsChildOf(FMetaSoundFrontendDocumentConfiguration::StaticStruct()))
+			{
+				continue;
+			}
+
+			// Presets are special snowflakes for now, as they require a bit more functionality for back
+			// compat like copying over SoundWave fields which still needs to be migrated to doc config.
+			if (Struct == FPresetDocumentConfiguration::StaticStruct())
+			{
+				continue;
+			}
+
+			// Don't show structs marked as hidden
+			if (Struct->FindMetaData("Hidden") != nullptr)
+			{
+				continue;
+			}
+
+			TInstancedStruct<FMetaSoundFrontendDocumentConfiguration> ConfigInstance;
+			ConfigInstance.InitializeAsScriptStruct(Struct);
+			TArray<const UClass*> SupportedMetaSoundClasses = ConfigInstance.GetPtr()->GetSupportedMetaSoundClasses();
+			if (SupportedMetaSoundClasses.IsEmpty())
+			{
+				IMetasoundUObjectRegistry::Get().IterateRegisteredUClasses([&SupportedMetaSoundClasses](UClass& RegClass)
+				{
+					SupportedMetaSoundClasses.Add(&RegClass);
+				});
+			}
+
+			if (SupportedMetaSoundClasses.Contains(TClass::StaticClass()))
+			{
+				FToolUIAction DocConfigCreateAction;
+				DocConfigCreateAction.ExecuteAction = FToolMenuExecuteAction::CreateLambda([Struct](const FToolMenuContext& MenuContext)
+				{
+					TInstancedStruct<FMetaSoundFrontendDocumentConfiguration> NewInstance;
+					NewInstance.InitializeAsScriptStruct(Struct);
+					return AssetDefinitionsPrivate::ExecuteCreateMetaSoundWithDocConfig<TClass, TFactoryClass>(MenuContext, MoveTemp(NewInstance));
+				});
+
+				TArray<const UClass*> AssetActionClasses = ConfigInstance.GetPtr()->GetSupportedAssetActionClasses();
+				if (AssetActionClasses.IsEmpty())
+				{
+					AssetDefinitionsPrivate::AddConfigActions(AssetIcon, *MetaSoundClass, *MetaSoundClass, MoveTemp(DocConfigCreateAction), ConfigInstance);
+				}
+				else
+				{
+					for (const UClass* Class : AssetActionClasses)
+					{
+						check(Class);
+						AssetDefinitionsPrivate::AddConfigActions(AssetIcon, *Class, *MetaSoundClass, DocConfigCreateAction, ConfigInstance);
+					}
+				}
+			}
+		}
+	}
+
  	static FDelayedAutoRegisterHelper DelayedAutoRegister(EDelayedRegisterRunPhase::EndOfEngineInit, []
 	{
 		if (!UToolMenus::IsToolMenuUIEnabled())
@@ -697,16 +897,23 @@ namespace MenuExtension_MetaSoundSourceTemplate
  		UToolMenus::RegisterStartupCallback(FSimpleMulticastDelegate::FDelegate::CreateLambda([]()
  		{
 			using namespace Metasound::Editor;
+			using namespace Metasound::Frontend;
 
  			FToolMenuOwnerScoped OwnerScoped(UE_MODULE_NAME);
 
+			// MetaSound Source Action Registration
 			{
+				const FSlateIcon AssetIcon = Style::CreateSlateIcon("ClassIcon.MetasoundSource");
+
+				RegisterDocumentConfigurationActions<UMetaSoundSource, UMetaSoundSourceFactory>(AssetIcon);
+
 				const UClass* MetaSoundClass = UMetaSoundSource::StaticClass();
 				check(MetaSoundClass);
-				FToolUIAction CreateAction;
-				CreateAction.ExecuteAction = FToolMenuExecuteAction::CreateStatic(&AssetDefinitionsPrivate::ExecuteCreateMetaSoundPreset<UMetaSoundSource, UMetaSoundSourceFactory>);
-				const FSlateIcon AssetIcon = Metasound::Editor::Style::CreateSlateIcon("ClassIcon.MetasoundSource");
-				AssetDefinitionsPrivate::AddMetaSoundActions(AssetIcon, *MetaSoundClass, MoveTemp(CreateAction));
+
+				FToolUIAction PresetCreateAction;
+				PresetCreateAction.ExecuteAction = FToolMenuExecuteAction::CreateStatic(&AssetDefinitionsPrivate::ExecuteCreateMetaSoundPreset<UMetaSoundSource, UMetaSoundSourceFactory>);
+				AssetDefinitionsPrivate::AddPresetActions(AssetIcon, *MetaSoundClass, *MetaSoundClass, MoveTemp(PresetCreateAction), FPresetDocumentConfiguration::StaticStruct()->GetDisplayNameText());
+
 				{
 					const UAssetDefinition* AssetDefinition = UAssetDefinitionRegistry::Get()->GetAssetDefinitionForClass(MetaSoundClass);
 					FToolMenuSection* PlaybackSection = CastChecked<UAssetDefinition_SoundAssetBase>(AssetDefinition)->FindSoundContextMenuSection("Playback");
@@ -809,14 +1016,18 @@ namespace MenuExtension_MetaSoundSourceTemplate
 					}));
 				}
  			}
-	
+
+			// MetaSound Patch Action Registration
 			{
-				const UClass* MetaSoundClass = UMetaSoundPatch::StaticClass();
-				check(MetaSoundClass);
 				FToolUIAction CreateAction;
 				CreateAction.ExecuteAction = FToolMenuExecuteAction::CreateStatic(&AssetDefinitionsPrivate::ExecuteCreateMetaSoundPreset<UMetaSoundPatch, UMetaSoundFactory>);
 				const FSlateIcon AssetIcon = Metasound::Editor::Style::CreateSlateIcon("ClassIcon.MetasoundPatch");
-				AssetDefinitionsPrivate::AddMetaSoundActions(AssetIcon, *MetaSoundClass, MoveTemp(CreateAction));
+
+				RegisterDocumentConfigurationActions<UMetaSoundPatch, UMetaSoundFactory>(AssetIcon);
+
+				const UClass* MetaSoundClass = UMetaSoundPatch::StaticClass();
+				check(MetaSoundClass);
+				AssetDefinitionsPrivate::AddPresetActions(AssetIcon, *MetaSoundClass, *MetaSoundClass, MoveTemp(CreateAction), FPresetDocumentConfiguration::StaticStruct()->GetDisplayNameText());
 			}
  		}));
 	});
diff --git a/Engine/Plugins/Runtime/Metasound/Source/MetasoundEditor/Private/MetasoundDefaultLiteralCustomization.cpp b/Engine/Plugins/Runtime/Metasound/Source/MetasoundEditor/Private/MetasoundDefaultLiteralCustomization.cpp
--- a/Engine/Plugins/Runtime/Metasound/Source/MetasoundEditor/Private/MetasoundDefaultLiteralCustomization.cpp
+++ b/Engine/Plugins/Runtime/Metasound/Source/MetasoundEditor/Private/MetasoundDefaultLiteralCustomization.cpp
@@ -380,9 +380,8 @@ namespace Metasound::Editor
 			if (DefaultValueArray.IsValid())
 			{
 				DefaultValueArray->GetNumElements(NumElements);
-				const bool bHasProjectPageValues = NumElements > 1; // 1 value is only the default
-				constexpr bool bPresetCanEditPageValues = true;
-				const bool bShowPageModifiers = Editor::PageEditorEnabled(Member->GetFrontendBuilderChecked(), bHasProjectPageValues, bPresetCanEditPageValues);
+				const bool bHasLocalPageValues = NumElements > 1; // 1 value is only the default
+				const bool bShowPageModifiers = Editor::IsInputPageEditingEnabled(Member->GetFrontendBuilderChecked(), bHasLocalPageValues);
 				if (bIsPagedDefault && bShowPageModifiers)
 				{
 					BuildPageDefaultComboBox(InLiteral, DefaultPageArrayHandle->GetPropertyDisplayName());
diff --git a/Engine/Plugins/Runtime/Metasound/Source/MetasoundEditor/Private/MetasoundDetailCustomization.cpp b/Engine/Plugins/Runtime/Metasound/Source/MetasoundEditor/Private/MetasoundDetailCustomization.cpp
--- a/Engine/Plugins/Runtime/Metasound/Source/MetasoundEditor/Private/MetasoundDetailCustomization.cpp
+++ b/Engine/Plugins/Runtime/Metasound/Source/MetasoundEditor/Private/MetasoundDetailCustomization.cpp
@@ -190,44 +190,42 @@ namespace Metasound::Editor
 				TSharedPtr<IPropertyHandle> MajorVersionHandle = DetailLayout.GetProperty(MajorVersionPropertyPath);
 				TSharedPtr<IPropertyHandle> MinorVersionHandle = DetailLayout.GetProperty(MinorVersionPropertyPath);
 
-				// Disabled as currently experimental and in local development.
-// 				{
-// 					TSharedPtr<IPropertyHandle> DocConfigHandle = DetailLayout.GetProperty(GetDocConfigPropertyPath());
-// 					if (DocConfigHandle.IsValid())
-// 					{
-// 						FSimpleDelegate OnConfigChanged = FSimpleDelegate::CreateLambda([this]()
-// 						{
-// 							if (Builder.IsValid())
-// 							{
-// 								FMetaSoundFrontendDocumentBuilder& DocBuilder = Builder->GetBuilder();
-// 								if (DocBuilder.IsValid())
-// 								{
-// 									DocBuilder.GetDocumentDelegates().OnDocumentConfigurationChanged.Broadcast({ });
-// 									DocBuilder.ConfigureDocument();
-// 								}
-// 							}
-// 						});
-// 						DocConfigHandle->SetOnPropertyValueChanged(OnConfigChanged);
-// 
-// 						const auto OnConfigPropChanged = TDelegate<void(const FPropertyChangedEvent&)>::CreateSPLambda(AsShared(), [this](const FPropertyChangedEvent& InEvent)
-// 						{
-// 							if (Builder.IsValid())
-// 							{
-// 								FMetaSoundFrontendDocumentBuilder& DocBuilder = Builder->GetBuilder();
-// 								if (DocBuilder.IsValid())
-// 								{
-// 									if (const FMetaSoundFrontendDocumentConfiguration* Config = DocBuilder.GetConstDocumentConfiguration().GetPtr())
-// 									{
-// 										const_cast<FMetaSoundFrontendDocumentConfiguration*>(Config)->OnPropertyChanged(InEvent, DocBuilder);
-// 									}
-// 								}
-// 							}
-// 						});
-// 						DocConfigHandle->SetOnChildPropertyValueChangedWithData(OnConfigPropChanged);
-// 					}
-// 
-// 					GeneralCategoryBuilder.AddProperty(DocConfigHandle);
-// 				}
+				{
+					TSharedPtr<IPropertyHandle> DocConfigHandle = DetailLayout.GetProperty(GetDocConfigPropertyPath());
+					if (DocConfigHandle.IsValid())
+					{
+						FSimpleDelegate OnConfigChanged = FSimpleDelegate::CreateLambda([this]()
+						{
+							if (Builder.IsValid())
+							{
+								FMetaSoundFrontendDocumentBuilder& DocBuilder = Builder->GetBuilder();
+								if (DocBuilder.IsValid())
+								{
+									DocBuilder.GetDocumentDelegates().OnDocumentConfigurationChanged.Broadcast({ });
+									DocBuilder.ConfigureDocument();
+								}
+							}
+						});
+						DocConfigHandle->SetOnPropertyValueChanged(OnConfigChanged);
+
+						const auto OnConfigPropChanged = TDelegate<void(const FPropertyChangedEvent&)>::CreateSPLambda(AsShared(), [this](const FPropertyChangedEvent& InEvent)
+						{
+							if (Builder.IsValid())
+							{
+								FMetaSoundFrontendDocumentBuilder& DocBuilder = Builder->GetBuilder();
+								if (DocBuilder.IsValid())
+								{
+									if (const FDocumentConfiguration* Config = DocBuilder.GetConstDocumentConfiguration().GetPtr())
+									{
+										const_cast<FDocumentConfiguration*>(Config)->OnPropertyChanged(InEvent, DocBuilder);
+									}
+								}
+							}
+						});
+						DocConfigHandle->SetOnChildPropertyValueChangedWithData(OnConfigPropChanged);
+						GeneralCategoryBuilder.AddProperty(DocConfigHandle);
+					}
+				}
 
 				// Invalid for UMetaSounds
 				TSharedPtr<IPropertyHandle> OutputFormat = DetailLayout.GetProperty(GET_MEMBER_NAME_CHECKED(UMetaSoundSource, OutputFormat));
diff --git a/Engine/Plugins/Runtime/Metasound/Source/MetasoundEditor/Private/MetasoundEditor.cpp b/Engine/Plugins/Runtime/Metasound/Source/MetasoundEditor/Private/MetasoundEditor.cpp
--- a/Engine/Plugins/Runtime/Metasound/Source/MetasoundEditor/Private/MetasoundEditor.cpp
+++ b/Engine/Plugins/Runtime/Metasound/Source/MetasoundEditor/Private/MetasoundEditor.cpp
@@ -10,6 +10,9 @@
 #include "Blueprint/UserWidget.h"
 #include "Components/AudioComponent.h"
 #include "DetailLayoutBuilder.h"
+#include "DocumentConfigurations/MetasoundDocumentConfigurator.h"
+#include "DocumentConfigurations/MetasoundFrontendDocumentConfiguration.h"
+#include "DocumentConfigurations/MetasoundFrontendPresetDocumentConfiguration.h"
 #include "EdGraph/EdGraphNode.h"
 #include "EdGraph/EdGraphPin.h"
 #include "EdGraphUtilities.h"
@@ -96,6 +99,8 @@
 #include "Widgets/SOverlay.h"
 #include "Widgets/SWindow.h"
 #include "Styling/StyleColors.h"
+#include "UObject/Class.h"
+#include "UObject/UObjectIterator.h"
 
 #include UE_INLINE_GENERATED_CPP_BY_NAME(MetasoundEditor)
 
@@ -110,8 +115,11 @@ namespace Metasound
 	{
 		namespace AssetEditorPrivate
 		{
-			static const FText LoadingAssetsFormat = LOCTEXT("MetaSoundScanInProgressNotificationText", "Loading {0} MetaSound Assets...");
 			static bool bEnableLiveAudition = false;
+			static bool bIgnoreDocConfigVisibilitySettings = false;
+
+			static const FText LoadingAssetsFormat = LOCTEXT("MetaSoundScanInProgressNotificationText", "Loading {0} MetaSound Assets...");
+
 			FAutoConsoleVariableRef CVarMetaSoundEditorEnableLiveAudition(
 				TEXT("au.MetaSound.Editor.Experimental.EnableLiveAudition"),
 				bEnableLiveAudition,
@@ -120,6 +128,12 @@ namespace Metasound
 				TEXT("Default: false"),
 				ECVF_Default);
 
+			FAutoConsoleVariableRef CVarMetaSoundEditorIgnoreConfigVisibilitySettings(
+				TEXT("au.MetaSound.Editor.IgnoreConfigVisibilitySettings"),
+				bIgnoreDocConfigVisibilitySettings,
+				TEXT("If true, all settings on a Document Configuration for what is and is not visible in editor are ignored and all tabs are invokable.")
+				TEXT("Default: false"),
+				ECVF_Default);
 		} // namespace AssetEditorPrivate
 
 		bool IsPreviewingPageInputDefault(const FMetaSoundFrontendDocumentBuilder& Builder, const FMetasoundFrontendClassInput& InClassInput, const FGuid& InPageID)
@@ -157,13 +171,14 @@ namespace Metasound
 			return false;
 		}
 
-		bool PageEditorEnabled(const FMetaSoundFrontendDocumentBuilder& Builder, bool bHasProjectPageValues, bool bPresetCanEditPageValues)
+		bool IsPageAuditioningEnabled()
 		{
-			if (!Builder.IsValid())
-			{
-				return false;
-			}
+			IMetasoundEditorModule& MetaSoundEditorModule = FModuleManager::GetModuleChecked<IMetasoundEditorModule>("MetaSoundEditor");
+			return !MetaSoundEditorModule.IsRestrictedMode();
+		}
 
+		bool IsPageGraphEditingEnabled(const FMetaSoundFrontendDocumentBuilder& Builder, bool bHasLocalPageValues)
+		{
 			IMetasoundEditorModule& MetaSoundEditorModule = FModuleManager::GetModuleChecked<IMetasoundEditorModule>("MetaSoundEditor");
 			if (MetaSoundEditorModule.IsRestrictedMode())
 			{
@@ -178,15 +193,48 @@ namespace Metasound
 
 			if (Settings->GetProjectPageSettings().IsEmpty())
 			{
-				if (!bHasProjectPageValues)
+				if (!bHasLocalPageValues)
 				{
 					return false;
 				}
 			}
 
-			if (!bPresetCanEditPageValues)
+			if (!Builder.IsValid())
+			{
+				return false;
+			}
+
+			TConstStructView<FMetaSoundFrontendDocumentConfiguration> ConfigView = Builder.GetConstDocumentConfiguration();
+			if (ConfigView.IsValid())
+			{
+				return ConfigView.GetPtr()->IsPageGraphEditingEnabled();
+			}
+
+			return true;
+		}
+
+		bool IsInputPageEditingEnabled(const FMetaSoundFrontendDocumentBuilder& Builder, bool bHasProjectPageValues)
+		{
+			if (!Builder.IsValid())
+			{
+				return false;
+			}
+
+			IMetasoundEditorModule& MetaSoundEditorModule = FModuleManager::GetModuleChecked<IMetasoundEditorModule>("MetaSoundEditor");
+			if (MetaSoundEditorModule.IsRestrictedMode())
+			{
+				return false;
+			}
+
+			const UMetaSoundSettings* Settings = GetDefault<UMetaSoundSettings>();
+			if (!Settings)
 			{
-				if (Builder.IsPreset())
+				return false;
+			}
+
+			if (Settings->GetProjectPageSettings().IsEmpty())
+			{
+				if (!bHasProjectPageValues)
 				{
 					return false;
 				}
@@ -205,7 +253,7 @@ namespace Metasound
 			const FName Interfaces = "MetasoundEditor_Interfaces";
 			const FName Pages = "MetasoundEditor_Pages";
 			const FName Find = "MetasoundEditor_Find";
-			const FName UserPresetWidget = "MetasoundEditor_UserPresetWidget";
+			const FName DocConfigWidget = "MetasoundEditor_DocConfigWidget";
 		} // namespace TabNamesPrivate
 
 		static const TArray<FText> NodeSectionNames
@@ -1011,55 +1059,44 @@ namespace Metasound
 			.SetGroup(WorkspaceMenuCategoryRef)
 			.SetIcon(FSlateIcon(FAppStyle::GetAppStyleSetName(), "Kismet.Tabs.Palette"));
 
-			if (ShowPageGraphDetails())
+			InTabManager->RegisterTabSpawner(TabNamesPrivate::Pages, FOnSpawnTab::CreateLambda([this, InPagesDetails = PagesDetails](const FSpawnTabArgs&)
 			{
-				if (Builder.IsValid() && !Builder->IsPreset())
-				{
-					const FCanSpawnTab CanSpawnTab = FCanSpawnTab::CreateLambda([this](const FSpawnTabArgs&)
-					{
-						return Builder.IsValid() && !Builder->IsPreset();
-					});
-
-					InTabManager->RegisterTabSpawner(TabNamesPrivate::Pages, FOnSpawnTab::CreateLambda([this, InPagesDetails = PagesDetails](const FSpawnTabArgs&)
-					{
-						return SNew(SDockTab)
-						.Label(LOCTEXT("MetasoundPagesDetailsTitle", "Pages"))
-						[
-							InPagesDetails.ToSharedRef()
-						];
-					}), CanSpawnTab)
-					.SetDisplayName(LOCTEXT("PagesTab", "Pages"))
-					.SetGroup(WorkspaceMenuCategoryRef)
-					.SetIcon(FSlateIcon(FAppStyle::GetAppStyleSetName(), "Kismet.Tabs.Palette"));
-				}
-			}
+				return SNew(SDockTab)
+				.Label(LOCTEXT("MetasoundPagesDetailsTitle", "Pages"))
+				[
+					InPagesDetails.ToSharedRef()
+				];
+			}))
+			.SetDisplayName(LOCTEXT("PagesTab", "Pages"))
+			.SetGroup(WorkspaceMenuCategoryRef)
+			.SetIcon(FSlateIcon(FAppStyle::GetAppStyleSetName(), "Kismet.Tabs.Palette"));
 
-			if (UserPresetWidget.IsValid())
+			InTabManager->RegisterTabSpawner(TabNamesPrivate::DocConfigWidget, FOnSpawnTab::CreateLambda([this](const FSpawnTabArgs&)
 			{
-				const FCanSpawnTab CanSpawnTab = FCanSpawnTab::CreateLambda([this](const FSpawnTabArgs&)
-				{
-					return Builder.IsValid() && Builder->IsPreset();
-				});
-
-				InTabManager->RegisterTabSpawner(TabNamesPrivate::UserPresetWidget, FOnSpawnTab::CreateLambda([this](const FSpawnTabArgs&)
+				if (DocConfigWidget.IsValid())
 				{
 					// Refresh widget if widget blueprint has been recompiled
-					if (UserPresetWidget->GetClass()->HasAnyClassFlags(EClassFlags::CLASS_NewerVersionExists))
+					if (DocConfigWidget->GetClass()->HasAnyClassFlags(EClassFlags::CLASS_NewerVersionExists))
 					{
-						UserPresetWidget = TStrongObjectPtr<UUserWidget>(CreateUserPresetWidget());
+						DocConfigWidget = TStrongObjectPtr<UUserWidget>(CreateDocConfigWidget());
 					}
+				}
 
-					IMetaSoundPresetWidgetInterface::Execute_OnConstructed(UserPresetWidget.Get(), Builder.Get());
+				if (DocConfigWidget.IsValid())
+				{
+					IMetaSoundPresetWidgetInterface::Execute_OnConstructed(DocConfigWidget.Get(), Builder.Get());
 					return SNew(SDockTab)
-						.Label(LOCTEXT("UserPresetWidgetTitle", "Preset Widget"))
+						.Label(LOCTEXT("DocConfigWidgetTitle", "Configuration Widget"))
 						[
-							UserPresetWidget.Get()->TakeWidget()
+							DocConfigWidget.Get()->TakeWidget()
 						];
-				}), CanSpawnTab)
-				.SetDisplayName(LOCTEXT("UserPresetWidgetTab", "Preset Widget"))
-				.SetGroup(WorkspaceMenuCategoryRef)
-				.SetIcon(FSlateIcon(FAppStyle::GetAppStyleSetName(), "Kismet.Tabs.Palette"));
-			}
+				}
+
+				return SNew(SDockTab);
+			}))
+			.SetDisplayName(LOCTEXT("DocConfigWidgetTab", "Configuration"))
+			.SetGroup(WorkspaceMenuCategoryRef)
+			.SetIcon(FSlateIcon(FAppStyle::GetAppStyleSetName(), "Kismet.Tabs.Palette"));
 
 			InTabManager->RegisterTabSpawner(TabNamesPrivate::Interfaces, FOnSpawnTab::CreateLambda([InInterfacesDetails = InterfacesDetails](const FSpawnTabArgs&)
 			{
@@ -1091,7 +1128,7 @@ namespace Metasound
 			InTabManager->UnregisterTabSpawner(TabNamesPrivate::Pages);
 			InTabManager->UnregisterTabSpawner(TabNamesPrivate::Interfaces);
 			InTabManager->UnregisterTabSpawner(TabNamesPrivate::Find);
-			InTabManager->UnregisterTabSpawner(TabNamesPrivate::UserPresetWidget);
+			InTabManager->UnregisterTabSpawner(TabNamesPrivate::DocConfigWidget);
 		}
 
 		bool FEditor::IsLiveAuditionEnabled()
@@ -1133,9 +1170,9 @@ namespace Metasound
 			FLevelEditorModule& LevelEditor = FModuleManager::LoadModuleChecked<FLevelEditorModule>("LevelEditor");
 			LevelEditor.OnMapChanged().RemoveAll(this);
 
-			if (UserPresetWidget.IsValid())
+			if (DocConfigWidget.IsValid())
 			{
-				if (UBlueprint* OwningBlueprint = UBlueprint::GetBlueprintFromClass(UserPresetWidget->GetClass()))
+				if (UBlueprint* OwningBlueprint = UBlueprint::GetBlueprintFromClass(DocConfigWidget->GetClass()))
 				{
 					OwningBlueprint->OnCompiled().RemoveAll(this);
 				}
@@ -1217,33 +1254,16 @@ namespace Metasound
 				->SetHideTabWell(false)
 				->AddTab(TabNamesPrivate::Details, ETabState::OpenedTab);
 
-			if (ShowPageGraphDetails())
-			{
-				DetailsStack->AddTab(TabNamesPrivate::Pages, ETabState::OpenedTab);
-			}
-			else
-			{
-				DetailsStack->AddTab(TabNamesPrivate::Pages, ETabState::InvalidTab);
-			}
+			DetailsStack->AddTab(TabNamesPrivate::Pages, ETabState::OpenedTab);
 
 			GraphCanvasTabStack = FTabManager::NewStack()
 				->SetSizeCoefficient(0.8f)
 				->SetHideTabWell(true)
 				->SetExtensionId(TabNamesPrivate::GraphCanvas)
 				->AddTab(TabNamesPrivate::GraphCanvas, ETabState::OpenedTab)
-				->AddTab(TabNamesPrivate::UserPresetWidget, ETabState::InvalidTab);
+				->AddTab(TabNamesPrivate::DocConfigWidget, ETabState::OpenedTab);
 
-			if (Builder.IsValid() && Builder->IsPreset())
-			{
-				GraphCanvasTabStack->SetHideTabWell(false);
-				GraphCanvasTabStack->AddTab(TabNamesPrivate::UserPresetWidget, ETabState::OpenedTab);
-			}
-			else
-			{
-				GraphCanvasTabStack->AddTab(TabNamesPrivate::UserPresetWidget, ETabState::InvalidTab);
-			}
-
-			const TSharedRef<FTabManager::FLayout> StandaloneDefaultLayout = FTabManager::NewLayout("Standalone_MetasoundEditor_Layout_v15")
+			const TSharedRef<FTabManager::FLayout> StandaloneDefaultLayout = FTabManager::NewLayout("Standalone_MetasoundEditor_Layout_v16")
 			->AddArea
 			(
 				FTabManager::NewPrimaryArea()
@@ -1326,7 +1346,6 @@ namespace Metasound
 			}
 
 			RefreshExecVisibility(Builder->GetConstBuilder().GetBuildPageID());
-			FSlateApplication::Get().SetUserFocus(0, MetasoundGraphEditor);
 
 			FMetaSoundAssetManager& AssetManager = FMetaSoundAssetManager::GetChecked();
 			const bool bAssetScanComplete = AssetManager.IsInitialAssetScanComplete();
@@ -1335,6 +1354,8 @@ namespace Metasound
 			{
 				NotifyAssetLoadingInProgress(NumProcessingAssets);
 			}
+
+			UpdateEnabledTabs();
 		}
 
 		UObject* FEditor::GetMetasoundObject() const
@@ -1406,9 +1427,17 @@ namespace Metasound
 				if (DocBuilder.IsValid())
 				{
 					const FMetasoundFrontendDocument& Document = DocBuilder.GetConstDocumentChecked();
+					if (const FMetaSoundFrontendDocumentConfiguration* DocConfig = Document.Configuration.GetPtr())
+					{
+						if (!DocConfig->IsPageGraphEditingEnabled())
+						{
+							return false;
+						}
+					}
+
 					const bool bLastGraph = Document.RootGraph.GetConstGraphPages().Num() == 1;
-					const bool bHasProjectPageValues = !bLastGraph && Document.RootGraph.FindConstGraph(Frontend::DefaultPageID) != nullptr;
-					return Editor::PageEditorEnabled(DocBuilder, bHasProjectPageValues);
+					const bool bHasLocalPageValues = !bLastGraph && Document.RootGraph.FindConstGraph(Frontend::DefaultPageID) != nullptr;
+					return Editor::IsPageGraphEditingEnabled(DocBuilder, bHasLocalPageValues);
 				}
 			}
 
@@ -1524,7 +1553,7 @@ namespace Metasound
 
 			FSlateApplication::Get().DismissAllMenus();
 
-			// In case of undoing 'convert from preset' refresh toolbar to include ConvertFromPreset button
+			// In case of undoing 'clear document config' refresh toolbar to include button
 			if (UToolMenus* ToolMenus = UToolMenus::Get())
 			{
 				ToolMenus->RefreshAllWidgets();
@@ -1756,9 +1785,9 @@ namespace Metasound
 			if (Builder.IsValid() && !Builder->GetConstBuilder().GetIsGraphEditable())
 			{
 				FLevelEditorModule& LevelEditor = FModuleManager::LoadModuleChecked<FLevelEditorModule>("LevelEditor");
-				LevelEditor.OnMapChanged().AddSP(this, &FEditor::ChangeUserPresetWidgetWorld);
+				LevelEditor.OnMapChanged().AddSP(this, &FEditor::ChangeDocConfigWidgetWorld);
 
-				UserPresetWidget = TStrongObjectPtr<UUserWidget>(CreateUserPresetWidget());
+				DocConfigWidget = TStrongObjectPtr<UUserWidget>(CreateDocConfigWidget());
 			}
 
 			Palette = SNew(SMetasoundPalette);
@@ -1799,7 +1828,7 @@ namespace Metasound
 			}
 		}
 
-		UUserWidget* FEditor::CreateUserPresetWidget()
+		UUserWidget* FEditor::CreateDocConfigWidget()
 		{
 			using namespace Metasound::Frontend;
 
@@ -1912,45 +1941,45 @@ namespace Metasound
 
 		void FEditor::OnPresetWidgetBlueprintCompiled(UBlueprint* Blueprint)
 		{
-			if (!Blueprint || !UserPresetWidget)
+			if (!Blueprint || !DocConfigWidget)
 			{
 				return;
 			}
-			if (UBlueprint* PresetWidgetBlueprint = UBlueprint::GetBlueprintFromClass(UserPresetWidget->GetClass()))
+			if (UBlueprint* PresetWidgetBlueprint = UBlueprint::GetBlueprintFromClass(DocConfigWidget->GetClass()))
 			{
 				if (PresetWidgetBlueprint == Blueprint)
 				{
 					// Regenerate user preset widget
-					TSharedPtr<SDockTab> UserPresetWidgetTab = TabManager->FindExistingLiveTab(TabNamesPrivate::UserPresetWidget);
-					if (UserPresetWidgetTab.IsValid())
+					TSharedPtr<SDockTab> DocConfigWidgetTab = TabManager->FindExistingLiveTab(TabNamesPrivate::DocConfigWidget);
+					if (DocConfigWidgetTab.IsValid())
 					{
-						UserPresetWidget = TStrongObjectPtr<UUserWidget>(CreateUserPresetWidget());
-						if (UserPresetWidget.IsValid())
+						DocConfigWidget = TStrongObjectPtr<UUserWidget>(CreateDocConfigWidget());
+						if (DocConfigWidget.IsValid())
 						{
-							IMetaSoundPresetWidgetInterface::Execute_OnConstructed(UserPresetWidget.Get(), Builder.Get());
-							UserPresetWidgetTab->SetContent(UserPresetWidget.Get()->TakeWidget());
+							IMetaSoundPresetWidgetInterface::Execute_OnConstructed(DocConfigWidget.Get(), Builder.Get());
+							DocConfigWidgetTab->SetContent(DocConfigWidget.Get()->TakeWidget());
 						}
 					}
 				}
 			}
 		}
 		
-		void FEditor::ChangeUserPresetWidgetWorld(UWorld* World, EMapChangeType MapChangeType)
+		void FEditor::ChangeDocConfigWidgetWorld(UWorld* World, EMapChangeType MapChangeType)
 		{
 			if (MapChangeType == EMapChangeType::TearDownWorld)
 			{
 				// We need to Delete the UMG widget if we are tearing down the World it was built with.
-				if (UserPresetWidget && World == UserPresetWidget->GetWorld())
+				if (DocConfigWidget && World == DocConfigWidget->GetWorld())
 				{
-					UserPresetWidget->Rename(nullptr, GetTransientPackage());
-					UserPresetWidget = nullptr;
+					DocConfigWidget->Rename(nullptr, GetTransientPackage());
+					DocConfigWidget = nullptr;
 				}
 			}
 			else if (MapChangeType != EMapChangeType::SaveMap)
 			{
 				// Recreate the widget if loading a map or opening a new map
 				// since the widget is parented to the world
-				UserPresetWidget = TStrongObjectPtr<UUserWidget>(CreateUserPresetWidget());
+				DocConfigWidget = TStrongObjectPtr<UUserWidget>(CreateDocConfigWidget());
 			}
 		}
 
@@ -1964,19 +1993,22 @@ namespace Metasound
 			TWeakObjectPtr<UMetaSoundBuilderBase> WeakBuilderPtr(Builder.Get());
 			MenuWidget->SetVisibility(TAttribute<EVisibility>::Create([WeakBuilderPtr]()
 			{
-				if (const TStrongObjectPtr<UMetaSoundBuilderBase> BuilderPtr = WeakBuilderPtr.Pin())
-				{
-					constexpr bool bHasProjectPageValues = true;
-					constexpr bool bPresetCanEditPageValues = true;
-					const bool bIsEnabled = PageEditorEnabled(BuilderPtr->GetConstBuilder(), bHasProjectPageValues, bPresetCanEditPageValues);
-					return bIsEnabled ? EVisibility::Visible : EVisibility::Collapsed;
-				}
-
-				return EVisibility::Collapsed;
+				return IsPageAuditioningEnabled() ? EVisibility::Visible : EVisibility::Collapsed;
 			}));
 			return MenuWidget;
 		}
 
+		TSharedRef<SWidget> FEditor::CreateConfigurationMenuOptions()
+		{
+			TSharedPtr<FUICommandList> Commands = MakeShared<FUICommandList>();
+			constexpr bool bShouldCloseWindowAfterMenuSelection = true;
+			FMenuBuilder MenuBuilder(bShouldCloseWindowAfterMenuSelection, Commands);
+			CreateConfigurationSubMenuOptions(MenuBuilder);
+			TSharedRef<SWidget> MenuWidget = MenuBuilder.MakeWidget();
+			TWeakObjectPtr<UMetaSoundBuilderBase> WeakBuilderPtr(Builder.Get());
+			return MenuWidget;
+		}
+
 		void FEditor::CreateAuditionPageSubMenuOptions(FMenuBuilder& MenuBuilder)
 		{
 			const UMetaSoundSettings* Settings = GetDefault<UMetaSoundSettings>();
@@ -2195,6 +2227,96 @@ namespace Metasound
 			MenuBuilder.EndSection();
 		}
 
+		void FEditor::CreateConfigurationSubMenuOptions(FMenuBuilder& MenuBuilder)
+		{
+			TAttribute<EVisibility> ClearVisibility = TAttribute<EVisibility>::CreateSPLambda(this, [this]()
+			{
+				const IMetasoundEditorModule & MetaSoundEditorModule = FModuleManager::GetModuleChecked<IMetasoundEditorModule>("MetaSoundEditor");
+				if (!MetaSoundEditorModule.IsRestrictedMode())
+				{
+					return GetConfiguration() ? EVisibility::Visible : EVisibility::Collapsed;
+				}
+
+				return EVisibility::Collapsed;
+			});
+
+			MenuBuilder.BeginSection("ConvertConfigurationHeader", LOCTEXT("ConvertConfigurationDisplayName", "Convert Configuration"), ClearVisibility);
+			{
+				if (const FMetaSoundFrontendDocumentBuilder* DocBuilder = GetFrontendBuilder())
+				{
+					if (const UScriptStruct* ConfigStruct = DocBuilder->GetConstDocumentConfiguration().GetScriptStruct())
+					{
+						MenuBuilder.AddMenuEntry(
+							FText::Format(LOCTEXT("ClearConfigurationButtonText", "Convert From '{0}' ..."), ConfigStruct->GetDisplayNameText()),
+							LOCTEXT("ConvertFromConfigurationTooltip", "Removes the MetaSound's configuration, leaving the MetaSound state intact and enabling editing of configuration-managed MetaSound functionality."),
+							FSlateIcon(),
+							FUIAction(FExecuteAction::CreateSP(this, &FEditor::ShowConvertFromConfigDialog)),
+							{ },
+							EUserInterfaceActionType::Button
+						);
+					}
+				}
+			}
+			MenuBuilder.EndSection();
+
+			MenuBuilder.BeginSection("SetConfigurationSectionHeader", LOCTEXT("SelectConfigurationSectionName", "Set Configuration"));
+			{
+				auto CreateConfigEntry = [this, &MenuBuilder](UScriptStruct* ConfigStruct)
+				{
+					FUIAction SetConfigAction;
+					SetConfigAction.ExecuteAction = FExecuteAction::CreateSPLambda(this, [this, ConfigStruct]()
+					{
+						if (FMetaSoundFrontendDocumentBuilder* DocBuilder = GetFrontendBuilder())
+						{
+							TInstancedStruct<FMetaSoundFrontendDocumentConfiguration> ConfigInstance;
+
+							check(ConfigStruct);
+							ConfigInstance.InitializeAsScriptStruct(ConfigStruct);
+
+							constexpr bool bLogErrors = true;
+							DocBuilder->SetDocumentConfiguration(MoveTemp(ConfigInstance));
+							DocBuilder->ConfigureDocument(bLogErrors);
+							UpdateEnabledTabs();
+						}
+					});
+					SetConfigAction.GetActionCheckState = FGetActionCheckState::CreateSPLambda(this, [this, ConfigStruct]()
+					{
+						if (const FMetaSoundFrontendDocumentBuilder* DocBuilder = GetFrontendBuilder())
+						{
+							if (ConfigStruct == DocBuilder->GetConstDocumentConfiguration().GetScriptStruct())
+							{
+								return ECheckBoxState::Checked;
+							}
+						}
+
+						return ECheckBoxState::Unchecked;
+					});
+
+					MenuBuilder.AddMenuEntry(
+						ConfigStruct->GetDisplayNameText(),
+						ConfigStruct->GetToolTipText(),
+						FSlateIcon(),
+						SetConfigAction,
+						{ },
+						EUserInterfaceActionType::RadioButton
+					);
+				};
+
+				for (TObjectIterator<UStruct> It; It; ++It)
+				{
+					if (UScriptStruct* Struct = Cast<UScriptStruct>(*It))
+					{
+						if (Struct->IsChildOf(FMetaSoundFrontendDocumentConfiguration::StaticStruct())
+							&& Struct != FMetaSoundFrontendDocumentConfiguration::StaticStruct())
+						{
+							CreateConfigEntry(Struct);
+						}
+					}
+				}
+			}
+			MenuBuilder.EndSection();
+		}
+
 		void FEditor::DestroyAnalyzers()
 		{
 			if (AnalyzerRack.IsValid())
@@ -2213,42 +2335,16 @@ namespace Metasound
 				GetToolkitCommands(),
 				FToolBarExtensionDelegate::CreateLambda([this](FToolBarBuilder& ToolbarBuilder)
 				{
-					// TODO: Clean-up json importer/exporter and re-enable this
  					ToolbarBuilder.BeginSection("Utilities");
  					{
-// 						ToolbarBuilder.AddToolBarButton
-// 						(
-// 							FEditorCommands::Get().Import,
-// 							NAME_None,
-// 							TAttribute<FText>(),
-// 							TAttribute<FText>(),
-// 							TAttribute<FSlateIcon>::Create([this]() { return GetImportStatusImage(); }),
-// 							"ImportMetasound"
-// 						);
-// 
-// 						ToolbarBuilder.AddToolBarButton
-// 						(
-// 							FEditorCommands::Get().Export,
-// 							NAME_None,
-// 							TAttribute<FText>(),
-// 							TAttribute<FText>(),
-// 							TAttribute<FSlateIcon>::Create([this]() { return GetExportStatusImage(); }),
-// 							"ExportMetasound"
-// 						);
-
-						const IMetasoundEditorModule& MetaSoundEditorModule = FModuleManager::GetModuleChecked<IMetasoundEditorModule>("MetaSoundEditor");
-						if (!IsGraphEditable() && !MetaSoundEditorModule.IsRestrictedMode())
-						{
-							ToolbarBuilder.AddToolBarButton
- 							(
- 								FEditorCommands::Get().ConvertFromPreset,
- 								NAME_None,
- 								TAttribute<FText>(),
- 								TAttribute<FText>(),
- 								TAttribute<FSlateIcon>::Create([this]() { return GetExportStatusImage(); }),
- 								"ConvertFromPreset"
-							);
-						}
+						ToolbarBuilder.AddComboButton(
+							FUIAction(),
+							FOnGetContent::CreateSP(this, &FEditor::CreateConfigurationMenuOptions),
+							LOCTEXT("ConfigurationSettingsMenu", "Configuration"),
+							LOCTEXT("ConfigurationSettingsMenu_Tooltip", "Settings related to procedurally-generated MetaSound Configurations"),
+							GetExportStatusImage(),
+							false
+						);
  					}
  					ToolbarBuilder.EndSection();
 
@@ -2305,9 +2401,7 @@ namespace Metasound
 							}
 							ToolbarBuilder.EndStyleOverride();
 
-							constexpr bool bHasProjectPageValues = true;
-							constexpr bool bPresetCanEditPageValues = true;
-							const bool bShowAuditionSettings = PageEditorEnabled(Builder->GetConstBuilder(), bHasProjectPageValues, bPresetCanEditPageValues);
+							const bool bShowAuditionSettings = IsPageAuditioningEnabled();
 							if (bShowAuditionSettings)
 							{
 								ToolbarBuilder.AddComboButton(
@@ -2319,6 +2413,17 @@ namespace Metasound
 									false
 								);
 							}
+
+							ToolbarBuilder.AddToolBarButton(
+								FEditorCommands::Get().ToggleShowGraph,
+								NAME_None,
+								TAttribute<FText>::CreateSPLambda(this, [this]() { return AssetEditorPrivate::bIgnoreDocConfigVisibilitySettings
+									? LOCTEXT("ShowGraphLabel", "Hide Graph")
+									: LOCTEXT("ShowGraphLabel", "Show Graph");
+								}),
+								TAttribute<FText>(),
+								Style::CreateSlateIcon("MetasoundEditor.Settings")
+							);
 						}
 						ToolbarBuilder.EndSection();
 					}
@@ -2380,6 +2485,11 @@ namespace Metasound
 				Commands.TogglePlayback,
 				FExecuteAction::CreateSP(this, &FEditor::TogglePlayback));
 
+			ToolkitCommands->MapAction(
+				Commands.ToggleShowGraph,
+				FExecuteAction::CreateSP(this, &FEditor::ToggleShowHideGraph),
+				FCanExecuteAction::CreateSP(this, &FEditor::CanToggleShowHideGraph));
+
 			ToolkitCommands->MapAction(
 				FGenericCommands::Get().Undo,
 				FExecuteAction::CreateSP(this, &FEditor::UndoGraphAction));
@@ -2396,10 +2506,6 @@ namespace Metasound
 				Commands.EditSourceSettings,
 				FExecuteAction::CreateSP(this, &FEditor::EditSourceSettings));
 
-			ToolkitCommands->MapAction(
-				Commands.ConvertFromPreset,
-				FExecuteAction::CreateSP(this, &FEditor::ConvertFromPreset));
-
 			ToolkitCommands->MapAction(FGenericCommands::Get().Delete,
 				FExecuteAction::CreateSP(this, &FEditor::DeleteSelectedInterfaceItems),
 				FCanExecuteAction::CreateSP(this, &FEditor::CanDeleteInterfaceItems));
@@ -2431,6 +2537,15 @@ namespace Metasound
 			ToolkitCommands->MapAction(
 				FEditorCommands::Get().FindInMetaSound,
 				FExecuteAction::CreateSP(this, &FEditor::ShowFindInMetaSound));
+
+			ToolkitCommands->MapAction(
+				FEditorCommands::Get().CopySelectionAsConfigurationCodeToClipboard,
+				FExecuteAction::CreateSP(this, &FEditor::CopySelectionAsConfigurationCodeToClipboard));
+
+			ToolkitCommands->MapAction(
+				FEditorCommands::Get().ToggleShowGraph,
+				FExecuteAction::CreateSP(this, &FEditor::ToggleShowHideGraph),
+				FCanExecuteAction::CreateSP(this, &FEditor::CanToggleShowHideGraph));
 		}
 
 		void FEditor::Import()
@@ -2633,10 +2748,10 @@ namespace Metasound
 							CastChecked<UMetaSoundSourceBuilder>(Builder.Get())->Audition(nullptr, PreviewComp, GeneratorHandleDelegate, AssetEditorPrivate::bEnableLiveAudition);
 						}
 
-						if (UserPresetWidget.IsValid() && IsPlaying())
+						if (DocConfigWidget.IsValid() && IsPlaying())
 						{
 							constexpr bool bIsAuditioning = true;
-							IMetaSoundPresetWidgetInterface::Execute_OnAuditionStateChanged(UserPresetWidget.Get(), PreviewComp, bIsAuditioning);
+							IMetaSoundPresetWidgetInterface::Execute_OnAuditionStateChanged(DocConfigWidget.Get(), PreviewComp, bIsAuditioning);
 						}
 
 						SetPreviewID(PreviewComp->GetUniqueID());
@@ -2665,10 +2780,10 @@ namespace Metasound
 							UpdatePageInfo(bIsPlaying);
 							GraphConnectionManager = RebuildConnectionManager();
 
-							if (UserPresetWidget.IsValid())
+							if (DocConfigWidget.IsValid())
 							{
 								constexpr bool bIsAuditioning = false;
-								IMetaSoundPresetWidgetInterface::Execute_OnAuditionStateChanged(UserPresetWidget.Get(), nullptr, bIsAuditioning);
+								IMetaSoundPresetWidgetInterface::Execute_OnAuditionStateChanged(DocConfigWidget.Get(), nullptr, bIsAuditioning);
 							}
 
 							return EActiveTimerReturnType::Stop;
@@ -2786,6 +2901,26 @@ namespace Metasound
 			}
 		}
 
+		bool FEditor::CanToggleShowHideGraph() const
+		{
+			if (Builder.IsValid())
+			{
+				if (const FMetaSoundFrontendDocumentBuilder& DocBuilder = Builder->GetBuilder(); ensure(DocBuilder.IsValid()))
+				{
+					return DocBuilder.GetConstDocumentConfiguration().IsValid();
+				}
+			}
+
+			return false;
+		}
+
+
+		void FEditor::ToggleShowHideGraph()
+		{
+			AssetEditorPrivate::bIgnoreDocConfigVisibilitySettings = !AssetEditorPrivate::bIgnoreDocConfigVisibilitySettings;
+			UpdateEnabledTabs();
+		}
+
 		void FEditor::ExecuteNode(UEdGraphNode* InNode)
 		{
 			using namespace Metasound;
@@ -2843,7 +2978,7 @@ namespace Metasound
 			SetSelection({ GetMetasoundObject() });
 		}
 
-		void FEditor::ConvertFromPreset()
+		void FEditor::ShowConvertFromConfigDialog()
 		{
 			using namespace Frontend;
 
@@ -2853,7 +2988,7 @@ namespace Metasound
 			{
 				TSharedPtr<SWindow> DialogWindow =
 					SNew(SWindow)
-					.Title(LOCTEXT("MetasoundPresetDialogTitle", "Convert From Preset?"))
+					.Title(LOCTEXT("ConvertMetaSoundConfigurationPrompt", "Convert From MetaSound Configuration?"))
 					.SupportsMinimize(false)
 					.SupportsMaximize(false)
 					.SizingRule(ESizingRule::Autosized)
@@ -2871,15 +3006,21 @@ namespace Metasound
 						.VAlign(VAlign_Bottom)
 						[
 							SNew(SButton)
-							.Text(LOCTEXT("MetasoundPresetDialogAccept", "Accept"))
+							.Text(LOCTEXT("MetasoundConvertFromConfigDialogAccept", "Accept"))
 							.OnClicked_Lambda([this, &DialogWindow]()
 							{
-								const FScopedTransaction Transaction(LOCTEXT("ConvertFromPresetText", "Convert From Preset"));
-								GetMetasoundObject()->Modify();
-								
-								EMetaSoundBuilderResult Result = EMetaSoundBuilderResult::Failed;
-								Builder->ConvertFromPreset(Result);
-								ensure(Result == EMetaSoundBuilderResult::Succeeded);
+								if (FMetaSoundFrontendDocumentBuilder& DocBuilder = Builder->GetBuilder(); ensure(DocBuilder.IsValid()))
+								{
+									const UScriptStruct* ScriptStruct = DocBuilder.GetConstDocumentConfiguration().GetScriptStruct();
+									const FText DisplayName = ScriptStruct ? ScriptStruct->GetDisplayNameText() : LOCTEXT("EmptyConfigName", "Unset");
+									const FScopedTransaction Transaction(FText::Format(LOCTEXT("ConvertFromConfigurationFormat", "Convert From '{0}' Configuration"), DisplayName));
+
+									GetMetasoundObject()->Modify();
+									EMetaSoundBuilderResult Result = EMetaSoundBuilderResult::Failed;
+										DocBuilder.SetDocumentConfiguration({});
+										DocBuilder.ConfigureDocument();
+										UpdateEnabledTabs();
+								}
 								
 								if (UToolMenus* ToolMenus = UToolMenus::Get())
 								{
@@ -2900,7 +3041,7 @@ namespace Metasound
 						.VAlign(VAlign_Bottom)
 						[
 							SNew(SButton)
-							.Text(LOCTEXT("MetasoundPresetDialogCancel", "Cancel"))
+							.Text(LOCTEXT("MetasoundConvertFromConfigDialogCancel", "Cancel"))
 							.OnClicked_Lambda([&DialogWindow]()
 							{
 								DialogWindow->RequestDestroyWindow();
@@ -3052,6 +3193,9 @@ namespace Metasound
 					FExecuteAction::CreateLambda([this] { PromoteToDeferredVariable(); }),
 					FCanExecuteAction::CreateLambda([this]() { return CanPromoteToDeferredVariable(); }));
 
+				GraphEditorCommands->MapAction(FEditorCommands::Get().CopySelectionAsConfigurationCodeToClipboard,
+					FExecuteAction::CreateSP(this, &FEditor::CopySelectionAsConfigurationCodeToClipboard));
+
 				GraphEditorCommands->MapAction(FGraphEditorCommands::Get().HideNoConnectionPins,
 					FExecuteAction::CreateSP(this, &FEditor::HideUnconnectedPins));
 
@@ -3096,6 +3240,9 @@ namespace Metasound
 
 				GraphEditorCommands->MapAction(FEditorCommands::Get().UpdateNodeClass,
 					FExecuteAction::CreateSP(this, &FEditor::UpdateSelectedNodeClasses));
+
+				GraphEditorCommands->MapAction(FEditorCommands::Get().CopySelectionAsConfigurationCodeToClipboard,
+					FExecuteAction::CreateSP(this, &FEditor::CopySelectionAsConfigurationCodeToClipboard));
 			}
 
 			SGraphEditor::FGraphEditorEvents GraphEvents;
@@ -3645,14 +3792,11 @@ namespace Metasound
 			return nullptr;
 		}
 
-		const FMetaSoundFrontendDocumentConfiguration* FEditor::GetConfiguration() const
+		const Frontend::FDocumentConfiguration* FEditor::GetConfiguration() const
 		{
 			if (const FMetaSoundFrontendDocumentBuilder* DocBuilder = GetFrontendBuilder())
 			{
-				if (DocBuilder->IsValid())
-				{
-					return DocBuilder->GetConstDocumentConfiguration().GetPtr();
-				}
+				return DocBuilder->GetConstDocumentConfiguration().GetPtr();
 			}
 
 			return nullptr;
@@ -3962,6 +4106,19 @@ namespace Metasound
 			DeleteSelectedInterfaceItems();
 		}
 
+		void FEditor::CopySelectionAsConfigurationCodeToClipboard() const
+		{
+			if (const FMetaSoundFrontendDocumentBuilder* DocBuilder = GetFrontendBuilder())
+			{
+				const FGraphPanelSelectionSet Selection = MetasoundGraphEditor->GetSelectedNodes();
+				FMetasoundFrontendGraphLayer Layer = FDocumentClipboardUtils::CopySelectionToLayer(*DocBuilder, Selection);
+
+				FString CodeString;
+				Engine::FDocumentConfigurator::Reflect(MoveTemp(Layer), CodeString);
+				FPlatformApplicationMisc::ClipboardCopy(*CodeString);
+			}
+		}
+
 		bool FEditor::CanPasteSelectedMemberItems() const
 		{
 			if (!GraphMembersMenu.IsValid())
@@ -4799,18 +4956,21 @@ PRAGMA_ENABLE_DEPRECATION_WARNINGS
 
 				if (!MembersModified.IsEmpty() || bForceRefreshViews)
 				{
-					UMetasoundEditorGraphMember* SelectedMember = RefreshGraphMemberMenu();
-
-					// If no member was selected by an action (ex. undo/redo), select a modified member 
-					if (!SelectedMember)
+					if (const Frontend::FDocumentConfiguration* Config = GetConfiguration(); !Config || Config->IsMemberEditingEnabled())
 					{
-						for (const FGuid& MemberGuid : MembersModified)
+						UMetasoundEditorGraphMember* SelectedMember = RefreshGraphMemberMenu();
+
+						// If no member was selected by an action (ex. undo/redo), select a modified member 
+						if (!SelectedMember)
 						{
-							if (UObject* Member = Graph->FindMember(MemberGuid))
+							for (const FGuid& MemberGuid : MembersModified)
 							{
-								// Currently only one member can be selected at a time, so only first found is added
-								Selection.Add(Member);
-								break;
+								if (UObject* Member = Graph->FindMember(MemberGuid))
+								{
+									// Currently only one member can be selected at a time, so only first found is added
+									Selection.Add(Member);
+									break;
+								}
 							}
 						}
 					}
@@ -4951,6 +5111,16 @@ PRAGMA_ENABLE_DEPRECATION_WARNINGS
 			}
 		}
 
+		int32 FEditor::GetNumNodesSelected() const
+		{
+			if (MetasoundGraphEditor.IsValid())
+			{
+				return MetasoundGraphEditor->GetSelectedNodes().Num();
+			}
+
+			return 0;
+		}
+
 		TSharedRef<SWidget> FEditor::OnGetMenuSectionWidget(TSharedRef<SWidget> RowWidget, int32 InSectionID)
 		{
 			TWeakPtr<SWidget> WeakRowWidget = RowWidget;
@@ -5151,6 +5321,96 @@ PRAGMA_ENABLE_DEPRECATION_WARNINGS
 			bMemberRenameRequested = true;
 		}
 
+		void FEditor::UpdateEnabledTabs()
+		{
+			if (!TabManager)
+			{
+				return;
+			}
+
+			auto ResolveTabEnablement = [this](FName TabName, bool bEnabled)
+			{
+				if (TSharedPtr<FTabSpawnerEntry> Spawner = TabManager->FindTabSpawnerFor(TabName))
+				{
+					Spawner->SetMenuType(bEnabled ? ETabSpawnerMenuType::Enabled : ETabSpawnerMenuType::Hidden);
+				}
+
+				if (TSharedPtr<SDockTab> Tab = TabManager->FindExistingLiveTab(TabName))
+				{
+					if (!bEnabled)
+					{
+						Tab->RequestCloseTab();
+					}
+				}
+				else
+				{
+					if (bEnabled)
+					{
+						TabManager->TryInvokeTab(TabName, false /* bInvokeAsInactive */);
+					}
+				}
+			};
+
+			auto RefreshViews = [this]()
+			{
+				RefreshGraphMemberMenu();
+				RefreshInterfaceView();
+				RefreshDetails();
+			};
+
+			if (FMetaSoundFrontendDocumentBuilder* DocBuilder = GetFrontendBuilder())
+			{
+				const bool bHasProjectPageValues = DocBuilder->GetConstDocumentChecked().RootGraph.GetConstGraphPages().Num() > 1;
+				const bool bEnablePageTab = IsPageGraphEditingEnabled(*DocBuilder, bHasProjectPageValues);
+
+				bool bVisibilitySet = false;
+				if (const FMetaSoundFrontendDocumentConfiguration* Config = GetConfiguration())
+				{
+					if (!AssetEditorPrivate::bIgnoreDocConfigVisibilitySettings)
+					{
+						ResolveTabEnablement(TabNamesPrivate::Find, Config->IsGraphEditorVisible());
+						ResolveTabEnablement(TabNamesPrivate::GraphCanvas, Config->IsGraphEditorVisible());
+						ResolveTabEnablement(TabNamesPrivate::Interfaces, Config->IsInterfaceEditingEnabled());
+						ResolveTabEnablement(TabNamesPrivate::Members, Config->IsMemberEditingEnabled());
+
+						// Don't show custom config tab if no config widget was constructed, even if enabled
+						ResolveTabEnablement(TabNamesPrivate::DocConfigWidget, DocConfigWidget.IsValid() && Config->IsCustomTabEnabled());
+
+						RefreshViews();
+						bVisibilitySet = true;
+					}
+				}
+
+				if (!bVisibilitySet)
+				{
+					{
+						constexpr bool bEnableTab = true;
+						ResolveTabEnablement(TabNamesPrivate::Find, bEnableTab);
+						ResolveTabEnablement(TabNamesPrivate::GraphCanvas, bEnableTab);
+						ResolveTabEnablement(TabNamesPrivate::Interfaces, bEnableTab);
+						ResolveTabEnablement(TabNamesPrivate::Members, bEnableTab);
+					}
+
+					{
+						// User preset widgets are only visible when doc config is set
+						constexpr bool bEnableTab = false;
+						ResolveTabEnablement(TabNamesPrivate::DocConfigWidget, bEnableTab);
+					}
+				}
+
+				// Pages need additional checks to determine edit eligibility
+				ResolveTabEnablement(TabNamesPrivate::Pages, bEnablePageTab);
+
+				{
+					UObject* MetaSoundObject = GetMetasoundObject();
+					const bool bEnableTab = MetaSoundObject && MetaSoundObject->IsA<USoundBase>();
+					ResolveTabEnablement(TabNamesPrivate::Analyzers, bEnableTab);
+				}
+
+				RefreshViews();
+			}
+		}
+
 		TUniquePtr<FGraphConnectionManager> FEditor::RebuildConnectionManager(UAudioComponent* PreviewComp) const
 		{
 			using namespace Engine;
@@ -5255,6 +5515,15 @@ PRAGMA_ENABLE_DEPRECATION_WARNINGS
 			OutDelegates.PageDelegates.OnPageSet.AddSP(this, &FDocumentListener::OnPageSet);
 			OutDelegates.InterfaceDelegates.OnInputDefaultChanged.AddSP(this, &FDocumentListener::OnInputDefaultChanged);
 			OutDelegates.InterfaceDelegates.OnMemberMetadataSet.AddSP(this, &FDocumentListener::OnMemberMetadataSet);
+			OutDelegates.OnDocumentConfigurationChanged.AddSP(this, &FDocumentListener::OnDocumentConfigurationChanged);
+		}
+
+		void FEditor::FDocumentListener::OnDocumentConfigurationChanged(const Frontend::FDocumentConfigChangedArgs&)
+		{
+			if (TSharedPtr<FEditor> ParentPtr = Parent.Pin())
+			{
+				ParentPtr->UpdateEnabledTabs();
+			}
 		}
 
 		void FEditor::FDocumentListener::OnPageSet(const Frontend::FDocumentMutatePageArgs& Args)
diff --git a/Engine/Plugins/Runtime/Metasound/Source/MetasoundEditor/Private/MetasoundEditor.h b/Engine/Plugins/Runtime/Metasound/Source/MetasoundEditor/Private/MetasoundEditor.h
--- a/Engine/Plugins/Runtime/Metasound/Source/MetasoundEditor/Private/MetasoundEditor.h
+++ b/Engine/Plugins/Runtime/Metasound/Source/MetasoundEditor/Private/MetasoundEditor.h
@@ -34,39 +34,38 @@
 
 
 // Forward Declarations
-namespace AudioWidgets
-{
-	class FAudioAnalyzerRack;
-}
-
+class FSlateRect;
 class FTabManager;
+class IDetailsView;
+class IToolkitHost;
 class SDockableTab;
 class SGraphEditor;
 class SMetasoundPalette;
-class FSlateRect;
-class IDetailsView;
-class IToolkitHost;
-class UUserWidget;
-
-struct FMetaSoundFrontendDocumentBuilder;
-struct FMetaSoundFrontendDocumentConfiguration;
-
-namespace Metasound::Editor
-{
-	class SFindInMetasound;
-	class SRenderStats;
-}
 class SVerticalBox;
 class UAudioComponent;
 class UEdGraphNode;
-class UMetaSoundPatch;
-class UMetaSoundSource;
 class UMetasoundEditorGraph;
 class UMetasoundEditorGraphNode;
+class UMetaSoundPatch;
+class UMetaSoundSource;
+class UUserWidget;
 
 struct FGraphActionNode;
+struct FMetaSoundFrontendDocumentBuilder;
 struct FPropertyChangedEvent;
 
+// Forward Declarations
+namespace AudioWidgets
+{
+	class FAudioAnalyzerRack;
+} // AudioWidgets
+
+namespace Metasound::Editor
+{
+	class SFindInMetasound;
+	class SRenderStats;
+} // namespace Metasound::Editor
+
 
 // Base implementation for editor wrappers that provide edit customizations
 UCLASS(MinimalAPI)
@@ -101,7 +100,10 @@ namespace Metasound::Editor
 {
 	bool IsPreviewingPageInputDefault(const FMetaSoundFrontendDocumentBuilder& Builder, const FMetasoundFrontendClassInput& InClassInput, const FGuid& InPageID);
 	bool IsPreviewingPageGraph(const FMetaSoundFrontendDocumentBuilder& Builder, const FGuid& InPageID);
-	bool PageEditorEnabled(const FMetaSoundFrontendDocumentBuilder& Builder, bool bHasProjectPageValues, bool bPresetCanEditPageValues = false);
+
+	bool IsPageAuditioningEnabled(const FMetaSoundFrontendDocumentBuilder& Builder, bool bHasProjectPageValues);
+	bool IsPageGraphEditingEnabled(const FMetaSoundFrontendDocumentBuilder& Builder, bool bHasProjectPageValues);
+	bool IsInputPageEditingEnabled(const FMetaSoundFrontendDocumentBuilder& Builder, bool bHasProjectPageValues);
 
 	// Forward Declarations
 	class FMetasoundGraphMemberSchemaAction;
@@ -136,7 +138,7 @@ namespace Metasound::Editor
 
 		UAudioComponent* GetAudioComponent() const;
 
-		const FMetaSoundFrontendDocumentConfiguration* GetConfiguration() const;
+		const Frontend::FDocumentConfiguration* GetConfiguration() const;
 		FMetaSoundFrontendDocumentBuilder* GetFrontendBuilder() const;
 
 		/** IMetasoundEditor interface */
@@ -208,10 +210,7 @@ namespace Metasound::Editor
 
 		void ClearSelectionAndSelectNode(UEdGraphNode* Node);
 
-		int32 GetNumNodesSelected() const
-		{
-			return MetasoundGraphEditor->GetSelectedNodes().Num();
-		}
+		int32 GetNumNodesSelected() const;
 
 		/** Creates analyzers */
 		void CreateAnalyzers(UMetaSoundSource& MetaSoundSource);
@@ -358,6 +357,9 @@ namespace Metasound::Editor
 		/** Whether the currently selected Member item(s) can be pasted. */
 		bool CanPasteSelectedMemberItems() const;
 
+		/** Copies the selected nodes, shared connections, and required interfaces to the clipboard as configurator c++ */
+		void CopySelectionAsConfigurationCodeToClipboard() const;
+
 		/** Paste selected Member items. */
 		void PasteSelectedMemberItems();
 
@@ -442,6 +444,8 @@ namespace Metasound::Editor
 
 		void RemoveInvalidSelection();
 
+		void UpdateEnabledTabs();
+
 		void SetPreviewID(uint32 InPreviewID);
 
 		void ExportNodesToText(FString& OutText) const;
@@ -454,11 +458,11 @@ namespace Metasound::Editor
 		/** Creates all internal widgets for the tabs to point at */
 		void CreateInternalWidgets(UObject& MetaSound);
 
-		UUserWidget* CreateUserPresetWidget();
+		UUserWidget* CreateDocConfigWidget();
 
 		/** For teardown and regeneration of user preset widget on world change
 			Based on UEditorUtilityWidgetBlueprint::ChangeTabWorld */
-		void ChangeUserPresetWidgetWorld(UWorld* World, EMapChangeType MapChangeType);
+		void ChangeDocConfigWidgetWorld(UWorld* World, EMapChangeType MapChangeType);
 
 		/** Reload preset widget tab on associated widget blueprint recompile */
 		void OnPresetWidgetBlueprintCompiled(UBlueprint* Blueprint);
@@ -486,14 +490,20 @@ namespace Metasound::Editor
 		/** Either play the Metasound or stop currently playing sound */
 		void TogglePlayback();
 
+		/** Whether or not to show the toggle for showing/hiding a graph (i.e. is or is not a MetaSound with a configuration set) */
+		bool CanToggleShowHideGraph() const;
+
+		/** Toggles showing the graph for the given configured MetaSound */
+		void ToggleShowHideGraph();
+
 		/** Executes specified node (If supported) */
 		void ExecuteNode(UEdGraphNode* Node);
 
 		/** Sync the content browser to the current selection of nodes */
 		void SyncInBrowser();
 
-		/** Converts the MetaSound from a preset to a fully modifiable MetaSound. */
-		void ConvertFromPreset();
+		/** Shows dialog confirming converting from configured document to make parent MetaSound openly modifiable. */
+		void ShowConvertFromConfigDialog();
 
 		/* Whether or not page details should be visible/enabled. */
 		bool ShowPageGraphDetails() const;
@@ -501,9 +511,15 @@ namespace Metasound::Editor
 		/** Creates audition menu options */
 		TSharedRef<SWidget> CreateAuditionMenuOptions();
 
+		/** Creates configuration menu options */
+		TSharedRef<SWidget> CreateConfigurationMenuOptions();
+
 		/** Creates page menu options */
 		void CreateAuditionPageSubMenuOptions(FMenuBuilder& MenuBuilder);
 
+		/** Creates configuration menu options */
+		void CreateConfigurationSubMenuOptions(FMenuBuilder& MenuBuilder);
+
 		/** Show the Metasound object's Source settings in the Details panel */
 		void EditSourceSettings();
 
@@ -579,8 +595,8 @@ namespace Metasound::Editor
 		/** Widget showing render performance information that overlays the graph */
 		TSharedPtr<SRenderStats> RenderStatsWidget;
 
-		/** User defined widget shown when editing a preset.*/
-		TStrongObjectPtr<UUserWidget> UserPresetWidget;
+		/** User defined widget shown when editing a doc config with user widget enabled.*/
+		TStrongObjectPtr<UUserWidget> DocConfigWidget;
 
 		TUniquePtr<FGraphConnectionManager> GraphConnectionManager;
 
@@ -627,6 +643,7 @@ namespace Metasound::Editor
 
 		private:
 			virtual void OnBuilderReloaded(Frontend::FDocumentModifyDelegates& OutDelegates) override;
+			void OnDocumentConfigurationChanged(const Frontend::FDocumentConfigChangedArgs& Args);
 			void OnPageSet(const Frontend::FDocumentMutatePageArgs& Args);
 			void OnInputDefaultChanged(int32 Index);
 			void OnMemberMetadataSet(FGuid MemberID);
diff --git a/Engine/Plugins/Runtime/Metasound/Source/MetasoundEditor/Private/MetasoundEditorCommands.cpp b/Engine/Plugins/Runtime/Metasound/Source/MetasoundEditor/Private/MetasoundEditorCommands.cpp
--- a/Engine/Plugins/Runtime/Metasound/Source/MetasoundEditor/Private/MetasoundEditorCommands.cpp
+++ b/Engine/Plugins/Runtime/Metasound/Source/MetasoundEditor/Private/MetasoundEditorCommands.cpp
@@ -13,6 +13,8 @@ namespace Metasound
 			UI_COMMAND(Stop, "Stop", "Stops MetaSound (If currently playing)", EUserInterfaceActionType::Button, FInputChord());
 			UI_COMMAND(TogglePlayback, "Toggle Playback", "Plays or stops the currently playing MetaSound", EUserInterfaceActionType::Button, FInputChord(EKeys::SpaceBar));
 
+			UI_COMMAND(ToggleShowGraph, "Show Graph", "Toggle showing/hiding the graph for the given configured MetaSound", EUserInterfaceActionType::Button, FInputChord(EKeys::G, EModifierKey::Control | EModifierKey::Shift));
+
 			UI_COMMAND(Import, "Import", "Imports MetaSound from Json", EUserInterfaceActionType::Button, FInputChord());
 			UI_COMMAND(Export, "Export", "Exports MetaSound to Json", EUserInterfaceActionType::Button, FInputChord());
 
@@ -24,7 +26,6 @@ namespace Metasound
 
 			UI_COMMAND(UpdateNodeClass, "Update Node Class", "Update selected node(s) class(es) that have an available update.", EUserInterfaceActionType::Button, FInputChord());
 
-			UI_COMMAND(ConvertFromPreset, "Convert From Preset", "Converts this preset to a fully accessible MetaSound.", EUserInterfaceActionType::Button, FInputChord());
 			UI_COMMAND(Delete, "Delete Selected", "Delete selected items.", EUserInterfaceActionType::None, FInputChord(EKeys::Delete));
 			UI_COMMAND(FindInMetaSound, "Find in MetaSound", "Find a node or pin within this MetaSound.", EUserInterfaceActionType::Button, FInputChord(EModifierKey::Control, EKeys::F));
 		
@@ -35,6 +36,12 @@ namespace Metasound
 
 			UI_COMMAND(PromoteAllToInput, "Promote All To Graph Input", "Promotes unconnected node inputs to graph inputs.", EUserInterfaceActionType::Button, FInputChord());
 			UI_COMMAND(PromoteAllToCommonInputs, "Promote All To Common Graph Inputs", "Promotes unconnected input pins from selected nodes to graph inputs, sharing inputs if possible.", EUserInterfaceActionType::Button, FInputChord());
+
+			UI_COMMAND(CopySelectionAsConfigurationCodeToClipboard,
+				"Copy Selected As Configurator Code to Clipboard",
+				"Reflect configurator C++ code to the clipboard for the given graph selection",
+				EUserInterfaceActionType::Button,
+				FInputChord(EKeys::C, EModifierKey::Control | EModifierKey::Shift));
 		}
 	} // namespace Editor
 } // namespace Metasound
diff --git a/Engine/Plugins/Runtime/Metasound/Source/MetasoundEditor/Private/MetasoundEditorDocumentClipboardUtils.cpp b/Engine/Plugins/Runtime/Metasound/Source/MetasoundEditor/Private/MetasoundEditorDocumentClipboardUtils.cpp
--- a/Engine/Plugins/Runtime/Metasound/Source/MetasoundEditor/Private/MetasoundEditorDocumentClipboardUtils.cpp
+++ b/Engine/Plugins/Runtime/Metasound/Source/MetasoundEditor/Private/MetasoundEditorDocumentClipboardUtils.cpp
@@ -9,7 +9,9 @@
 #include "HAL/PlatformApplicationMisc.h"
 #include "Logging/TokenizedMessage.h"
 #include "MetasoundDocumentBuilderRegistry.h"
+#include "MetasoundEditorGraph.h"
 #include "MetasoundEditorSubsystem.h"
+#include "MetasoundFrontendGraphLayer.h"
 #include "MetasoundFrontendSearchEngine.h"
 #include "Misc/StringOutputDevice.h"
 #include "NodeTemplates/MetasoundFrontendNodeTemplateInput.h"
@@ -18,6 +20,31 @@
 
 namespace Metasound::Editor
 {
+	namespace DocumentClipboardUtilsPrivate
+	{
+		class FMemberClipboardObjectTextFactory : public FCustomizableTextObjectFactory
+		{
+		public:
+			FMemberClipboardObjectTextFactory()
+				: FCustomizableTextObjectFactory(GWarn)
+			{
+			}
+
+			UMetasoundEditorGraphMember* Member = nullptr;
+
+		protected:
+			virtual bool CanCreateClass(UClass* InObjectClass, bool& bOmitSubObjs) const override
+			{
+				return InObjectClass->IsChildOf(UMetasoundEditorGraphMember::StaticClass());
+			}
+
+			virtual void ProcessConstructedObject(UObject* CreatedObject) override
+			{
+				Member = Cast<UMetasoundEditorGraphMember>(CreatedObject);
+			}
+		};
+	} // DocumentClipboardUtilsPrivate
+
 	void FDocumentClipboardUtils::ProcessPastedInputNodes(FMetaSoundFrontendDocumentBuilder& Builder, FMetasoundAssetBase& OutAsset, TArray<UMetasoundEditorGraphNode*>& OutPastedNodes)
 	{
 		using namespace Engine;
@@ -528,38 +555,6 @@ namespace Metasound::Editor
 		}
 	}
 
-	class FMemberClipboardObjectTextFactory : public FCustomizableTextObjectFactory
-	{
-	public:
-
-		FMemberClipboardObjectTextFactory()
-			: FCustomizableTextObjectFactory(GWarn)
-			, Member(nullptr)
-		{
-		}
-
-		class UMetasoundEditorGraphMember* Member;
-
-	protected:
-
-		virtual bool CanCreateClass(UClass* InObjectClass, bool& bOmitSubObjs) const override
-		{
-			if (InObjectClass->IsChildOf(UMetasoundEditorGraphMember::StaticClass()))
-			{
-				return true;
-			}
-			return false;
-		}
-
-		virtual void ProcessConstructedObject(UObject* CreatedObject) override
-		{
-			if (CreatedObject->IsA<UMetasoundEditorGraphMember>())
-			{
-				Member = CastChecked<UMetasoundEditorGraphMember>(CreatedObject);
-			}
-		}
-	};
-
 	void FDocumentClipboardUtils::ProcessPastedNodeConnections(FMetasoundAssetBase& OutAsset, TArray<UMetasoundEditorGraphNode*>& OutPastedNodes)
 	{
 		using namespace Frontend;
@@ -665,8 +660,114 @@ namespace Metasound::Editor
 		FPlatformApplicationMisc::ClipboardCopy(*Archive);
 	}
 
+	FMetasoundFrontendGraphLayer FDocumentClipboardUtils::CopySelectionToLayer(const FMetaSoundFrontendDocumentBuilder& Builder, const FGraphPanelSelectionSet& Selection)
+{
+		using namespace Engine;
+		using namespace Frontend;
+
+		FMetasoundFrontendGraphLayer Layer;
+		Layer.Graph.PageID = FGuid::NewGuid();
+
+		TSet<FGuid> NodeIDs;
+		TMap<FGuid, FMetasoundFrontendClass> DependenciesToAdd;
+		for (UObject* Object : Selection)
+		{
+			if (UMetasoundEditorGraphExternalNode* ExternalNode = Cast<UMetasoundEditorGraphExternalNode>(Object))
+			{
+				const FMetasoundFrontendNode& FrontendNode = ExternalNode->GetFrontendNodeChecked();
+				if (const FMetasoundFrontendClass* Dependency = Builder.FindDependency(FrontendNode.ClassID))
+				{
+					Layer.Graph.Nodes.Add(FrontendNode);
+					NodeIDs.Add(FrontendNode.GetID());
+					DependenciesToAdd.Add(Dependency->ID, *Dependency);
+				}
+			}
+			else if (const UMetasoundEditorGraphInputNode* EdInputNode = Cast<UMetasoundEditorGraphInputNode>(Object))
+			{
+				if (const UMetasoundEditorGraphInput* EdInput = EdInputNode->Input)
+				{
+					const FMetasoundFrontendNode& TemplateNode = EdInputNode->GetFrontendNodeChecked();
+					{
+						if (const FMetasoundFrontendNode* InputNode = Builder.FindNode(EdInput->NodeID))
+						{
+							if (const FMetasoundFrontendClassInput* FrontendInput = Builder.FindGraphInput(EdInput->GetMemberName()))
+							{
+								const FMetasoundFrontendClass* TemplateDependency = Builder.FindDependency(TemplateNode.ClassID);
+								const FMetasoundFrontendClass* InputDependency = Builder.FindDependency(InputNode->ClassID);
+
+								if (TemplateDependency && InputDependency)
+								{
+									Layer.Graph.Nodes.Add(TemplateNode);
+									NodeIDs.Add(TemplateNode.GetID());
+
+									DependenciesToAdd.Add(TemplateDependency->ID, *TemplateDependency);
+									DependenciesToAdd.Add(InputDependency->ID, *InputDependency);
+
+									if (!NodeIDs.Contains(InputNode->GetID()))
+									{
+										Layer.Graph.Nodes.Add(*InputNode);
+										NodeIDs.Add(InputNode->GetID());
+										Layer.Inputs.Add(*FrontendInput);
+									}
+								}
+							}
+						}
+					}
+				}
+			}
+			else if (const UMetasoundEditorGraphOutputNode* OutputNode = Cast<UMetasoundEditorGraphOutputNode>(Object))
+			{
+				const FMetasoundFrontendNode& FrontendNode = OutputNode->GetFrontendNodeChecked();
+				if (const FMetasoundFrontendClass* Dependency = Builder.FindDependency(FrontendNode.ClassID))
+				{
+					Layer.Graph.Nodes.Add(FrontendNode);
+					NodeIDs.Add(FrontendNode.GetID());
+					DependenciesToAdd.Add(Dependency->ID, *Dependency);
+
+					if (UMetasoundEditorGraphOutput* Output = OutputNode->Output)
+					{
+						if (const FMetasoundFrontendClassOutput* FrontendOutput = Builder.FindGraphOutput(Output->GetMemberName()))
+						{
+							Layer.Outputs.Add(*FrontendOutput);
+						}
+					}
+				}
+			}
+		}
+
+		for (const UObject* Object : Selection)
+		{
+			if (const UMetasoundEditorGraphNode* EditorNode = Cast<const UMetasoundEditorGraphNode>(Object))
+			{
+				if (EditorNode->IsA<UMetasoundEditorGraphVariableNode>())
+				{
+					continue;
+				}
+
+				const FMetasoundFrontendNode& FrontendNode = EditorNode->GetFrontendNodeChecked();
+				for (const FMetasoundFrontendVertex& Vertex : FrontendNode.Interface.Inputs)
+				{
+					const FMetasoundFrontendNode* OutputNode = nullptr;
+					if (const FMetasoundFrontendVertex* Output = Builder.FindNodeOutputConnectedToNodeInput(EditorNode->GetNodeID(), Vertex.VertexID, &OutputNode))
+					{
+						check(OutputNode);
+						if (NodeIDs.Contains(OutputNode->GetID()))
+						{
+							Layer.Graph.Edges.Add(FMetasoundFrontendEdge { OutputNode->GetID(), Output->VertexID, FrontendNode.GetID(), Vertex.VertexID });
+						}
+					}
+				}
+			}
+		}
+
+		DependenciesToAdd.GenerateValueArray(Layer.Dependencies);
+		return Layer;
+	}
+
 	const UMetasoundEditorGraphMember* FDocumentClipboardUtils::GetMemberFromClipboard()
 	{
+		using namespace DocumentClipboardUtilsPrivate;
+
 		// Get the text from the clipboard.
 		FString ClipboardText;
 		FPlatformApplicationMisc::ClipboardPaste(ClipboardText);
@@ -683,7 +784,7 @@ namespace Metasound::Editor
 
 	const bool FDocumentClipboardUtils::CanImportMemberFromText(const FString& TextToImport)
 	{
-		FMemberClipboardObjectTextFactory Factory;
-		return Factory.CanCreateObjectsFromText(TextToImport);
+		using namespace DocumentClipboardUtilsPrivate;
+		return FMemberClipboardObjectTextFactory().CanCreateObjectsFromText(TextToImport);
 	}
 } // namespace Metasound::Editor
diff --git a/Engine/Plugins/Runtime/Metasound/Source/MetasoundEditor/Private/MetasoundEditorDocumentClipboardUtils.h b/Engine/Plugins/Runtime/Metasound/Source/MetasoundEditor/Private/MetasoundEditorDocumentClipboardUtils.h
--- a/Engine/Plugins/Runtime/Metasound/Source/MetasoundEditor/Private/MetasoundEditorDocumentClipboardUtils.h
+++ b/Engine/Plugins/Runtime/Metasound/Source/MetasoundEditor/Private/MetasoundEditorDocumentClipboardUtils.h
@@ -12,6 +12,10 @@
 #include "UObject/NameTypes.h"
 #include "UObject/NoExportTypes.h"
 
+// Forward Declarations
+
+struct FMetasoundFrontendGraphLayer;
+
 
 namespace Metasound::Editor
 {
@@ -29,6 +33,8 @@ namespace Metasound::Editor
 		/** Copy MetasoundMember to Clipboard*/
 		static void CopyMemberToClipboard(class UMetasoundEditorGraphMember* Content);
 
+		static FMetasoundFrontendGraphLayer CopySelectionToLayer(const FMetaSoundFrontendDocumentBuilder& Builder, const FGraphPanelSelectionSet& Selection);
+
 		/** Return MetasoundMember that is in Clipboard, if any*/
 		static const UMetasoundEditorGraphMember* GetMemberFromClipboard();
 
diff --git a/Engine/Plugins/Runtime/Metasound/Source/MetasoundEditor/Private/MetasoundEditorGraphSchema.cpp b/Engine/Plugins/Runtime/Metasound/Source/MetasoundEditor/Private/MetasoundEditorGraphSchema.cpp
--- a/Engine/Plugins/Runtime/Metasound/Source/MetasoundEditor/Private/MetasoundEditorGraphSchema.cpp
+++ b/Engine/Plugins/Runtime/Metasound/Source/MetasoundEditor/Private/MetasoundEditorGraphSchema.cpp
@@ -1712,6 +1712,12 @@ void UMetasoundEditorGraphSchema::GetContextMenuActions(class UToolMenu* Menu, c
 	}
 
 	Super::GetContextMenuActions(Menu, Context);
+
+	if (Context->Graph)
+	{
+		FToolMenuSection& Section = Menu->AddSection("MetasoundGraphSchemaGraphActions", LOCTEXT("DevActionsMenuHeader", "Dev Actions"));
+		Section.AddMenuEntry(FEditorCommands::Get().CopySelectionAsConfigurationCodeToClipboard);
+	}
 }
 
 void UMetasoundEditorGraphSchema::CreateDefaultNodesForGraph(UEdGraph& Graph) const
