MetaSound Document Configuration: Vertical Checkpoint 2

diff --git a/Engine/Plugins/Runtime/Metasound/Source/MetasoundEditor/Private/MetasoundEditorSubsystem.cpp b/Engine/Plugins/Runtime/Metasound/Source/MetasoundEditor/Private/MetasoundEditorSubsystem.cpp
--- a/Engine/Plugins/Runtime/Metasound/Source/MetasoundEditor/Private/MetasoundEditorSubsystem.cpp
+++ b/Engine/Plugins/Runtime/Metasound/Source/MetasoundEditor/Private/MetasoundEditorSubsystem.cpp
@@ -2,6 +2,8 @@
 #include "MetasoundEditorSubsystem.h"
 
 #include "AudioPropertiesSheetAssetBase.h"
+#include "DocumentConfigurations/MetasoundFrontendDocumentConfiguration.h"
+#include "DocumentConfigurations/MetasoundFrontendPresetDocumentConfiguration.h"
 #include "FileHelpers.h"
 #include "IAssetTools.h"
 #include "ISourceControlModule.h"
@@ -56,18 +58,12 @@ TScriptInterface<IMetaSoundDocumentInterface> UMetaSoundEditorSubsystem::BuildTo
 
 		// Not about to follow this lack of const correctness down a multidecade in the works rabbit hole.
 		UClass& MetaSoundUClass = const_cast<UClass&>(InBuilder->GetBaseMetaSoundUClass());
-
-		UObject* NewMetaSound;
 		
 		// Duplicate referenced preset object to preserve object properties (ex. quality settings, soundwave properties)
-		if (InBuilder->IsPreset())
+		UObject* NewMetaSound = nullptr;
+		if (const FPresetDocumentConfiguration* PresetConfig = InBuilder->GetConstBuilder().GetConstDocumentConfigurationAs<FPresetDocumentConfiguration>())
 		{
-			// Rebuild referenced classes to find referenced preset asset
-			FMetasoundAssetBase& InAsset = InBuilder->GetBuilder().GetMetasoundAsset();
-			InAsset.RebuildReferencedAssetClasses();
-
-			UObject* ReferencedObject = InBuilder->GetReferencedPresetAsset();
-			NewMetaSound = IAssetTools::Get().DuplicateAsset(AssetName, PackagePath, ReferencedObject);
+			NewMetaSound = IAssetTools::Get().DuplicateAsset(AssetName, PackagePath, PresetConfig->Parent.GetObject());
 		}
 		else
 		{
@@ -100,9 +96,9 @@ TScriptInterface<IMetaSoundDocumentInterface> UMetaSoundEditorSubsystem::BuildTo
 			// Apply template SoundWave settings 
 			{
 				// Template SoundWave settings only apply to sources
-				// and will override settings from a preset's referenced asset settings
-				const bool bIsSource = &MetaSoundUClass == UMetaSoundSource::StaticClass();
-				if (TemplateSoundWave != nullptr && bIsSource)
+				// and will override settings from a template SoundWave's values
+				const bool bIsSound = MetaSoundUClass.IsChildOf(USoundWave::StaticClass());
+				if (TemplateSoundWave && bIsSound)
 				{
 					SetSoundWaveSettingsFromTemplate(*CastChecked<USoundWave>(NewMetaSound), *TemplateSoundWave);
 				}
@@ -142,8 +138,13 @@ UMetasoundEditorGraphMemberDefaultLiteral* UMetaSoundEditorSubsystem::CreateMemb
 	FName InMemberName,
 	TSubclassOf<UMetasoundEditorGraphMemberDefaultLiteral> LiteralClass) const
 {
-	using namespace Metasound::Engine;
-	// Create brand new member metadata
+	return CreateMemberMetadata(Builder, LiteralClass);
+}
+
+UMetasoundEditorGraphMemberDefaultLiteral* UMetaSoundEditorSubsystem::CreateMemberMetadata(
+	FMetaSoundFrontendDocumentBuilder& Builder,
+	TSubclassOf<UMetasoundEditorGraphMemberDefaultLiteral> LiteralClass) const
+{
 	return NewObject<UMetasoundEditorGraphMemberDefaultLiteral>(&Builder.CastDocumentObjectChecked<UObject>(), LiteralClass, FName(), RF_Transactional, nullptr);
 }
 
@@ -179,7 +180,7 @@ bool UMetaSoundEditorSubsystem::BindMemberMetadata(
 			}
 		}
 
-		NewLiteral = CreateMemberMetadata(Builder, InMember.GetMemberName(), LiteralClass);
+		NewLiteral = CreateMemberMetadata(Builder, LiteralClass);
 	}
 
 	if (NewLiteral)
@@ -234,7 +235,7 @@ UMetaSoundFrontendMemberMetadata* UMetaSoundEditorSubsystem::FindOrCreateGraphIn
 				TSubclassOf<UMetasoundEditorGraphMemberDefaultLiteral> LiteralClass = GetLiteralClassForType(TypeName);
 
 				// Create new literal and setup
-				UMetasoundEditorGraphMemberDefaultLiteral* NewLiteral = CreateMemberMetadata(DocBuilder, InputName, LiteralClass);
+				UMetasoundEditorGraphMemberDefaultLiteral* NewLiteral = CreateMemberMetadata(DocBuilder, LiteralClass);
 				if (NewLiteral)
 				{
 					NewLiteral->MemberID = GraphInputNodeHandle.NodeID;
@@ -309,64 +310,70 @@ const TArray<TSharedRef<FExtender>>& UMetaSoundEditorSubsystem::GetToolbarExtend
 	return EditorToolbarExtenders;
 }
 
-void UMetaSoundEditorSubsystem::InitAsset(UObject& InNewMetaSound, UObject* InReferencedMetaSound, const bool bClearDocument)
+void UMetaSoundEditorSubsystem::InitAsset(UObject& InNewMetaSound, UObject* InParentMetaSound, const bool bClearDocument)
 {
-	using namespace Metasound;
 	using namespace Metasound::Editor;
-	using namespace Metasound::Frontend;
-
-	TScriptInterface<IMetaSoundDocumentInterface> DocInterface = &InNewMetaSound;
-
-	// Must force 'FinishBuilding' and use local builder as re-initializing document via
-	// 'InitDocument' call below can lead to stale builder entries in the registry when renaming.
-	const FMetasoundFrontendClassName InitClassName = DocInterface->GetConstDocument().RootGraph.Metadata.GetClassName();
-	const FTopLevelAssetPath AssetPath = DocInterface->GetAssetPathChecked();
-	IDocumentBuilderRegistry::GetChecked().FinishBuilding(InitClassName, AssetPath);
 
-	FMetaSoundFrontendDocumentBuilder Builder = FMetaSoundFrontendDocumentBuilder(DocInterface);
-	if (bClearDocument)
+	TInstancedStruct<FMetaSoundFrontendDocumentConfiguration> PresetConfig = TInstancedStruct<FMetaSoundFrontendDocumentConfiguration>::Make<FMetaSoundFrontendPresetDocumentConfiguration>();
+	PresetConfig.GetMutable<FMetaSoundFrontendPresetDocumentConfiguration>().Parent = InParentMetaSound;
+	InitAsset(InNewMetaSound, FInitAssetArgs
 	{
-		TSharedRef<FDocumentModifyDelegates> ModifyDelegates = MakeShared<FDocumentModifyDelegates>(Builder.GetConstDocumentChecked());
-		Builder.ClearDocument(ModifyDelegates);
-	}
-	Builder.InitDocument();
-
-	Builder.InitNodeLocations();
-
-	constexpr bool bForceNodeCreation = true;
-	FInputNodeTemplate::GetChecked().Inject(Builder, bForceNodeCreation);
-
-	const FString& Author = GetDefaultAuthor();
-	Builder.SetAuthor(Author);
+		.DocumentConfiguration = MoveTemp(PresetConfig),
+		.SelectedObjects = { InParentMetaSound }
+	});
 
-	// Initialize asset as a preset
-	if (InReferencedMetaSound)
+	Metasound::Editor::IMetasoundEditorModule& MetaSoundEditorModule = FModuleManager::GetModuleChecked<Metasound::Editor::IMetasoundEditorModule>("MetaSoundEditor");
+	if (MetaSoundEditorModule.IsRestrictedMode())
 	{
-		// Ensure the referenced MetaSound is registered already
-		RegisterGraphWithFrontend(*InReferencedMetaSound);
-
-		// Initialize preset with referenced Metasound 
-		TScriptInterface<IMetaSoundDocumentInterface> ReferencedDocInterface = InReferencedMetaSound;
-		const FMetaSoundFrontendDocumentBuilder& ReferenceBuilder = IDocumentBuilderRegistry::GetChecked().FindOrBeginBuilding(ReferencedDocInterface);
-		Builder.ConvertToPreset(ReferenceBuilder);
-		
-		if (&DocInterface->GetBaseMetaSoundUClass() == UMetaSoundSource::StaticClass())
+		if (USoundWave* MetaSoundWave = Cast<USoundWave>(&InNewMetaSound))
 		{
-			// If in restricted mode, copy source soundwave settings from referenced MetaSound 
-			// (In non restricted mode, soundwave settings are automatically copied over by creating the new asset using duplicate) 
-			Metasound::Editor::IMetasoundEditorModule& MetaSoundEditorModule = FModuleManager::GetModuleChecked<Metasound::Editor::IMetasoundEditorModule>("MetaSoundEditor");
-			if (MetaSoundEditorModule.IsRestrictedMode())
+			if (USoundWave* ParentWave = Cast<USoundWave>(InParentMetaSound))
 			{
-				SetSoundWaveSettingsFromTemplate(*CastChecked<USoundWave>(&InNewMetaSound), *CastChecked<USoundWave>(InReferencedMetaSound));
+				FGraphBuilder::RegisterGraphWithFrontend(InNewMetaSound);
+				SetSoundWaveSettingsFromTemplate(*MetaSoundWave, *ParentWave);
 			}
 		}
 	}
 }
 
-void UMetaSoundEditorSubsystem::InitEdGraph(UObject& InMetaSound)
+void UMetaSoundEditorSubsystem::InitAsset(UObject& InNewMetaSound, UMetaSoundEditorSubsystem::FInitAssetArgs InitAssetArgs)
 {
+	using namespace Metasound;
+	using namespace Metasound::Editor;
 	using namespace Metasound::Frontend;
-	Metasound::Editor::FGraphBuilder::BindEditorGraph(IDocumentBuilderRegistry::GetChecked().FindOrBeginBuilding(&InMetaSound));
+
+	{
+		TScriptInterface<IMetaSoundDocumentInterface> DocInterface = &InNewMetaSound;
+
+		// Must force 'FinishBuilding' and use local builder as re-initializing document via
+		// 'InitDocument' call below can lead to stale builder entries in the registry when renaming.
+		const FMetasoundFrontendClassName InitClassName = DocInterface->GetConstDocument().RootGraph.Metadata.GetClassName();
+		const FTopLevelAssetPath AssetPath = DocInterface->GetAssetPathChecked();
+		IDocumentBuilderRegistry::GetChecked().FinishBuilding(InitClassName, AssetPath);
+
+		FMetaSoundFrontendDocumentBuilder DocBuilder(DocInterface);
+
+		TSharedRef<FDocumentModifyDelegates> ModifyDelegates = MakeShared<FDocumentModifyDelegates>(DocBuilder.GetConstDocumentChecked());
+		DocBuilder.ClearDocument(ModifyDelegates);
+
+		DocBuilder.InitDocument();
+		DocBuilder.InitNodeLocations();
+
+		constexpr bool bForceNodeCreation = true;
+		FInputNodeTemplate::GetChecked().Inject(DocBuilder, bForceNodeCreation);
+
+		const FString& Author = GetDefaultAuthor();
+		DocBuilder.SetAuthor(Author);
+	}
+
+	// Must be called in new context after asset is initialized as blank document as configurations use a registered builder
+	if (FMetaSoundFrontendDocumentConfiguration* DocConfig = InitAssetArgs.DocumentConfiguration.GetMutablePtr())
+	{
+		FMetaSoundFrontendDocumentBuilder& DocBuilder = Frontend::IDocumentBuilderRegistry::GetChecked().FindOrBeginBuilding(&InNewMetaSound);
+		DocBuilder.SetDocumentConfiguration(MoveTemp(InitAssetArgs.DocumentConfiguration));
+		DocConfig->OnAssetInitialized(MoveTemp(InitAssetArgs.SelectedObjects), DocBuilder);
+		DocBuilder.ConfigureDocument(true /* bLogFailure */);
+	}
 }
 
 bool UMetaSoundEditorSubsystem::IsPageAuditionPlatformCookTarget(FName InPageName) const
diff --git a/Engine/Plugins/Runtime/Metasound/Source/MetasoundEditor/Private/MetasoundFactory.cpp b/Engine/Plugins/Runtime/Metasound/Source/MetasoundEditor/Private/MetasoundFactory.cpp
--- a/Engine/Plugins/Runtime/Metasound/Source/MetasoundEditor/Private/MetasoundFactory.cpp
+++ b/Engine/Plugins/Runtime/Metasound/Source/MetasoundEditor/Private/MetasoundFactory.cpp
@@ -1,34 +1,22 @@
 // Copyright Epic Games, Inc. All Rights Reserved.
 #include "MetasoundFactory.h"
 
+#include "DocumentConfigurations/MetasoundFrontendDocumentConfiguration.h"
+#include "DocumentConfigurations/MetasoundFrontendPresetDocumentConfiguration.h"
 #include "Metasound.h"
 #include "MetasoundEditorGraph.h"
 #include "MetasoundEditorGraphBuilder.h"
 #include "MetasoundEditorGraphSchema.h"
+#include "MetasoundEditorModule.h"
 #include "MetasoundEditorSubsystem.h"
 #include "MetasoundFrontendDocumentBuilder.h"
 #include "MetasoundFrontendTransform.h"
 #include "MetasoundSource.h"
+#include "UObject/ScriptInterface.h"
 
 #include UE_INLINE_GENERATED_CPP_BY_NAME(MetasoundFactory)
 
 
-namespace Metasound
-{
-	namespace FactoryPrivate
-	{
-		template <typename T>
-		T* CreateNewMetaSoundObject(UObject* InParent, FName InName, EObjectFlags InFlags, UObject* InReferencedMetaSoundObject)
-		{
-			T* MetaSoundObject = NewObject<T>(InParent, InName, InFlags);
-			check(MetaSoundObject);
-
-			UMetaSoundEditorSubsystem::GetChecked().InitAsset(*MetaSoundObject, InReferencedMetaSoundObject);
-			return MetaSoundObject;
-		}
-	}
-}
-
 UMetaSoundBaseFactory::UMetaSoundBaseFactory(const FObjectInitializer& ObjectInitializer)
 	: Super(ObjectInitializer)
 {
@@ -46,15 +34,18 @@ UMetaSoundFactory::UMetaSoundFactory(const FObjectInitializer& ObjectInitializer
 UObject* UMetaSoundFactory::FactoryCreateNew(UClass* InClass, UObject* InParent, FName InName, EObjectFlags InFlags, UObject* InContext, FFeedbackContext* InFeedbackContext)
 {
 	using namespace Metasound::Editor;
-	using namespace Metasound::FactoryPrivate;
 
-	if (UMetaSoundPatch* NewPatch = CreateNewMetaSoundObject<UMetaSoundPatch>(InParent, InName, InFlags, ReferencedMetaSoundObject))
+	UMetaSoundPatch* NewPatch = NewObject<UMetaSoundPatch>(InParent, InName, InFlags);
+	check(NewPatch);
+
+	UMetaSoundEditorSubsystem::GetChecked().InitAsset(*NewPatch, UMetaSoundEditorSubsystem::FInitAssetArgs
 	{
-		FGraphBuilder::RegisterGraphWithFrontend(*NewPatch);
-		return NewPatch;
-	}
+		.DocumentConfiguration = DocumentConfiguration,
+		.SelectedObjects = SelectedObjects
+	});
 
-	return nullptr;
+	FGraphBuilder::RegisterGraphWithFrontend(*NewPatch);
+	return NewPatch;
 }
 
 UMetaSoundSourceFactory::UMetaSoundSourceFactory(const FObjectInitializer& ObjectInitializer)
@@ -66,20 +57,31 @@ UMetaSoundSourceFactory::UMetaSoundSourceFactory(const FObjectInitializer& Objec
 UObject* UMetaSoundSourceFactory::FactoryCreateNew(UClass* InClass, UObject* InParent, FName InName, EObjectFlags InFlags, UObject* InContext, FFeedbackContext* InFeedbackContext)
 {
 	using namespace Metasound::Editor;
-	using namespace Metasound::FactoryPrivate;
 
-	if (UMetaSoundSource* NewSource = CreateNewMetaSoundObject<UMetaSoundSource>(InParent, InName, InFlags, ReferencedMetaSoundObject))
+	UMetaSoundSource* NewSource = NewObject<UMetaSoundSource>(InParent, InName, InFlags);
+	check(NewSource);
+
+	UMetaSoundEditorSubsystem::GetChecked().InitAsset(*NewSource, UMetaSoundEditorSubsystem::FInitAssetArgs
 	{
-		// Copy over referenced fields that are specific to sources
-		if (UMetaSoundSource* ReferencedMetaSound = Cast<UMetaSoundSource>(ReferencedMetaSoundObject))
+		.DocumentConfiguration = DocumentConfiguration,
+		.SelectedObjects = SelectedObjects
+	});
+
+
+	const IMetasoundEditorModule& MetaSoundEditorModule = FModuleManager::GetModuleChecked<IMetasoundEditorModule>("MetaSoundEditor");
+	if (MetaSoundEditorModule.IsRestrictedMode())
+	{
+		if (const FMetaSoundFrontendPresetDocumentConfiguration* PresetConfig = DocumentConfiguration.GetPtr<FMetaSoundFrontendPresetDocumentConfiguration>())
 		{
-			NewSource->OutputFormat = ReferencedMetaSound->OutputFormat;
+			if (const USoundWave* ParentWave = Cast<const USoundWave>(PresetConfig->Parent.GetObject()))
+			{
+				UMetaSoundEditorSubsystem::GetChecked().SetSoundWaveSettingsFromTemplate(*NewSource, *ParentWave);
+			}
 		}
-
-		FGraphBuilder::RegisterGraphWithFrontend(*NewSource);
-		return NewSource;
 	}
 
-	return nullptr;
+	NewSource->ConformObjectToDocument();
+	FGraphBuilder::RegisterGraphWithFrontend(*NewSource);
+	return NewSource;
 }
 
diff --git a/Engine/Plugins/Runtime/Metasound/Source/MetasoundEditor/Private/MetasoundNodeDetailCustomization.cpp b/Engine/Plugins/Runtime/Metasound/Source/MetasoundEditor/Private/MetasoundNodeDetailCustomization.cpp
--- a/Engine/Plugins/Runtime/Metasound/Source/MetasoundEditor/Private/MetasoundNodeDetailCustomization.cpp
+++ b/Engine/Plugins/Runtime/Metasound/Source/MetasoundEditor/Private/MetasoundNodeDetailCustomization.cpp
@@ -5,6 +5,7 @@
 #include "Components/AudioComponent.h"
 #include "Containers/Set.h"
 #include "Delegates/Delegate.h"
+#include "DocumentConfigurations/MetasoundFrontendPresetDocumentConfiguration.h"
 #include "Framework/Notifications/NotificationManager.h"
 #include "HAL/PlatformApplicationMisc.h"
 #include "IDetailChildrenBuilder.h"
@@ -62,10 +63,6 @@ namespace Metasound
 				"MetasoundParameterPack"
 			};
 
-			static const FText OverrideInputDefaultText = LOCTEXT("OverridePresetInputDefault", "Override Inherited Default");
-			static const FText OverrideInputDefaultTooltip = LOCTEXT("OverridePresetInputTooltip",
-				"Enables overriding the input's inherited default value otherwise provided by the referenced graph. Setting to true disables auto-updating the input's default value if modified on the referenced asset.");
-
 			static const FText ConstructorPinText = LOCTEXT("ConstructorPinText", "Is Constructor Pin");
 			static const FText ConstructorPinTooltip = LOCTEXT("ConstructorPinTooltip",
 				"Whether this input or output is a constructor pin. Constructor values are only read on construction (on play), and are not dynamically updated at runtime.");
@@ -181,12 +178,19 @@ namespace Metasound
 							{
 								Value = Value.LeftChop(2).RightChop(9);
 							}
-							// Validate the class path (before the first ')
-							FString ValueClassPath = Value.Left(Value.Find("'"));
-							// Remove beginning quote (BP case)
-							if (ValueClassPath.StartsWith(TEXT("\"")))
+
+							FString ValueClassPath;
+
+							const int32 CommaIndex = Value.Find("'");
+							if (CommaIndex != INDEX_NONE)
 							{
-								ValueClassPath.RightChopInline(1);
+								// Validate the class path (before the first ')
+								ValueClassPath = Value.Left(CommaIndex);
+								// Remove beginning quote (BP case)
+								if (ValueClassPath.StartsWith(TEXT("\"")))
+								{
+									ValueClassPath.RightChopInline(1);
+								}
 							}
 							
 							// Wrap objects in (Object=*)
@@ -199,7 +203,6 @@ namespace Metasound
 							else
 							{
 								UE_LOG(LogMetaSound, Warning, TEXT("Failed to paste object of type %s which does not match default value type %s"), *ValueClassPath, *ClassPath.ToString());
-								return;
 							}
 						}
 
@@ -248,23 +251,39 @@ namespace Metasound
 
 				FUIAction PasteAction;
 				// Paste only enabled if graph is editable (for variables/outputs)
-				// or if graph is editable and input is not an interface member and is overridden (for inputs)
-				PasteAction.CanExecuteAction = FCanExecuteAction::CreateLambda([GraphMember = GraphMember]()
+				// or if graph is editable and input is not an interface member.
+				PasteAction.CanExecuteAction = FCanExecuteAction::CreateLambda([MemberPtr = TWeakObjectPtr<UMetasoundEditorGraphMember>(GraphMember)]()
 				{
-					if (!GraphMember) { return false; }
-					const bool bIsGraphEditable = GraphMember->GetOwningGraph()->IsEditable();
+					using namespace Frontend;
 
-					if (UMetasoundEditorGraphInput* Input = Cast<UMetasoundEditorGraphInput>(GraphMember))
+					if (!MemberPtr.IsValid())
 					{
-						Frontend::FConstNodeHandle InputNodeHandle = Input->GetConstNodeHandle();
-						const TSet<FName>& InputsInheritingDefault = InputNodeHandle->GetOwningGraph()->GetInputsInheritingDefault();
-						FName NodeName = InputNodeHandle->GetNodeName();
-						return !Input->IsInterfaceMember() && (bIsGraphEditable || !InputsInheritingDefault.Contains(NodeName));
+						return false;
 					}
-					else
+
+					const bool bIsGraphEditable = MemberPtr->GetOwningGraph()->IsEditable();
+					if (bIsGraphEditable)
 					{
-						return bIsGraphEditable;
+						if (UMetasoundEditorGraphInput* Input = Cast<UMetasoundEditorGraphInput>(MemberPtr.Get()))
+						{
+							if (Input->IsInterfaceMember())
+							{
+								return false;
+							}
+
+							if (const FDocumentVertexConfiguration* DocConfig = Input->GetFrontendBuilderChecked().GetConstDocumentConfigurationAs<FDocumentVertexConfiguration>())
+							{
+								if (const FDocumentConfigurationVertexMetadata* VertexMetadata = DocConfig->FindConstVertexMetadata(Input->NodeID))
+								{
+									return VertexMetadata->IsDefaultEditable();
+								}
+							}
+						}
+
+						return true;
 					}
+
+					return false;
 				});
 
 				PasteAction.ExecuteAction = CreateDefaultValueObjectPasteExecuteAction(PropertyHandle);
@@ -467,12 +486,8 @@ namespace Metasound
 		{
 			if (UMetasoundEditorGraphMemberDefaultBool* DefaultBool = BoolLiteral.Pin().Get())
 			{
-				UMetasoundEditorGraphMember* Member = DefaultBool->FindMember();
-				if (Member->IsA<UMetasoundEditorGraphInput>())
-				{
-					return FMetasoundDefaultLiteralCustomizationBase::GetDefaultVisibility();
-				}
-				else
+				const UMetasoundEditorGraphMember* Member = DefaultBool->FindMember();
+				if (Member && !Member->IsA<UMetasoundEditorGraphInput>())
 				{
 					const bool bIsTrigger = Member->GetDataType() == Metasound::GetMetasoundDataTypeName<Metasound::FTrigger>();
 					if (bIsTrigger)
@@ -734,7 +749,10 @@ PRAGMA_ENABLE_DEPRECATION_WARNINGS
 				UObject* Object = nullptr;
 				if (PropertyHandle->GetValue(Object) == FPropertyAccess::Success)
 				{
-					return Object->GetPathName();
+					if (Object)
+					{
+						return Object->GetPathName();
+					}
 				}
 				return FString();
 			};
@@ -1192,7 +1210,6 @@ PRAGMA_ENABLE_DEPRECATION_WARNINGS
 				{
 					LiteralCustomization->SetDefaultVisibility(Visibility);
 					LiteralCustomization->SetEnabled(GetEnabled());
-					LiteralCustomization->SetResetOverride(GetResetOverride());
 					LiteralCustomization->CustomizeDefaults(*MemberDefaultLiteral, InDetailLayout);
 				}
 				else
@@ -1334,6 +1351,20 @@ PRAGMA_ENABLE_DEPRECATION_WARNINGS
 			}
 		}
 
+		FMetaSoundFrontendDocumentBuilder* FMetasoundMemberDetailCustomization::GetBuilder() const
+		{
+			if (UMetasoundEditorGraphMember* MemberPtr = GraphMember.Get())
+			{
+				FMetaSoundFrontendDocumentBuilder& DocBuilder = MemberPtr->GetFrontendBuilderChecked();
+				if (DocBuilder.IsValid())
+				{
+					return &DocBuilder;
+				}
+			}
+
+			return nullptr;
+		}
+
 		FText FMetasoundMemberDetailCustomization::GetName() const
 		{
 			if (GraphMember.IsValid())
@@ -1390,6 +1421,36 @@ PRAGMA_ENABLE_DEPRECATION_WARNINGS
 			return FText::GetEmpty();
 		}
 
+		const Frontend::FDocumentConfigurationVertexMetadata* FMetasoundVertexDetailCustomization::GetConstConfigMetadata() const
+		{
+			using namespace Frontend;
+
+			if (UMetasoundEditorGraphMember* MemberPtr = GraphMember.Get())
+			{
+				if (const FMetaSoundFrontendDocumentBuilder* DocBuilder = GetBuilder())
+				{
+					return DocBuilder->FindConstVertexMetadata<FDocumentConfigurationVertexMetadata>(MemberPtr->GetMemberID());
+				}
+			}
+
+			return nullptr;
+		}
+
+		Frontend::FDocumentConfigurationVertexMetadata* FMetasoundVertexDetailCustomization::GetConfigMetadata() const
+		{
+			using namespace Frontend;
+
+			if (UMetasoundEditorGraphMember* MemberPtr = GraphMember.Get())
+			{
+				if (const FMetaSoundFrontendDocumentBuilder* DocBuilder = GetBuilder())
+				{
+					return DocBuilder->FindVertexMetadata<FDocumentConfigurationVertexMetadata>(MemberPtr->GetMemberID());
+				}
+			}
+
+			return nullptr;
+		}
+
 		EVisibility FMetasoundVertexDetailCustomization::GetDefaultVisibility() const
 		{
 			using namespace Frontend;
@@ -1602,83 +1663,6 @@ PRAGMA_ENABLE_DEPRECATION_WARNINGS
 			return false;
 		}
 
-		bool FMetasoundInputDetailCustomization::GetInputInheritsDefault() const
-		{
-			if (const UMetasoundEditorGraphInput* Input = Cast<UMetasoundEditorGraphInput>(GraphMember.Get()))
-			{
-				if (const TSet<FName>* InputsInheritingDefault = Input->GetFrontendBuilderChecked().GetGraphInputsInheritingDefault())
-				{
-					FName MemberName = Input->GetMemberName();
-					return InputsInheritingDefault->Contains(MemberName);
-				}
-			}
-
-			return false;
-		}
-
-		void FMetasoundInputDetailCustomization::SetInputInheritsDefault()
-		{
-			if (UMetasoundEditorGraphInput* Input = Cast<UMetasoundEditorGraphInput>(GraphMember.Get()))
-			{
-				if (UMetasoundEditorGraphMemberDefaultLiteral* MemberDefaultLiteral = Input->GetLiteral())
-				{
-					FScopedTransaction Transaction(LOCTEXT("SetPresetInputOverrideTransaction", "Set MetaSound Preset Input Overridden"));
-
-					Input->GetOutermost()->Modify();
-					Input->GetOutermostObject()->Modify();
-					Input->Modify();
-					MemberDefaultLiteral->Modify();
-
-					constexpr bool bDefaultIsInherited = true;
-					const FName MemberName = Input->GetMemberName();
-					FMetaSoundFrontendDocumentBuilder& DocumentBuilder = GraphMember->GetFrontendBuilderChecked();
-					DocumentBuilder.SetGraphInputInheritsDefault(MemberName, bDefaultIsInherited);
-
-					Input->UpdateFrontendDefaultLiteral(false /* bPostTransaction */);
-
-					MemberDefaultLiteral->ForceRefresh();
-					
-					if (UObject* Metasound = Input->GetOutermostObject())
-					{
-						FGraphBuilder::RegisterGraphWithFrontend(*Metasound);
-					}
-				}
-			}
-		}
-
-		void FMetasoundInputDetailCustomization::ClearInputInheritsDefault()
-		{
-			if (UMetasoundEditorGraphInput* Input = Cast<UMetasoundEditorGraphInput>(GraphMember.Get()))
-			{
-				if (UMetasoundEditorGraphMemberDefaultLiteral* MemberDefaultLiteral = Input->GetLiteral())
-				{
-					FScopedTransaction Transaction(LOCTEXT("ClearPresetInputOverrideTransaction", "Clear MetaSound Preset Input Overridden"));
-
-					Input->GetOutermost()->Modify();
-					Input->GetOutermostObject()->Modify();
-					Input->Modify();
-					MemberDefaultLiteral->Modify();
-
-					constexpr bool bDefaultIsInherited = false;
-					const FName MemberName = Input->GetMemberName();
-					FMetaSoundFrontendDocumentBuilder& DocumentBuilder = GraphMember->GetFrontendBuilderChecked();
-					DocumentBuilder.SetGraphInputInheritsDefault(MemberName, bDefaultIsInherited);
-
-					Input->UpdateFrontendDefaultLiteral(false /* bPostTransaction */);
-
-					if (UMetasoundEditorGraphMemberDefaultLiteral* Literal = Input->GetLiteral())
-					{
-						Literal->ForceRefresh();
-					}
-
-					if (UObject* Metasound = Input->GetOutermostObject())
-					{
-						FGraphBuilder::RegisterGraphWithFrontend(*Metasound);
-					}
-				}
-			}
-		}
-
 		ECheckBoxState FMetasoundVertexDetailCustomization::OnGetConstructorPinCheckboxState(TWeakObjectPtr<UMetasoundEditorGraphVertex> InGraphVertex) const
 		{
 			if (InGraphVertex.IsValid())
@@ -1745,6 +1729,8 @@ PRAGMA_ENABLE_DEPRECATION_WARNINGS
 
 		void FMetasoundInputDetailCustomization::CustomizeDetails(IDetailLayoutBuilder& InDetailLayout)
 		{
+			using namespace Frontend;
+
 			CacheMemberData(InDetailLayout);
 			if (!GraphMember.IsValid())
 			{
@@ -1759,87 +1745,106 @@ PRAGMA_ENABLE_DEPRECATION_WARNINGS
 				return;
 			}
 
-			// Build preset row first if graph has managed interface, not default constructed, & not a trigger
-			const bool bIsPreset = GraphMember->GetFrontendBuilderChecked().IsPreset();
-			const bool bIsDefaultConstructed = MemberDefaultLiteral->GetLiteralType() == EMetasoundFrontendLiteralType::None;
-			const bool bIsTriggerDataType = GraphMember->GetDataType() == GetMetasoundDataTypeName<FTrigger>();
+			FMetaSoundFrontendDocumentBuilder* DocBuilder = GetBuilder();
+			if (!DocBuilder)
+			{
+				return;
+			}
 
-			if (bIsPreset && !bIsDefaultConstructed && !bIsTriggerDataType)
+			TStructView<FMetaSoundFrontendDocumentConfiguration> ConfigStruct = DocBuilder->GetDocumentConfiguration();
+			if (FDocumentVertexConfiguration* Config = ConfigStruct.GetPtr<FDocumentVertexConfiguration>())
 			{
-				GetDefaultCategoryBuilder(InDetailLayout)
-				.AddCustomRow(MemberCustomizationPrivate::OverrideInputDefaultText)
-				.NameContent()
-				[
-					SNew(STextBlock)
-					.Text(MemberCustomizationPrivate::OverrideInputDefaultText)
-					.Font(IDetailLayoutBuilder::GetDetailFontBold())
-					.ToolTipText(MemberCustomizationPrivate::OverrideInputDefaultTooltip)
-				]
-				.ValueContent()
-				[
-					SNew(SCheckBox)
-					.OnCheckStateChanged_Lambda([this](ECheckBoxState State)
+				const FGuid& MemberID = GraphMember->GetMemberID();
+				if (FDocumentConfigurationVertexMetadata* Metadata = Config->FindVertexMetadata(MemberID))
+				{
+					const UScriptStruct* StructClass = Config->GetVertexMetadataStruct(MemberID);
+					if (const TSharedPtr<FStructOnScope> DataView = MakeShared<FStructOnScope>(StructClass, reinterpret_cast<uint8*>(Metadata)))
 					{
-						switch(State)
+						TArray<TSharedPtr<IPropertyHandle>> StructPropHandles = GetDefaultCategoryBuilder(InDetailLayout).AddAllExternalStructureProperties(DataView.ToSharedRef());
+
+						const auto OnConfigPreChange = TDelegate<void()>::CreateSPLambda(AsShared(), [this]()
 						{
-							case ECheckBoxState::Checked:
+							if (GEditor && GraphMember.IsValid())
 							{
-								ClearInputInheritsDefault();
-								break;
+								if (UMetasoundEditorGraphMemberDefaultLiteral* MemberDefaultLiteral = GraphMember->GetLiteral())
+								{
+									GEditor->BeginTransaction(LOCTEXT("GraphMemberDocConfig_Transact", "Set GraphMember Config Vertex Metadata"));
+									GraphMember->GetOutermost()->Modify();
+									GraphMember->GetOutermostObject()->Modify();
+									GraphMember->Modify();
+									MemberDefaultLiteral->Modify();
+								}
 							}
-							case ECheckBoxState::Unchecked:
-							case ECheckBoxState::Undetermined:
-							default:
+						});
+
+						const auto OnConfigPropChanged = TDelegate<void(const FPropertyChangedEvent&)>::CreateSPLambda(AsShared(), [this](const FPropertyChangedEvent& InEvent)
+						{
+							if (FDocumentConfigurationVertexMetadata* ThisMetadata = GetConfigMetadata())
 							{
-								SetInputInheritsDefault();
+								if (FMetaSoundFrontendDocumentBuilder* DocBuilder = GetBuilder())
+								{
+									ThisMetadata->OnPropertyChanged(InEvent, *DocBuilder);
+
+									if (UMetasoundEditorGraphInput* Input = Cast<UMetasoundEditorGraphInput>(GraphMember.Get()))
+									{
+										if (UMetasoundEditorGraphMemberDefaultLiteral* Literal = Input->GetLiteral())
+										{
+											Literal->ForceRefresh();
+										}
+
+										if (UObject* Metasound = Input->GetOutermostObject())
+										{
+											FGraphBuilder::RegisterGraphWithFrontend(*Metasound);
+										}
+									}
+
+									if (GEditor)
+									{
+										GEditor->EndTransaction();
+									}
+								}
 							}
+						});
+
+						for (TSharedPtr<IPropertyHandle> StructPropHandle : StructPropHandles)
+						{
+// 							Row->ShouldAutoExpand(true);
+							StructPropHandle->SetOnPropertyValuePreChange(OnConfigPreChange);
+							StructPropHandle->SetOnPropertyValueChangedWithData(OnConfigPropChanged);
 						}
-					})
-					.IsChecked_Lambda([this]() { return GetInputInheritsDefault() ? ECheckBoxState::Unchecked : ECheckBoxState::Checked; })
-					.ToolTipText(MemberCustomizationPrivate::OverrideInputDefaultTooltip)
-				];
+					}
+				}
 			}
 
-			if (bIsPreset)
+			Enabled = TAttribute<bool>::CreateSPLambda(AsShared(), [this]
 			{
-				if (!bIsDefaultConstructed && !bIsTriggerDataType)
+				using namespace Frontend;
+				if (const FDocumentConfigurationVertexMetadata* VertexMetadata = GetConstConfigMetadata())
 				{
-					const UMetasoundEditorGraphInput* Input = Cast<UMetasoundEditorGraphInput>(MemberDefaultLiteral->FindMember());
-					if (Input)
-					{
-						Enabled = TAttribute<bool>::CreateSPLambda(AsShared(), [this] { return !GetInputInheritsDefault(); });
-						ResetOverride = FResetToDefaultOverride::Create(
-							FIsResetToDefaultVisible::CreateSPLambda(AsShared(), [this](TSharedPtr<IPropertyHandle> /* PropertyHandle */) { return !GetInputInheritsDefault(); }),
-							FResetToDefaultHandler::CreateSPLambda(AsShared(), [this](TSharedPtr<IPropertyHandle> /* PropertyHandle */) { SetInputInheritsDefault(); }));
-					}
+					return VertexMetadata->IsDefaultEditable();
 				}
-			}
-			else
-			{
-				Enabled = TAttribute<bool>::CreateSPLambda(AsShared(), [this]
+
+				if (GraphMember.IsValid())
 				{
-					if (GraphMember.IsValid())
+					if (UMetasoundEditorGraphMemberDefaultLiteral* MemberDefaultLiteral = GraphMember->GetLiteral())
 					{
-						if (UMetasoundEditorGraphMemberDefaultLiteral* MemberDefaultLiteral = GraphMember->GetLiteral())
+						// Make default value uneditable while playing for constructor inputs
+						const UMetasoundEditorGraphInput* Input = Cast<UMetasoundEditorGraphInput>(MemberDefaultLiteral->FindMember());
+						if (Input)
 						{
-							// Make default value uneditable while playing for constructor inputs
-							const UMetasoundEditorGraphInput* Input = Cast<UMetasoundEditorGraphInput>(MemberDefaultLiteral->FindMember());
-							if (Input)
+							if (Input->GetVertexAccessType() == EMetasoundFrontendVertexAccessType::Value)
 							{
-								if (Input->GetVertexAccessType() == EMetasoundFrontendVertexAccessType::Value)
+								UObject* MetaSoundObject = Input->GetOutermostObject();
+								if (TSharedPtr<FEditor> MetaSoundEditor = FGraphBuilder::GetEditorForMetasound(*MetaSoundObject))
 								{
-									UObject* MetaSoundObject = Input->GetOutermostObject();
-									if (TSharedPtr<FEditor> MetaSoundEditor = FGraphBuilder::GetEditorForMetasound(*MetaSoundObject))
-									{
-										return !MetaSoundEditor->IsPlaying();
-									}
+									return !MetaSoundEditor->IsPlaying();
 								}
 							}
 						}
 					}
-					return true;
-				});
-			}
+				}
+				return true;
+			});
 
 			CustomizeDefaultCategory(InDetailLayout);
 		}
@@ -1849,11 +1854,6 @@ PRAGMA_ENABLE_DEPRECATION_WARNINGS
 			return Enabled;
 		}
 
-		bool FMetasoundInputDetailCustomization::IsDefaultEditable() const
-		{
-			return !GetInputInheritsDefault();
-		}
-
 		EVisibility FMetasoundVariableDetailCustomization::GetDefaultVisibility() const
 		{
 			using namespace Frontend;
diff --git a/Engine/Plugins/Runtime/Metasound/Source/MetasoundEditor/Private/MetasoundNodeDetailCustomization.h b/Engine/Plugins/Runtime/Metasound/Source/MetasoundEditor/Private/MetasoundNodeDetailCustomization.h
--- a/Engine/Plugins/Runtime/Metasound/Source/MetasoundEditor/Private/MetasoundNodeDetailCustomization.h
+++ b/Engine/Plugins/Runtime/Metasound/Source/MetasoundEditor/Private/MetasoundNodeDetailCustomization.h
@@ -6,6 +6,7 @@
 #include "DetailCategoryBuilder.h"
 #include "DetailLayoutBuilder.h"
 #include "DetailWidgetRow.h"
+#include "DocumentConfigurations/MetasoundFrontendDocumentConfiguration.h"
 #include "EdGraph/EdGraphSchema.h"
 #include "Framework/Application/SlateApplication.h"
 #include "IDetailCustomization.h"
@@ -355,10 +356,6 @@ namespace Metasound
 
 			virtual EVisibility GetDefaultVisibility() const { return EVisibility::Visible; }
 			virtual TAttribute<bool> GetEnabled() const { return { }; }
-			virtual const TOptional<FResetToDefaultOverride>& GetResetOverride() const { return ResetOverride; }
-
-			UE_DEPRECATED(5.5, "Use GetEnabled attribute instead")
-			virtual bool IsDefaultEditable() const { return true; }
 
 			virtual bool IsInterfaceMember() const { return false; }
 
@@ -375,13 +372,13 @@ namespace Metasound
 			void OnNameCommitted(const FText& InNewName, ETextCommit::Type InTextCommit);
 
 		protected:
+			FMetaSoundFrontendDocumentBuilder* GetBuilder() const;
+
 			TWeakObjectPtr<UMetasoundEditorGraphMember> GraphMember;
 
 			TSharedPtr<SEditableTextBox> NameEditableTextBox;
 			FMetasoundDataTypeSelector DataTypeSelector;
 
-			TOptional<FResetToDefaultOverride> ResetOverride;
-
 			bool bIsNameInvalid = false;
 
 			FDelegateHandle RenameRequestedHandle;
@@ -404,6 +401,9 @@ namespace Metasound
 			ECheckBoxState OnGetConstructorPinCheckboxState(TWeakObjectPtr<UMetasoundEditorGraphVertex> InGraphMember) const;
 			void OnConstructorPinStateChanged(TWeakObjectPtr<UMetasoundEditorGraphVertex> InGraphMember, ECheckBoxState InNewState);
 
+			const Frontend::FDocumentConfigurationVertexMetadata* GetConstConfigMetadata() const;
+			Frontend::FDocumentConfigurationVertexMetadata* GetConfigMetadata() const;
+
 			TSharedPtr<SCheckBox> ConstructorPinCheckbox;
 
 #if WITH_EDITORONLY_DATA
@@ -423,13 +423,8 @@ namespace Metasound
 			
 			virtual void CustomizeDetails(IDetailLayoutBuilder& InDetailLayout) override;
 			virtual TAttribute<bool> GetEnabled() const override;
-			virtual bool IsDefaultEditable() const override;
 
 		private:
-			bool GetInputInheritsDefault() const;
-			void SetInputInheritsDefault();
-			void ClearInputInheritsDefault();
-
 			TAttribute<bool> Enabled;
 		};
 
diff --git a/Engine/Plugins/Runtime/Metasound/Source/MetasoundEditor/Public/MetasoundEditorCommands.h b/Engine/Plugins/Runtime/Metasound/Source/MetasoundEditor/Public/MetasoundEditorCommands.h
--- a/Engine/Plugins/Runtime/Metasound/Source/MetasoundEditor/Public/MetasoundEditorCommands.h
+++ b/Engine/Plugins/Runtime/Metasound/Source/MetasoundEditor/Public/MetasoundEditorCommands.h
@@ -34,6 +34,9 @@ namespace Metasound
 			/** Plays stops the currently playing Metasound */
 			TSharedPtr<FUICommandInfo> TogglePlayback;
 
+			/** Toggles showing the graph for the given configured MetaSound */
+			TSharedPtr<FUICommandInfo> ToggleShowGraph;
+
 			/** Selects the Metasound in the content browser. If referencing Metasound nodes are selected, selects referenced assets instead. */
 			TSharedPtr<FUICommandInfo> BrowserSync;
 
@@ -52,9 +55,6 @@ namespace Metasound
 			// Run updates on selected node(s) class(es) if required by any subset of selected nodes
 			TSharedPtr<FUICommandInfo> UpdateNodeClass;
 
-			// Converts a MetaSound from a restricted, preset edit state to a fully accessible MetaSound.
-			TSharedPtr<FUICommandInfo> ConvertFromPreset;
-
 			// Delete selected items. 
 			TSharedPtr<FUICommandInfo> Delete;
 
@@ -79,6 +79,9 @@ namespace Metasound
 			/** Promotes all unconnected node inputs to shared graph inputs. */
 			TSharedPtr<FUICommandInfo> PromoteAllToCommonInputs;
 
+			/** Copies the selected nodes, shared connections, and required interfaces to the clipboard as configurator c++ */
+			TSharedPtr<FUICommandInfo> CopySelectionAsConfigurationCodeToClipboard;
+
 			/** Initialize commands */
 			virtual void RegisterCommands() override;
 		};
diff --git a/Engine/Plugins/Runtime/Metasound/Source/MetasoundEditor/Public/MetasoundEditorSubsystem.h b/Engine/Plugins/Runtime/Metasound/Source/MetasoundEditor/Public/MetasoundEditorSubsystem.h
--- a/Engine/Plugins/Runtime/Metasound/Source/MetasoundEditor/Public/MetasoundEditorSubsystem.h
+++ b/Engine/Plugins/Runtime/Metasound/Source/MetasoundEditor/Public/MetasoundEditorSubsystem.h
@@ -29,6 +29,12 @@ class UMetaSoundEditorSubsystem : public UEditorSubsystem
 	GENERATED_BODY()
 
 public:
+	struct FInitAssetArgs
+	{
+		TInstancedStruct<FMetaSoundFrontendDocumentConfiguration> DocumentConfiguration;
+		TArray<UObject*> SelectedObjects;
+	};
+
 	// Binds literal editor Metadata to the given member.  If the literal already exists, adds literal
 	// reference to given member (asserts that existing literal is of similar subclass provided).  If 
 	// it does not exist, or an optional template object is provided, metadata is generated then bound.
@@ -54,11 +60,8 @@ public:
 		EMetaSoundBuilderResult& OutResult,
 		UPARAM(DisplayName = "Template SoundWave") const USoundWave* TemplateSoundWave = nullptr);
 
-	// Creates new member metadata for a member of a given builder, copying data from the referenced asset in the case of preset inherited inputs
-	UE_API UMetasoundEditorGraphMemberDefaultLiteral* CreateMemberMetadata(
-		FMetaSoundFrontendDocumentBuilder& Builder,
-		FName InMemberName,
-		TSubclassOf<UMetasoundEditorGraphMemberDefaultLiteral> LiteralClass) const;
+	UE_DEPRECATED(5.8, "Functionality is now private")
+	UE_API UMetasoundEditorGraphMemberDefaultLiteral* CreateMemberMetadata(FMetaSoundFrontendDocumentBuilder& Builder, FName InMemberName, TSubclassOf<UMetasoundEditorGraphMemberDefaultLiteral> LiteralClass) const;
 
 	// Returns a builder for the given MetaSound asset. Returns null if provided a transient MetaSound. For finding builders for transient
 	// MetaSounds, use the UMetaSoundBuilderSubsystem's API (FindPatchBuilder, FindSourceBuilder, FindBuilderByName etc.)
@@ -80,13 +83,12 @@ public:
 		UPARAM(DisplayName = "Location") const FVector2D& InLocation,
 		EMetaSoundBuilderResult& OutResult);
 
-	// Initialize the UObject asset
-	// with an optional MetaSound to be referenced if the asset is a preset
-	// and optionally clearing the existing MetaSound document (for the case of duplicated assets) 
-	UE_API void InitAsset(UObject& InNewMetaSound, UObject* InReferencedMetaSound = nullptr, const bool bClearDocument = false);
+	// Initialize the UObject asset with an optional MetaSound document configuration to be applied.
+	// Optionally clears the existing MetaSound document (for the case of duplicated assets)
+	UE_API void InitAsset(UObject& InNewMetaSound, FInitAssetArgs InitAssetArgs = { });
 
-	UE_DEPRECATED(5.5, "EdGraph is now transiently generated and privately managed for asset editor use only.")
-	UE_API void InitEdGraph(UObject& InMetaSound);
+	UE_DEPRECATED(5.8, "Use overload that is passed optional initialization arguments. Call SetSoundWaveSettingsFromTemplate directly if desired.")
+	UE_API void InitAsset(UObject& InNewMetaSound, UObject* InReferencedMetaSound = nullptr, const bool bClearDocument = false);
 
 	// Returns whether or not a page with the given name both exists and is set as
 	// a valid, cooked target for the currently set audition platform in editor.
@@ -129,14 +131,17 @@ public:
 	// Returns all currently toolbar extenders registered to be displayed within the MetaSound Asset Editor.
 	UE_API const TArray<TSharedRef<FExtender>>& GetToolbarExtenders() const;
 
+	// Copy over sound wave settings such as attenuation, modulation, and sound class from the template sound wave to the MetaSound
+	UE_API void SetSoundWaveSettingsFromTemplate(USoundWave& NewMetasound, const USoundWave& TemplateSoundWave) const;
+
 	static UE_API UMetaSoundEditorSubsystem& GetChecked();
 	static UE_API const UMetaSoundEditorSubsystem& GetConstChecked();
 
 private:
+	UE_API UMetasoundEditorGraphMemberDefaultLiteral* CreateMemberMetadata(FMetaSoundFrontendDocumentBuilder& Builder, TSubclassOf<UMetasoundEditorGraphMemberDefaultLiteral> LiteralClass) const;
+
 	UE_API bool SetFocusedPageInternal(FName PageName, const FGuid& InPageID, UMetaSoundBuilderBase& Builder, bool bOpenEditor, bool bPostTransaction) const;
 
-	// Copy over sound wave settings such as attenuation, modulation, and sound class from the template sound wave to the MetaSound
-	UE_API void SetSoundWaveSettingsFromTemplate(USoundWave& NewMetasound, const USoundWave& TemplateSoundWave) const;
 
 	// Editor Toolbar Extenders
 	TArray<TSharedRef<FExtender>> EditorToolbarExtenders;
diff --git a/Engine/Plugins/Runtime/Metasound/Source/MetasoundEditor/Public/MetasoundFactory.h b/Engine/Plugins/Runtime/Metasound/Source/MetasoundEditor/Public/MetasoundFactory.h
--- a/Engine/Plugins/Runtime/Metasound/Source/MetasoundEditor/Public/MetasoundFactory.h
+++ b/Engine/Plugins/Runtime/Metasound/Source/MetasoundEditor/Public/MetasoundFactory.h
@@ -2,24 +2,34 @@
 
 #pragma once
 
-#include "CoreMinimal.h"
+#include "DocumentConfigurations/MetasoundFrontendDocumentConfiguration.h"
 #include "Factories/Factory.h"
+#include "MetasoundFrontendDocument.h"
+#include "StructUtils/InstancedStruct.h"
 #include "UObject/ObjectMacros.h"
 
 #include "MetasoundFactory.generated.h"
 
+// Forward Declarations
 class UMetaSoundPatch;
 class UMetaSoundSource;
 
+
 UCLASS(abstract)
 class UMetaSoundBaseFactory : public UFactory
 {
 	GENERATED_UCLASS_BODY()
 
-	// Set to initialize MetaSound type as Preset,
-	// using the provided MetaSound as a Reference
-	UPROPERTY(EditAnywhere, Transient, Category = Factory)
+	UPROPERTY(Transient, meta = (Deprecated = 5.8, DeprecationMessage = "Use document configuration instead"))
 	TObjectPtr<UObject> ReferencedMetaSoundObject;
+
+	UPROPERTY(EditAnywhere, Transient, meta = (Category = Factory))
+	TInstancedStruct<FMetaSoundFrontendDocumentConfiguration> DocumentConfiguration;
+
+	// If Document configuration is set, selection is provided to the configuration
+	// as an initial argument to be set on the configuration if applicable.
+	UPROPERTY(EditAnywhere, Transient, meta = (Category = Factory))
+	TArray<TObjectPtr<UObject>> SelectedObjects;
 };
 
 UCLASS(hidecategories=Object, MinimalAPI)
diff --git a/Engine/Plugins/Runtime/Metasound/Source/MetasoundEngine/Private/DocumentConfigurations/MetasoundRandomizerDocumentConfiguration.cpp b/Engine/Plugins/Runtime/Metasound/Source/MetasoundEngine/Private/DocumentConfigurations/MetasoundRandomizerDocumentConfiguration.cpp
--- a/Engine/Plugins/Runtime/Metasound/Source/MetasoundEngine/Private/DocumentConfigurations/MetasoundRandomizerDocumentConfiguration.cpp
+++ b/Engine/Plugins/Runtime/Metasound/Source/MetasoundEngine/Private/DocumentConfigurations/MetasoundRandomizerDocumentConfiguration.cpp
@@ -16,6 +16,8 @@
 #include "MetasoundFrontendDocumentBuilderRegistry.h"
 #include "MetasoundFrontendDocumentIdGenerator.h"
 #include "MetasoundWave.h"
+#include "Sound/SoundWave.h"
+#include "UObject/SoftObjectPath.h"
 
 
 #define LOCTEXT_NAMESPACE "MetaSound"
@@ -242,8 +244,8 @@ bool FMetaSoundRandomizerDocumentConfiguration::ConfigureDocument(FMetaSoundFron
 		{ .Name = "WavePlayer", .ClassName = PlayerClassName, .Location = FVector2D(500, 0) }
 	})
 
-	.Set(FDC::FNodeInputDefault { .Node = "Pitch", .Input = "Min", .Value = MakeFrontendLiteral(-12) })
-	.Set(FDC::FNodeInputDefault { .Node = "Pitch", .Input = "Max", .Value = MakeFrontendLiteral(12) })
+	.Set(FDC::FNodeInputDefault { .Node = "Pitch", .Input = "Min", .Value = MakeFrontendLiteral(FMath::Min(Pitch.X, Pitch.Y)) })
+	.Set(FDC::FNodeInputDefault { .Node = "Pitch", .Input = "Max", .Value = MakeFrontendLiteral(FMath::Max(Pitch.X, Pitch.Y)) })
 
 	// Connect Graph Inputs to Nodes
 	.Add<FDC::FEdgeFromInput>(
@@ -255,9 +257,9 @@ bool FMetaSoundRandomizerDocumentConfiguration::ConfigureDocument(FMetaSoundFron
 	// Connect Nodes Together
 	.Add<FDC::FEdge>(
 	{
-		{ .OutputNode = "RandomWave", .Output = "On Next",		.InputNode = "WavePlayer", .Input = "Play" },
-		{ .OutputNode = "RandomWave", .Output = "Value",		.InputNode = "WavePlayer", .Input = "Wave Asset" },
-		{ .OutputNode = "Pitch",	  .Output = "Value",		.InputNode = "WavePlayer", .Input = "Pitch Shift" },
+		{ .OutputNode = "RandomWave", .Output = "On Next",	.InputNode = "WavePlayer", .Input = "Play" },
+		{ .OutputNode = "RandomWave", .Output = "Value",	.InputNode = "WavePlayer", .Input = "Wave Asset" },
+		{ .OutputNode = "Pitch",	  .Output = "Value",	.InputNode = "WavePlayer", .Input = "Pitch Shift" },
 	})
 
 	// Connect Nodes to Graph Outputs
@@ -277,6 +279,19 @@ bool FMetaSoundRandomizerDocumentConfiguration::ConfigureDocument(FMetaSoundFron
 #endif // WITH_EDITORONLY_DATA
 
 #if WITH_EDITOR
+TArray<const UClass*> FMetaSoundRandomizerDocumentConfiguration::GetSupportedAssetActionClasses() const
+{
+	return { USoundWave::StaticClass() };
+}
+
+void FMetaSoundRandomizerDocumentConfiguration::OnAssetInitialized(TArray<UObject*> SelectedObjects, FMetaSoundFrontendDocumentBuilder& OutBuilder)
+{
+	Sounds.Reset();
+	Algo::TransformIf(SelectedObjects, Sounds,
+		[](const UObject* Object) { return Object->IsA<USoundWave>(); }, // Has to match type exactly so procedural sounds, MetaSounds, etc. aren't used
+		[](UObject* Object) { return TObjectPtr<USoundWave>(CastChecked<USoundWave>(Object)); });
+}
+
 void FMetaSoundRandomizerDocumentConfiguration::OnPropertyChanged(const FPropertyChangedEvent& InEvent, FMetaSoundFrontendDocumentBuilder& OutBuilder)
 {
 	ConfigureDocument(OutBuilder, false);
diff --git a/Engine/Plugins/Runtime/Metasound/Source/MetasoundEngine/Private/Metasound.cpp b/Engine/Plugins/Runtime/Metasound/Source/MetasoundEngine/Private/Metasound.cpp
--- a/Engine/Plugins/Runtime/Metasound/Source/MetasoundEngine/Private/Metasound.cpp
+++ b/Engine/Plugins/Runtime/Metasound/Source/MetasoundEngine/Private/Metasound.cpp
@@ -158,7 +158,9 @@ void UMetaSoundPatch::PostDuplicate(EDuplicateMode::Type InDuplicateMode)
 void UMetaSoundPatch::PostEditUndo()
 {
 	Super::PostEditUndo();
-	Metasound::Engine::FAssetHelper::PostEditUndo(*this);
+
+	TScriptInterface<IMetaSoundDocumentInterface> DocInterface(this);
+	Metasound::Engine::FAssetHelper::PostEditUndo(DocInterface);
 }
 
 EDataValidationResult UMetaSoundPatch::IsDataValid(FDataValidationContext& Context) const
diff --git a/Engine/Plugins/Runtime/Metasound/Source/MetasoundEngine/Private/MetasoundBuilderBase.cpp b/Engine/Plugins/Runtime/Metasound/Source/MetasoundEngine/Private/MetasoundBuilderBase.cpp
--- a/Engine/Plugins/Runtime/Metasound/Source/MetasoundEngine/Private/MetasoundBuilderBase.cpp
+++ b/Engine/Plugins/Runtime/Metasound/Source/MetasoundEngine/Private/MetasoundBuilderBase.cpp
@@ -6,6 +6,7 @@
 #include "Algo/Transform.h"
 #include "AudioDevice.h"
 #include "Components/AudioComponent.h"
+#include "DocumentConfigurations/MetasoundFrontendPresetDocumentConfiguration.h"
 #include "Engine/Engine.h"
 #include "HAL/IConsoleManager.h"
 #include "Interfaces/MetasoundOutputFormatInterfaces.h"
@@ -604,56 +605,68 @@ void UMetaSoundBuilderBase::ConnectGraphInputToNode(const FName& GraphInputName,
 
 void UMetaSoundBuilderBase::ConvertFromPreset(EMetaSoundBuilderResult& OutResult)
 {
-	const bool bSuccess = Builder.ConvertFromPreset();
-	OutResult = bSuccess ? EMetaSoundBuilderResult::Succeeded : EMetaSoundBuilderResult::Failed;
+	using namespace Metasound::Frontend;
+	if (Builder.IsDocumentConfiguredBy<FPresetDocumentConfiguration>())
+	{
+		Builder.SetDocumentConfiguration({ });
+		Builder.ConfigureDocument();
+
+#if WITH_EDITORONLY_DATA
+		Builder.SetIsGraphEditable(true);
+#endif // WITH_EDITORONLY_DATA
+
+		OutResult = EMetaSoundBuilderResult::Succeeded;
+		return;
+	}
+
+	OutResult = EMetaSoundBuilderResult::Failed;
 }
 
-void UMetaSoundBuilderBase::ConvertToPreset(const TScriptInterface<IMetaSoundDocumentInterface>& ReferencedNodeClass, EMetaSoundBuilderResult& OutResult)
+void UMetaSoundBuilderBase::ConvertToPreset(const TScriptInterface<IMetaSoundDocumentInterface>& ParentNodeClass, EMetaSoundBuilderResult& OutResult)
 {
 	using namespace Metasound::Frontend;
 
-	const IMetaSoundDocumentInterface* ReferencedInterface = ReferencedNodeClass.GetInterface();
-	if (!ReferencedInterface)
+	const IMetaSoundDocumentInterface* ParentInterface = ParentNodeClass.GetInterface();
+	if (!ParentInterface)
 	{
 		OutResult = EMetaSoundBuilderResult::Failed;
 		return;
 	}
 	
 	// Ensure the referenced node class isn't transient
-	if (Cast<UMetaSoundBuilderDocument>(ReferencedInterface))
+	if (Cast<UMetaSoundBuilderDocument>(ParentInterface))
 	{
-		UE_LOG(LogMetaSound, Warning, TEXT("Transient document builders cannot be referenced when converting builder '%s' to a preset. Build the referenced node class an asset first or use an existing asset instead"), *GetName());
+		UE_LOG(LogMetaSound, Warning, TEXT("Transient document builders cannot be referenced when converting builder '%s' to a preset. Build the parent node class as an asset first or use an existing asset instead"), *GetName());
 		OutResult = EMetaSoundBuilderResult::Failed;
 		return;
 	}
 
-	// Ensure the referenced node class is a matching object type 
-	const UClass& BaseMetaSoundClass = ReferencedInterface->GetBaseMetaSoundUClass();
-	UObject* ReferencedObject = ReferencedNodeClass.GetObject();
-	if (!ReferencedObject || (ReferencedObject && !ReferencedObject->IsA(&BaseMetaSoundClass)))
+	// Ensure the parent node class is a matching object type
+	const UClass& BaseMetaSoundClass = ParentInterface->GetBaseMetaSoundUClass();
+	UObject* ParentObject = ParentNodeClass.GetObject();
+	if (!ParentObject || (ParentObject && !ParentObject->IsA(&BaseMetaSoundClass)))
 	{
 		UE_LOG(LogMetaSound, Warning, TEXT("The referenced node type must match the base MetaSound class when converting builder '%s' to a preset (ex. source preset must reference another source)"), *GetName());
 		OutResult = EMetaSoundBuilderResult::Failed;
 		return;
 	}
 
-	// Ensure the referenced node is registered
-	if (FMetasoundAssetBase* ReferencedMetaSoundAsset = Metasound::IMetasoundUObjectRegistry::Get().GetObjectAsAssetBase(ReferencedObject))
+	// Ensure the parent node class is registered
+	if (FMetasoundAssetBase* ParentMetaSoundAsset = Metasound::IMetasoundUObjectRegistry::Get().GetObjectAsAssetBase(ParentObject))
 	{
 		Metasound::Frontend::FMetaSoundAssetRegistrationOptions RegOptions;
 		RegOptions.PageOrder = UMetaSoundSettings::GetPageOrder();
-		ReferencedMetaSoundAsset->UpdateAndRegisterForExecution(RegOptions);
+		ParentMetaSoundAsset->UpdateAndRegisterForExecution(RegOptions);
 	}
-#if WITH_EDITORONLY_DATA
-	const FMetaSoundFrontendDocumentBuilder& ReferencedBuilder = Metasound::Frontend::IDocumentBuilderRegistry::GetChecked().FindOrBeginBuilding(ReferencedObject);
-#else
-	const FMetaSoundFrontendDocumentBuilder& ReferencedBuilder = FMetaSoundFrontendDocumentBuilder(ReferencedObject);
-#endif // WITH_EDITORONLY_DATA
-	TSharedRef<FDocumentModifyDelegates> DocumentDelegates = MakeShared<FDocumentModifyDelegates>(ReferencedInterface->GetConstDocument());
+
+	TSharedRef<FDocumentModifyDelegates> DocumentDelegates = MakeShared<FDocumentModifyDelegates>(ParentInterface->GetConstDocument());
 	InitDelegates(*DocumentDelegates);
-	const bool bConvertToPreset = Builder.ConvertToPreset(ReferencedBuilder, DocumentDelegates);
-	ConformObjectToDocument();
-	OutResult = bConvertToPreset ? EMetaSoundBuilderResult::Succeeded : EMetaSoundBuilderResult::Failed;
+
+	TInstancedStruct<FMetaSoundFrontendPresetDocumentConfiguration> NewPresetConfig = TInstancedStruct<FMetaSoundFrontendPresetDocumentConfiguration>::Make();
+	NewPresetConfig.GetMutable().Parent = ParentObject;
+
+	Builder.SetDocumentConfiguration(MoveTemp(NewPresetConfig));
+	OutResult = Builder.ConfigureDocument() ? EMetaSoundBuilderResult::Succeeded : EMetaSoundBuilderResult::Failed;
 }
 
 void UMetaSoundBuilderBase::DisconnectNodes(const FMetaSoundBuilderNodeOutputHandle& NodeOutputHandle, const FMetaSoundBuilderNodeInputHandle& NodeInputHandle, EMetaSoundBuilderResult& OutResult)
@@ -1000,15 +1013,13 @@ int32 UMetaSoundBuilderBase::GetLastTransactionRegistered() const
 UObject* UMetaSoundBuilderBase::GetReferencedPresetAsset() const
 {
 	using namespace Metasound::Frontend;
-	if (!IsPreset())
+	const FPresetDocumentConfiguration* PresetConfig = GetConstBuilder()
+		.GetConstDocumentConfigurationAs<FPresetDocumentConfiguration>();
+	if (PresetConfig)
 	{
-		return nullptr;
+		return PresetConfig->Parent.GetObject();
 	}
 
-	if (FMetasoundAssetBase* Asset = Builder.GetReferencedPresetAsset())
-	{
-		return Asset->GetOwningAsset();
-	}
 	return nullptr;
 }
 
@@ -1061,7 +1072,7 @@ bool UMetaSoundBuilderBase::InterfaceIsDeclared(FName InterfaceName) const
 
 bool UMetaSoundBuilderBase::IsPreset() const
 {
-	return Builder.IsPreset();
+	return Builder.IsDocumentConfiguredBy<FMetaSoundFrontendPresetDocumentConfiguration>();
 }
 
 bool UMetaSoundBuilderBase::NodesAreConnected(const FMetaSoundBuilderNodeOutputHandle& OutputHandle, const FMetaSoundBuilderNodeInputHandle& InputHandle) const
diff --git a/Engine/Plugins/Runtime/Metasound/Source/MetasoundEngine/Private/MetasoundBuilderSubsystem.cpp b/Engine/Plugins/Runtime/Metasound/Source/MetasoundEngine/Private/MetasoundBuilderSubsystem.cpp
--- a/Engine/Plugins/Runtime/Metasound/Source/MetasoundEngine/Private/MetasoundBuilderSubsystem.cpp
+++ b/Engine/Plugins/Runtime/Metasound/Source/MetasoundEngine/Private/MetasoundBuilderSubsystem.cpp
@@ -6,6 +6,7 @@
 #include "AudioDevice.h"
 #include "Components/AudioComponent.h"
 #include "CoreGlobals.h"
+#include "DocumentConfigurations/MetasoundFrontendPresetDocumentConfiguration.h"
 #include "Engine/Engine.h"
 #include "HAL/IConsoleManager.h"
 #include "Interfaces/MetasoundOutputFormatInterfaces.h"
@@ -1007,8 +1008,9 @@ void UMetaSoundSourceBuilder::SetFormat(EMetaSoundOutputAudioFormat OutputFormat
 	using namespace Metasound::Engine;
 	using namespace Metasound::Frontend;
 
-	// Convert to non-preset MetaSoundSource since interface data is being altered
-	Builder.ConvertFromPreset();
+	// Clear document configuration since interface data is being altered
+	Builder.SetDocumentConfiguration({ });
+	Builder.ConfigureDocument();
 
 	const FOutputAudioFormatInfoMap& FormatMap = GetOutputAudioFormatInfo();
 
@@ -1373,16 +1375,17 @@ UMetaSoundBuilderBase* UMetaSoundBuilderSubsystem::FindParentBuilderOfPreset(con
 		return nullptr;
 	}
 
-	if (!Builder->IsPreset())
+	const FMetaSoundFrontendPresetDocumentConfiguration* PresetConfig = Builder->GetConstBuilder().GetConstDocumentConfigurationAs<FMetaSoundFrontendPresetDocumentConfiguration>();
+	if (!PresetConfig)
 	{
 		UE_LOG(LogMetaSound, Error, TEXT("MetaSound '%s' is not a preset."), *GetNameSafe(InMetaSoundPreset.GetObject()));
 		return nullptr;
 	}
 
-	TScriptInterface<IMetaSoundDocumentInterface> ParentAsset = Builder->GetReferencedPresetAsset();
+	TScriptInterface<IMetaSoundDocumentInterface> ParentAsset = PresetConfig->Parent;
 	if (!ParentAsset)
 	{
-		UE_LOG(LogMetaSound, Error, TEXT("Could not find referenced preset asset for MetaSound '%s'."), *GetNameSafe(InMetaSoundPreset.GetObject()));
+		UE_LOG(LogMetaSound, Error, TEXT("MetaSound preset '%s' contains null reference."), *GetNameSafe(InMetaSoundPreset.GetObject()));
 		return nullptr;
 	}
 
@@ -1394,19 +1397,27 @@ UMetaSoundBuilderBase* UMetaSoundBuilderSubsystem::FindParentBuilderOfPreset(con
 
 	// Follow preset chain to the highest non-preset ancestor
 	TSet<const UObject*> CheckedAssets;
-	while (ParentBuilder && ParentBuilder->IsPreset())
+	bool bIsCircular = false;
+	while (ParentBuilder)
 	{
 		// Check for circular references
-		if (CheckedAssets.Contains(ParentAsset.GetObject()))
+		CheckedAssets.Add(ParentAsset.GetObject(), &bIsCircular);
+		if (bIsCircular)
 		{
+			UE_LOG(LogMetaSound, Error, TEXT("MetaSound preset '%s' contains circular reference."), *GetNameSafe(ParentAsset.GetObject()));
 			return nullptr;
 		}
 
-		CheckedAssets.Add(ParentAsset.GetObject());
+		PresetConfig = ParentBuilder->GetConstBuilder().GetConstDocumentConfigurationAs<FMetaSoundFrontendPresetDocumentConfiguration>();
+		if (!PresetConfig)
+		{
+			break;
+		}
 
-		ParentAsset = ParentBuilder->GetReferencedPresetAsset();
+		ParentAsset = PresetConfig->Parent;
 		if (!ParentAsset)
 		{
+			UE_LOG(LogMetaSound, Warning, TEXT("MetaSound preset '%s' contains null reference."), *GetNameSafe(ParentAsset.GetObject()));
 			return nullptr;
 		}
 
diff --git a/Engine/Plugins/Runtime/Metasound/Source/MetasoundEngine/Private/MetasoundDocumentBuilderRegistry.cpp b/Engine/Plugins/Runtime/Metasound/Source/MetasoundEngine/Private/MetasoundDocumentBuilderRegistry.cpp
--- a/Engine/Plugins/Runtime/Metasound/Source/MetasoundEngine/Private/MetasoundDocumentBuilderRegistry.cpp
+++ b/Engine/Plugins/Runtime/Metasound/Source/MetasoundEngine/Private/MetasoundDocumentBuilderRegistry.cpp
@@ -193,8 +193,7 @@ namespace Metasound::Engine
 
 				auto RemoveDefault = [&Builder, &GraphInput](const FGuid& InPageID)
 				{
-					const bool bClearInheritsDefault = false;
-					return Builder.RemoveGraphInputDefault(GraphInput.Name, InPageID, bClearInheritsDefault);
+					return Builder.RemoveGraphInputDefault(GraphInput.Name, InPageID);
 				};
 
 				const int32 NumInitDefaults = GraphInput.GetDefaults().Num();
diff --git a/Engine/Plugins/Runtime/Metasound/Source/MetasoundEngine/Private/MetasoundEngineAsset.h b/Engine/Plugins/Runtime/Metasound/Source/MetasoundEngine/Private/MetasoundEngineAsset.h
--- a/Engine/Plugins/Runtime/Metasound/Source/MetasoundEngine/Private/MetasoundEngineAsset.h
+++ b/Engine/Plugins/Runtime/Metasound/Source/MetasoundEngine/Private/MetasoundEngineAsset.h
@@ -2,6 +2,7 @@
 
 #pragma once
 
+#include "DocumentConfigurations/MetasoundFrontendDocumentConfiguration.h"
 #include "Metasound.h"
 #include "MetasoundAssetManager.h"
 #include "MetasoundAssetSubsystem.h"
@@ -16,6 +17,7 @@
 #include "Misc/App.h"
 #include "Modules/ModuleManager.h"
 #include "Serialization/Archive.h"
+#include "StructUtils/StructView.h"
 #include "UObject/AssetRegistryTagsContext.h"
 
 #if WITH_EDITORONLY_DATA
@@ -71,15 +73,24 @@ namespace Metasound::Engine
 			}
 		}
 
-		template <typename TMetaSoundObject>
-		static void PostEditUndo(TMetaSoundObject& InMetaSound)
+		static void PostEditUndo(TScriptInterface<IMetaSoundDocumentInterface>& InDocInterface)
 		{
-			InMetaSound.GetModifyContext().SetForceRefreshViews();
+			using namespace Frontend;
+
+			InDocInterface->GetConstDocument().Metadata.ModifyContext.SetForceRefreshViews();
+
+			const FMetasoundFrontendClassName& ClassName = InDocInterface->GetConstDocument().RootGraph.Metadata.GetClassName();
+			IDocumentBuilderRegistry& BuilderRegistry = IDocumentBuilderRegistry::GetChecked();
+			
+			BuilderRegistry.ReloadBuilder(ClassName);
+
+			FMetaSoundFrontendDocumentBuilder& Builder = BuilderRegistry.FindOrBeginBuilding(InDocInterface);
 
-			const FMetasoundFrontendClassName& ClassName = InMetaSound.GetConstDocument().RootGraph.Metadata.GetClassName();
-			Frontend::IDocumentBuilderRegistry::GetChecked().ReloadBuilder(ClassName);
+#if WITH_EDITORONLY_DATA
+			Builder.ConfigureDocument();
+#endif // WITH_EDITORONLY_DATA
 
-			if (UMetasoundEditorGraphBase* Graph = Cast<UMetasoundEditorGraphBase>(InMetaSound.GetGraph()))
+			if (UMetasoundEditorGraphBase* Graph = Cast<UMetasoundEditorGraphBase>(Builder.GetMetasoundAsset().GetGraph()))
 			{
 				Graph->RegisterGraphWithFrontend();
 			}
diff --git a/Engine/Plugins/Runtime/Metasound/Source/MetasoundEngine/Private/MetasoundSource.cpp b/Engine/Plugins/Runtime/Metasound/Source/MetasoundEngine/Private/MetasoundSource.cpp
--- a/Engine/Plugins/Runtime/Metasound/Source/MetasoundEngine/Private/MetasoundSource.cpp
+++ b/Engine/Plugins/Runtime/Metasound/Source/MetasoundEngine/Private/MetasoundSource.cpp
@@ -8,6 +8,7 @@
 #include "AssetRegistry/AssetRegistryModule.h"
 #include "AudioDeviceManager.h"
 #include "Containers/Ticker.h"
+#include "DocumentConfigurations/MetasoundFrontendPresetDocumentConfiguration.h"
 #include "IAudioParameterInterfaceRegistry.h"
 #include "Interfaces/ITargetPlatform.h"
 #include "Interfaces/MetasoundOutputFormatInterfaces.h"
@@ -391,22 +392,47 @@ public:
 	FAudioParameterCollector& operator=(FAudioParameterCollector&&) = delete;
 
 	// Merge in parameters from this preset.
-	void CollectPresetOverrides(const TSet<Metasound::FVertexName>& InInputsInheritingDefault, const Metasound::TSortedVertexNameMap<FRuntimeInput>& InInputMap)
+	void CollectPresetOverrides(const FMetasoundFrontendDocument& Document, const Metasound::TSortedVertexNameMap<FRuntimeInput>& InInputMap)
 	{
 		using namespace Metasound;
+		using namespace Metasound::Frontend;
+
+		const FPresetDocumentConfiguration* PresetConfig = Document.Configuration.GetPtr<FPresetDocumentConfiguration>();
+		if (!PresetConfig)
+		{
+			return;
+		}
+
+		if (InInputMap.IsEmpty())
+		{
+			return;
+		}
+
+		TArray<const FMetasoundFrontendClassInput*> ClassInputs;
+		Algo::Transform(Document.RootGraph.GetDefaultInterface().Inputs, ClassInputs, [](const FMetasoundFrontendClassInput& Input) { return &Input; });
 
 		for (const TPair<FVertexName, FRuntimeInput>& Pair : InInputMap)
 		{
 			// Avoid adding parameters which already exist in the ParameterArray
 			int32 InsertPos = Algo::LowerBoundBy(Parameters, Pair.Key, FProjection(), FSortPredicate());
-			const bool bDoesNotAlreadyExist = (InsertPos >= Parameters.Num()) || (Pair.Key != Parameters[InsertPos].ParamName);
 
+			const bool bDoesNotAlreadyExist = (InsertPos >= Parameters.Num()) || (Pair.Key != Parameters[InsertPos].ParamName);
 			if (bDoesNotAlreadyExist)
 			{
-				// Avoid adding parameters which are NOT overriding a default. 
-				if (!InInputsInheritingDefault.Contains(Pair.Key))
+				auto IsInput = [&Pair](const FMetasoundFrontendClassInput* Input) { return Input->Name == Pair.Key; };
+				const int32 ClassIndex = ClassInputs.IndexOfByPredicate(IsInput);
+				if (ClassIndex != INDEX_NONE)
 				{
-					Parameters.Insert(Pair.Value.DefaultParameter, InsertPos);
+					// Avoid adding parameters which are NOT overriding a default.
+					const FGuid& NodeID = ClassInputs[ClassIndex]->NodeID;
+					if (const FPresetVertexMetadata* VertexMetadata = PresetConfig->FindConstVertexMetadata<FPresetVertexMetadata>(NodeID))
+					{
+						if (VertexMetadata->bOverrideInheritedDefault)
+						{
+							Parameters.Insert(Pair.Value.DefaultParameter, InsertPos);
+						}
+					}
+					ClassInputs.RemoveAtSwap(ClassIndex, EAllowShrinking::No);
 				}
 			}
 		}
@@ -456,7 +482,9 @@ const FMetasoundFrontendDocument& UMetaSoundSource::GetConstDocument() const
 void UMetaSoundSource::PostEditUndo()
 {
 	Super::PostEditUndo();
-	Metasound::Engine::FAssetHelper::PostEditUndo(*this);
+
+	TScriptInterface<IMetaSoundDocumentInterface> DocInterface(this);
+	Metasound::Engine::FAssetHelper::PostEditUndo(DocInterface);
 }
 
 void UMetaSoundSource::PreDuplicate(FObjectDuplicationParameters& DupParams)
@@ -499,7 +527,7 @@ bool UMetaSoundSource::CanEditChange(const FProperty* InProperty) const
 	{
 		const TArray<FName> Platforms = FDataDrivenPlatformInfoRegistry::GetSortedPlatformNames(EPlatformInfoType::AllPlatformInfos);
 		const int32 DefaultBlockRate = BlockRateOverride.GetDefault();
-		const float DefaultSampleRate = BlockRateOverride.GetDefault();
+		const float DefaultSampleRate = SampleRateOverride.GetDefault();
 
 		if (DefaultBlockRate > 0 && DefaultSampleRate > 0)
 		{
@@ -824,12 +852,11 @@ void UMetaSoundSource::ResolveQualitySettings(const UMetaSoundSettings* Settings
 	auto FindByName = [&Name = QualitySetting](const FMetaSoundQualitySettings& Q) -> bool { return Q.Name == Name; };
 	Resolved = Settings->GetQualitySettings().FindByPredicate(FindByName);
 
-
 	// 2. If that failed, try by guid (if its been renamed in the settings, we can still find it).
 	if (!Resolved && QualitySettingGuid.IsValid())
 	{
 		auto FindByGuid = [&Guid = QualitySettingGuid](const FMetaSoundQualitySettings& Q) -> bool { return Q.UniqueId == Guid; };
-		Resolved = Settings->GetQualitySettings().FindByPredicate(FindByName);
+		Resolved = Settings->GetQualitySettings().FindByPredicate(FindByGuid);
 	}
 
 	// 3. If still failed to resolve, use defaults and warn.
@@ -846,7 +873,7 @@ void UMetaSoundSource::ResolveQualitySettings(const UMetaSoundSettings* Settings
 			// Default doesn't point to anything, use first one in the list.
 			QualitySetting = Settings->GetQualitySettings()[0].Name;
 			QualitySettingGuid = Settings->GetQualitySettings()[0].UniqueId;
-		}				
+		}
 	}
 
 	// Refresh the guid/name now we've resolved to correctly reflect.
@@ -855,7 +882,6 @@ void UMetaSoundSource::ResolveQualitySettings(const UMetaSoundSettings* Settings
 		QualitySetting = Resolved->Name;
 		QualitySettingGuid = Resolved->UniqueId;
 	}
-
 #endif //WITH_EDITORONLY_DATA
 }
 
@@ -1881,8 +1907,9 @@ Metasound::TSortedVertexNameMap<UMetaSoundSource::FRuntimeInput> UMetaSoundSourc
 void UMetaSoundSource::CacheRuntimeInputData(TArrayView<const FGuid> InPageOrder)
 {
 	using namespace Metasound;
+	using namespace Metasound::Frontend;
 
-	constexpr bool bCreateUObjectProxies = true; 
+	constexpr bool bCreateUObjectProxies = true;
 	RuntimeInputData.InputMap = CreateRuntimeInputMap(bCreateUObjectProxies, InPageOrder);
 
 	// Determine if preset graph inflation is possible 
@@ -1893,17 +1920,26 @@ void UMetaSoundSource::CacheRuntimeInputData(TArrayView<const FGuid> InPageOrder
 	// 
 	// Operator caching of base preset graphs fail when there are constructor inputs because
 	// constructor inputs set on the preset cannot be updated after the base operator is 
-	// cached. 
-	auto IsOverriddenConstructorInput = [&InputsInheritingDefault=RootMetasoundDocument.RootGraph.PresetOptions.InputsInheritingDefault](const TPair<FVertexName, FRuntimeInput>& Pair) 
+	// cached.
+
+	bIsPresetGraphInflationSupported = false;
+	if (const FPresetDocumentConfiguration* PresetConfig = RootMetasoundDocument.Configuration.GetPtr<FPresetDocumentConfiguration>())
 	{
-		if (Pair.Value.AccessType == EMetasoundFrontendVertexAccessType::Value)
+		auto IsOverriddenConstructorInput = [this, &PresetConfig](const FMetasoundFrontendClassInput& Input)
 		{
-			return !InputsInheritingDefault.Contains(Pair.Key);
-		}
-		return false;
-	};
+			if (const FRuntimeInput* RuntimeInput = RuntimeInputData.InputMap.Find(Input.Name))
+			{
+				if (const FPresetVertexMetadata* VertexMetadata = PresetConfig->FindConstVertexMetadata<FPresetVertexMetadata>(Input.NodeID))
+				{
+					return VertexMetadata->bOverrideInheritedDefault && RuntimeInput->AccessType == EMetasoundFrontendVertexAccessType::Value;
+				}
+			}
+			return false;
+		};
 
-	bIsPresetGraphInflationSupported = RootMetasoundDocument.RootGraph.PresetOptions.bIsPreset && !Algo::AnyOf(RuntimeInputData.InputMap, IsOverriddenConstructorInput);
+		const TArray<FMetasoundFrontendClassInput>& Inputs = RootMetasoundDocument.RootGraph.GetDefaultInterface().Inputs;
+		bIsPresetGraphInflationSupported = !Algo::AnyOf(Inputs, IsOverriddenConstructorInput);
+	}
 
 	RuntimeInputData.bIsValid.store(true);
 }
@@ -1973,7 +2009,7 @@ TSharedPtr<const Metasound::IGraph> UMetaSoundSource::FindFirstNoninflatableGrap
 	auto OnGraphInflation = [&InOutParameterCollector](const UMetaSoundSource& InInflatedSource)
 	{
 		// Any preset overrides on this object need to be baked in to the parameters
-		InOutParameterCollector.CollectPresetOverrides(InInflatedSource.RootMetasoundDocument.RootGraph.PresetOptions.InputsInheritingDefault, InInflatedSource.RuntimeInputData.InputMap);
+		InOutParameterCollector.CollectPresetOverrides(InInflatedSource.RootMetasoundDocument, InInflatedSource.RuntimeInputData.InputMap);
 	};
 
 	const UMetaSoundSource& NoninflatableSource = FindFirstNoninflatableSource(InOutEnvironment, OnGraphInflation);
diff --git a/Engine/Plugins/Runtime/Metasound/Source/MetasoundEngine/Public/DocumentConfigurations/MetasoundDocumentConfigurator.h b/Engine/Plugins/Runtime/Metasound/Source/MetasoundEngine/Public/DocumentConfigurations/MetasoundDocumentConfigurator.h
--- a/Engine/Plugins/Runtime/Metasound/Source/MetasoundEngine/Public/DocumentConfigurations/MetasoundDocumentConfigurator.h
+++ b/Engine/Plugins/Runtime/Metasound/Source/MetasoundEngine/Public/DocumentConfigurations/MetasoundDocumentConfigurator.h
@@ -35,9 +35,10 @@ namespace Metasound::Engine
 
 #if WITH_EDITORONLY_DATA
 		UE_EXPERIMENTAL(5.8, "Configurator code reflection is in development and may change implementation or be removed in the future")
-		UE_API static void Reflect(const FMetasoundFrontendGraphLayer& InLayer, FString& OutCode);
+		UE_API static void Reflect(FMetasoundFrontendGraphLayer InLayer, FString& OutCode);
 #endif // WITH_EDITORONLY_DATA
 
+	public:
 		struct FInterface
 		{
 			FMetasoundFrontendVersion Version;
@@ -47,22 +48,33 @@ namespace Metasound::Engine
 
 		struct FInput
 		{
-			FName Name = "Float";
+			FName Name = "FloatInput";
 			FName DataType = GetMetasoundDataTypeName<float>();
 			FMetasoundFrontendLiteral DefaultValue = Frontend::MakeFrontendLiteral(0.f);
 			FVector2D Location = FVector2D::Zero();
-			EMetasoundFrontendVertexAccessType AccessType = EMetasoundFrontendVertexAccessType::Reference;
+			bool bIsConstructorInput = false;
 		};
 		UE_EXPERIMENTAL(5.8, "Configurator is in active development and subject to change")
 		UE_API FDocumentConfigurator& Add(FInput Input);
 
+#if WITH_EDITORONLY_DATA
+		struct FInputNode
+		{
+			FName NodeName = "FloatInputNode";
+			FName InputName = "FloatInput";
+			FVector2D Location = FVector2D::Zero();
+		};
+		UE_EXPERIMENTAL(5.8, "Configurator is in active development and subject to change")
+		UE_API FDocumentConfigurator& Add(FInputNode InputNode);
+#endif // WITH_EDITORONLY_DATA
+
 		struct FOutput
 		{
-			FName Name = "Float";
+			FName Name = "FloatOutput";
 			FName DataType = GetMetasoundDataTypeName<float>();
 			FMetasoundFrontendLiteral DefaultValue = Frontend::MakeFrontendLiteral(0.f);
 			FVector2D Location = FVector2D::Zero();
-			EMetasoundFrontendVertexAccessType AccessType = EMetasoundFrontendVertexAccessType::Reference;
+			bool bIsConstructorOutput = false;
 		};
 		UE_EXPERIMENTAL(5.8, "Configurator is in active development and subject to change")
 		UE_API FDocumentConfigurator& Add(FOutput Output);
@@ -71,6 +83,7 @@ namespace Metasound::Engine
 		{
 			FName Name;
 			FMetasoundFrontendClassName ClassName;
+			FTopLevelAssetPath AssetPath;
 			FVector2D Location = FVector2D::Zero();
 		};
 		UE_EXPERIMENTAL(5.8, "Configurator is in active development and subject to change")
diff --git a/Engine/Plugins/Runtime/Metasound/Source/MetasoundEngine/Public/DocumentConfigurations/MetasoundRandomizerDocumentConfiguration.h b/Engine/Plugins/Runtime/Metasound/Source/MetasoundEngine/Public/DocumentConfigurations/MetasoundRandomizerDocumentConfiguration.h
--- a/Engine/Plugins/Runtime/Metasound/Source/MetasoundEngine/Public/DocumentConfigurations/MetasoundRandomizerDocumentConfiguration.h
+++ b/Engine/Plugins/Runtime/Metasound/Source/MetasoundEngine/Public/DocumentConfigurations/MetasoundRandomizerDocumentConfiguration.h
@@ -1,8 +1,8 @@
 // Copyright Epic Games, Inc. All Rights Reserved.
 #pragma once
 
+#include "DocumentConfigurations/MetasoundFrontendDocumentConfiguration.h"
 #include "MetasoundDocumentInterface.h"
-#include "MetasoundFrontendDocumentConfiguration.h"
 #include "Sound/SoundWave.h"
 #include "UObject/ObjectPtr.h"
 
@@ -23,11 +23,16 @@ struct FMetaSoundRandomizerDocumentConfiguration : public FMetaSoundFrontendDocu
 	virtual ~FMetaSoundRandomizerDocumentConfiguration() = default;
 
 #if WITH_EDITORONLY_DATA
-	UE_EXPERIMENTAL(5.8, "Document Configuration is in active development and subject to change")
-	UE_API virtual bool ConfigureDocument(FMetaSoundFrontendDocumentBuilder& OutBuilder, bool bLogFailure) override;
+	UE_INTERNAL UE_API virtual bool ConfigureDocument(FMetaSoundFrontendDocumentBuilder& OutBuilder, bool bLogFailure) override;
 #endif // WITH_EDITORONLY_DATA
 
 #if WITH_EDITOR
+	UE_EXPERIMENTAL(5.8, "Document Configuration is in active development and subject to change")
+	UE_API virtual TArray<const UClass*> GetSupportedAssetActionClasses() const override;
+
+	UE_EXPERIMENTAL(5.8, "Document Configuration is in active development and subject to change")
+	UE_API virtual void OnAssetInitialized(TArray<UObject*> SelectedObjects, FMetaSoundFrontendDocumentBuilder& OutBuilder) override;
+
 	UE_EXPERIMENTAL(5.8, "Document Configuration is in active development and subject to change")
 	UE_API virtual void OnPropertyChanged(const FPropertyChangedEvent& InEvent, FMetaSoundFrontendDocumentBuilder& OutBuilder) override;
 #endif // WITH_EDITOR
@@ -38,5 +43,8 @@ public:
 
 	UPROPERTY(EditAnywhere, BlueprintReadOnly, meta = (Category = "Randomizer Options"))
 	TArray<TObjectPtr<USoundWave>> Sounds;
+
+	UPROPERTY(EditAnywhere, BlueprintReadOnly, meta = (Category = "Randomizer Options", ClampMin = "-36.0", ClampMax = "36.0", UIMin = "-12.0", UIMax = "12.0"))
+	FVector2f Pitch = { -3.f, 3.f };
 };
 #undef UE_API
\ No newline at end of file
diff --git a/Engine/Plugins/Runtime/Metasound/Source/MetasoundEngine/Public/MetasoundBuilderBase.h b/Engine/Plugins/Runtime/Metasound/Source/MetasoundEngine/Public/MetasoundBuilderBase.h
--- a/Engine/Plugins/Runtime/Metasound/Source/MetasoundEngine/Public/MetasoundBuilderBase.h
+++ b/Engine/Plugins/Runtime/Metasound/Source/MetasoundEngine/Public/MetasoundBuilderBase.h
@@ -396,9 +396,9 @@ public:
 	UFUNCTION(BlueprintCallable, Category = "Audio|MetaSound|Builder", meta = (ExpandEnumAsExecs = "OutResult"))
 	UE_API void ConvertFromPreset(EMetaSoundBuilderResult& OutResult);
 
-	// Convert this builder to a MetaSound source preset with the given referenced source builder 
+	// Convert this builder to a MetaSound preset using the given MetaSoundDocumentInterface as its parent.
 	UFUNCTION(BlueprintCallable, Category = "Audio|MetaSound|Builder", meta = (ExpandEnumAsExecs = "OutResult"))
-	UE_API void ConvertToPreset(const TScriptInterface<IMetaSoundDocumentInterface>& ReferencedNodeClass, EMetaSoundBuilderResult& OutResult);
+	UE_API void ConvertToPreset(UPARAM(DisplayName = "Parent Node Class") const TScriptInterface<IMetaSoundDocumentInterface>& ReferencedNodeClass, EMetaSoundBuilderResult& OutResult);
 
 #if WITH_EDITORONLY_DATA
 	// Removes all graph pages except the default.  If bClearDefaultPage is true, clears the default graph page implementation.
diff --git a/Engine/Plugins/Runtime/Metasound/Source/MetasoundEngine/Public/MetasoundDocumentBuilderRegistry.h b/Engine/Plugins/Runtime/Metasound/Source/MetasoundEngine/Public/MetasoundDocumentBuilderRegistry.h
--- a/Engine/Plugins/Runtime/Metasound/Source/MetasoundEngine/Public/MetasoundDocumentBuilderRegistry.h
+++ b/Engine/Plugins/Runtime/Metasound/Source/MetasoundEngine/Public/MetasoundDocumentBuilderRegistry.h
@@ -6,6 +6,7 @@
 #include "MetasoundBuilderBase.h"
 #include "MetasoundDocumentInterface.h"
 #include "MetasoundFrontendDocument.h"
+#include "MetasoundFrontendDocumentBuilderRegistry.h"
 #include "MetasoundFrontendRegistries.h"
 #include "MetasoundLog.h"
 #include "MetasoundSettings.h"
diff --git a/Engine/Plugins/Runtime/Metasound/Source/MetasoundFrontend/Private/DocumentConfigurations/MetasoundFrontendDocumentConfiguration.cpp b/Engine/Plugins/Runtime/Metasound/Source/MetasoundFrontend/Private/DocumentConfigurations/MetasoundFrontendDocumentConfiguration.cpp
new file mode 100644
--- /dev/null
+++ b/Engine/Plugins/Runtime/Metasound/Source/MetasoundFrontend/Private/DocumentConfigurations/MetasoundFrontendDocumentConfiguration.cpp
@@ -0,0 +1,42 @@
+// Copyright Epic Games, Inc. All Rights Reserved.
+
+#include "DocumentConfigurations/MetasoundFrontendDocumentConfiguration.h"
+
+#include "MetasoundDocumentInterface.h"
+#include "MetasoundFrontendDocument.h"
+#include "MetasoundFrontendDocumentBuilder.h"
+#include "MetasoundFrontendDocumentBuilderRegistry.h"
+
+
+#define LOCTEXT_NAMESPACE "MetaSound"
+
+bool FMetaSoundFrontendDocumentConfiguration::ConfigureDocument(FMetaSoundFrontendDocumentBuilder& OutBuilder, bool bLogFailure)
+{
+#if WITH_EDITORONLY_DATA
+	OutBuilder.SetIsGraphEditable(false);
+#endif // WITH_EDITORONLY_DATA
+
+	return true;
+}
+
+#if WITH_EDITOR
+TArray<const UClass*> FMetaSoundFrontendDocumentConfiguration::GetSupportedAssetActionClasses() const
+{
+	return { };
+}
+
+TArray<const UClass*> FMetaSoundFrontendDocumentConfiguration::GetSupportedMetaSoundClasses() const
+{
+	return { };
+}
+
+FText FMetaSoundFrontendDocumentConfiguration::GetToolTip() const
+{
+	return { };
+}
+
+void FMetaSoundFrontendDocumentConfiguration::OnPropertyChanged(const FPropertyChangedEvent& InEvent, FMetaSoundFrontendDocumentBuilder& OutBuilder)
+{
+}
+#endif // WITH_EDITOR
+#undef LOCTEXT_NAMESPACE
diff --git a/Engine/Plugins/Runtime/Metasound/Source/MetasoundFrontend/Private/DocumentConfigurations/MetasoundFrontendPresetDocumentConfiguration.cpp b/Engine/Plugins/Runtime/Metasound/Source/MetasoundFrontend/Private/DocumentConfigurations/MetasoundFrontendPresetDocumentConfiguration.cpp
new file mode 100644
--- /dev/null
+++ b/Engine/Plugins/Runtime/Metasound/Source/MetasoundFrontend/Private/DocumentConfigurations/MetasoundFrontendPresetDocumentConfiguration.cpp
@@ -0,0 +1,142 @@
+// Copyright Epic Games, Inc. All Rights Reserved.
+#include "DocumentConfigurations/MetasoundFrontendPresetDocumentConfiguration.h"
+
+#include "Algo/AnyOf.h"
+#include "MetasoundDataTypeGetTraits.h"
+#include "MetasoundDocumentInterface.h"
+#include "MetasoundFrontendDocument.h"
+#include "MetasoundFrontendDocumentBuilder.h"
+#include "MetasoundFrontendDocumentBuilderRegistry.h"
+#include "MetasoundTrigger.h"
+
+
+#define LOCTEXT_NAMESPACE "MetaSound"
+
+namespace Metasound::Frontend::PresetDocConfigPrivate
+{
+	bool IsInheritableInputNode(const FMetaSoundFrontendDocumentBuilder& Builder, const FGuid& NodeID)
+	{
+		const FMetasoundFrontendNode* Node = Builder.FindNode(NodeID);
+		if (!Node)
+		{
+			return false;
+		}
+
+		const FMetasoundFrontendClass* Class = Builder.FindDependency(Node->ClassID);
+		if (!Class)
+		{
+			return false;
+		}
+
+		if (Class->Metadata.GetType() != EMetasoundFrontendClassType::Input)
+		{
+			return false;
+		}
+
+		const FMetasoundFrontendClassInput* Input = Builder.FindGraphInput(Node->Name);
+		if (!Input)
+		{
+			return false;
+		}
+
+		auto IsDefaultConstructed = [](const FMetasoundFrontendClassInputDefault& Default) { return Default.Literal.GetType() == EMetasoundFrontendLiteralType::None; };
+		if (Algo::AnyOf(Input->GetDefaults(), IsDefaultConstructed))
+		{
+			return false;
+		}
+
+		if (Input->TypeName == GetMetasoundDataTypeName<FTrigger>())
+		{
+			return false;
+		}
+
+		return true;
+	}
+} // namespace Metasound::Frontend::PresetDocConfigPrivate
+
+#if WITH_EDITORONLY_DATA
+bool FMetaSoundFrontendPresetDocumentConfiguration::ConfigureDocument(FMetaSoundFrontendDocumentBuilder& OutBuilder, bool bLogFailure)
+{
+	using namespace Metasound::Frontend;
+
+	if (UObject* RefObject = Parent.GetObject())
+	{
+		if (RefObject->IsA(&OutBuilder.GetConstDocumentInterfaceChecked().GetBaseMetaSoundUClass()))
+		{
+			const FMetaSoundFrontendDocumentBuilder& ParentBuilder = IDocumentBuilderRegistry::GetChecked().FindOrBeginBuilding(Parent);
+
+			OutBuilder.SetIsGraphEditable(false);
+
+			// Apply root graph transform
+			FRebuildPresetRootGraph RebuildPresetRootGraph(ParentBuilder);
+			return RebuildPresetRootGraph.Transform(OutBuilder);
+		}
+		else
+		{
+			if (bLogFailure)
+			{
+				UE_LOG(LogMetaSound, Error, TEXT("The referenced node type must match the base MetaSound class when converting builder '%s' to a preset (ex. source preset must reference another source)"), *OutBuilder.GetDebugName());
+			}
+		}
+	}
+	else if (bLogFailure)
+	{
+		UE_LOG(LogMetaSound, Error, TEXT("Preset '%s' generation failed: Referenced class not set."), *OutBuilder.GetDebugName());
+	}
+
+	return false;
+}
+#endif // WITH_EDITORONLY_DATA
+
+TInstancedStruct<FMetaSoundFrontendDocumentConfigurationVertexMetadata> FMetaSoundFrontendPresetDocumentConfiguration::ConstructVertexMetadata(const FMetaSoundFrontendDocumentBuilder& Builder, const FGuid& NodeID) const
+{
+	using namespace Metasound::Frontend;
+
+	// Presets only need ConfigVertexMetadata for inputs, so don't construct otherwise
+	if (PresetDocConfigPrivate::IsInheritableInputNode(Builder, NodeID))
+	{
+		return TInstancedStruct<FPresetVertexMetadata>::Make();
+	}
+
+	return { };
+}
+
+#if WITH_EDITOR
+TArray<const UClass*> FMetaSoundFrontendPresetDocumentConfiguration::GetSupportedAssetActionClasses() const
+{
+	return { };
+}
+
+void FMetaSoundFrontendPresetDocumentConfiguration::OnAssetInitialized(TArray<UObject*> SelectedObjects, FMetaSoundFrontendDocumentBuilder& OutBuilder)
+{
+	for (UObject* SelectedObject : SelectedObjects)
+	{
+		if (&OutBuilder.GetConstDocumentInterfaceChecked().GetBaseMetaSoundUClass() == SelectedObject->GetClass())
+		{
+			Parent = SelectedObject;
+			return;
+		}
+	}
+}
+
+void FMetaSoundFrontendPresetDocumentConfiguration::OnPropertyChanged(const FPropertyChangedEvent& InEvent, FMetaSoundFrontendDocumentBuilder& OutBuilder)
+{
+	ConfigureDocument(OutBuilder, false);
+}
+
+bool FMetaSoundFrontendPresetConfigurationVertexMetadata::IsDefaultEditable() const
+{
+	return bOverrideInheritedDefault;
+}
+#endif // WITH_EDITOR
+
+void FMetaSoundFrontendPresetConfigurationVertexMetadata::OnDefaultReset(const FMetaSoundFrontendDocumentBuilder& Builder)
+{
+	bOverrideInheritedDefault = false;
+}
+
+void FMetaSoundFrontendPresetConfigurationVertexMetadata::OnDefaultUpdated(const FMetaSoundFrontendDocumentBuilder& Builder)
+{
+	bOverrideInheritedDefault = true;
+}
+#undef LOCTEXT_NAMESPACE
diff --git a/Engine/Plugins/Runtime/Metasound/Source/MetasoundFrontend/Private/DocumentConfigurations/MetasoundFrontendVertexDocumentConfiguration.cpp b/Engine/Plugins/Runtime/Metasound/Source/MetasoundFrontend/Private/DocumentConfigurations/MetasoundFrontendVertexDocumentConfiguration.cpp
new file mode 100644
--- /dev/null
+++ b/Engine/Plugins/Runtime/Metasound/Source/MetasoundFrontend/Private/DocumentConfigurations/MetasoundFrontendVertexDocumentConfiguration.cpp
@@ -0,0 +1,46 @@
+// Copyright Epic Games, Inc. All Rights Reserved.
+
+#include "DocumentConfigurations/MetasoundFrontendDocumentVertexConfiguration.h"
+
+#include "MetasoundDocumentInterface.h"
+#include "MetasoundFrontendDocument.h"
+#include "MetasoundFrontendDocumentBuilder.h"
+#include "MetasoundFrontendDocumentBuilderRegistry.h"
+
+
+#define LOCTEXT_NAMESPACE "MetaSound"
+
+FMetaSoundFrontendDocumentConfigurationVertexMetadata* FMetaSoundFrontendDocumentVertexConfiguration::AddVertexMetadata(const FMetaSoundFrontendDocumentBuilder& Builder, const FGuid& NodeID)
+{
+	if (TInstancedStruct<FMetaSoundFrontendDocumentConfigurationVertexMetadata>* Entry = VertexMetadata.Find(NodeID))
+	{
+		return Entry->GetMutablePtr();
+	}
+
+PRAGMA_DISABLE_EXPERIMENTAL_WARNINGS
+	TInstancedStruct<FMetaSoundFrontendDocumentConfigurationVertexMetadata> NewMetadata = ConstructVertexMetadata(Builder, NodeID);
+PRAGMA_ENABLE_EXPERIMENTAL_WARNINGS
+
+	if (NewMetadata.IsValid())
+	{
+		return VertexMetadata.Add(NodeID, MoveTemp(NewMetadata)).GetMutablePtr();
+	}
+
+	return nullptr;
+}
+
+TInstancedStruct<FMetaSoundFrontendDocumentConfigurationVertexMetadata> FMetaSoundFrontendDocumentVertexConfiguration::ConstructVertexMetadata(const FMetaSoundFrontendDocumentBuilder& Builder, const FGuid& NodeID) const
+{
+	return { };
+}
+
+bool FMetaSoundFrontendDocumentVertexConfiguration::ContainsVertexMetadata(const FGuid& NodeID) const
+{
+	return VertexMetadata.Contains(NodeID);
+}
+
+bool FMetaSoundFrontendDocumentVertexConfiguration::RemoveVertexMetadata(const FGuid& NodeID)
+{
+	return VertexMetadata.Remove(NodeID) > 0;
+}
+#undef LOCTEXT_NAMESPACE
diff --git a/Engine/Plugins/Runtime/Metasound/Source/MetasoundFrontend/Private/MetasoundAssetBase.cpp b/Engine/Plugins/Runtime/Metasound/Source/MetasoundFrontend/Private/MetasoundAssetBase.cpp
--- a/Engine/Plugins/Runtime/Metasound/Source/MetasoundFrontend/Private/MetasoundAssetBase.cpp
+++ b/Engine/Plugins/Runtime/Metasound/Source/MetasoundFrontend/Private/MetasoundAssetBase.cpp
@@ -7,7 +7,7 @@
 #include "Algo/ForEach.h"
 #include "Algo/Transform.h"
 #include "Containers/Set.h"
-#include "DocumentConfigurations/MetasoundPresetDocumentConfiguration.h"
+#include "DocumentConfigurations/MetasoundFrontendPresetDocumentConfiguration.h"
 #include "HAL/FileManager.h"
 #include "HAL/IConsoleManager.h"
 #include "IAudioParameterTransmitter.h"
@@ -635,15 +635,17 @@ FText FMetasoundAssetBase::GetDisplayName(FString&& InTypeName) const
 {
 	using namespace Metasound::Frontend;
 
-	FConstGraphHandle GraphHandle = GetRootGraphHandle();
-	const bool bIsPreset = !GraphHandle->GetGraphStyle().bIsGraphEditable;
-
-	if (!bIsPreset)
+	TScriptInterface<const IMetaSoundDocumentInterface> DocInterface(GetOwningAsset());
+	TConstStructView<FMetaSoundFrontendDocumentConfiguration> Config = DocInterface->GetConstDocument().Configuration;
+	if (Config.IsValid())
 	{
-		return FText::FromString(MoveTemp(InTypeName));
+		const UScriptStruct* Struct = Config.GetScriptStruct();
+		check(Struct);
+		return FText::Format(LOCTEXT("DocConfigDisplayNameFormat", "{0} ({1})"), FText::FromString(MoveTemp(InTypeName)), Struct->GetDisplayNameText());
 	}
 
-	return FText::Format(LOCTEXT("PresetDisplayNameFormat", "{0} (Preset)"), FText::FromString(MoveTemp(InTypeName)));
+	return FText::FromString(MoveTemp(InTypeName));
+
 }
 #endif // WITH_EDITOR
 
diff --git a/Engine/Plugins/Runtime/Metasound/Source/MetasoundFrontend/Private/MetasoundAssetManager.cpp b/Engine/Plugins/Runtime/Metasound/Source/MetasoundFrontend/Private/MetasoundAssetManager.cpp
--- a/Engine/Plugins/Runtime/Metasound/Source/MetasoundFrontend/Private/MetasoundAssetManager.cpp
+++ b/Engine/Plugins/Runtime/Metasound/Source/MetasoundFrontend/Private/MetasoundAssetManager.cpp
@@ -4,6 +4,7 @@
 
 #include "Algo/Copy.h"
 #include "AssetRegistry/AssetData.h"
+#include "DocumentConfigurations/MetasoundFrontendPresetDocumentConfiguration.h"
 #include "JsonObjectConverter.h"
 #include "MetasoundAssetManagerTransactionPrivate.h"
 #include "MetasoundAssetTagCollections.h"
@@ -60,7 +61,7 @@ FMetaSoundDocumentInfo::FMetaSoundDocumentInfo(const IMetaSoundDocumentInterface
 
 	if (!IsRunningCookCommandlet())
 	{
-		bIsPreset = (uint8)Document.RootGraph.PresetOptions.bIsPreset;
+		bIsPreset = (uint8)(Document.Configuration.GetPtr<FMetaSoundFrontendPresetDocumentConfiguration>() != nullptr);
 		DocumentVersion = Document.Metadata.Version.Number;
 
 		{
diff --git a/Engine/Plugins/Runtime/Metasound/Source/MetasoundFrontend/Private/MetasoundFrontendController.cpp b/Engine/Plugins/Runtime/Metasound/Source/MetasoundFrontend/Private/MetasoundFrontendController.cpp
--- a/Engine/Plugins/Runtime/Metasound/Source/MetasoundFrontend/Private/MetasoundFrontendController.cpp
+++ b/Engine/Plugins/Runtime/Metasound/Source/MetasoundFrontend/Private/MetasoundFrontendController.cpp
@@ -47,6 +47,31 @@ namespace Metasound
 			return Invalid;
 		}
 
+		const FMetasoundFrontendGraphClassPresetOptions& IGraphController::GetGraphPresetOptions() const
+		{
+			PRAGMA_DISABLE_DEPRECATION_WARNINGS
+				return Invalid::GetInvalidGraphClassPresetOptions();
+			PRAGMA_ENABLE_DEPRECATION_WARNINGS
+		}
+
+		void IGraphController::SetGraphPresetOptions(const FMetasoundFrontendGraphClassPresetOptions& InPresetOptions)
+		{
+		}
+
+		const TSet<FName>& IGraphController::GetInputsInheritingDefault() const
+		{
+			return Invalid::GetInvalidNameSet();
+		}
+
+		bool IGraphController::SetInputInheritsDefault(FName InName, bool bDefaultIsInherited)
+		{
+			return false;
+		}
+
+		void IGraphController::SetInputsInheritingDefault(TSet<FName>&& InNames)
+		{
+		}
+
 		FGraphHandle IGraphController::GetInvalidHandle()
 		{
 			static FGraphHandle Invalid = MakeShared<FInvalidGraphController>();
