MetaSound Document Configuration: Vertical Checkpoint 2

diff --git a/Engine/Plugins/Runtime/Metasound/Source/MetasoundFrontend/Private/MetasoundFrontendDocumentBuilder.cpp b/Engine/Plugins/Runtime/Metasound/Source/MetasoundFrontend/Private/MetasoundFrontendDocumentBuilder.cpp
--- a/Engine/Plugins/Runtime/Metasound/Source/MetasoundFrontend/Private/MetasoundFrontendDocumentBuilder.cpp
+++ b/Engine/Plugins/Runtime/Metasound/Source/MetasoundFrontend/Private/MetasoundFrontendDocumentBuilder.cpp
@@ -9,7 +9,8 @@
 #include "Algo/Sort.h"
 #include "Algo/Transform.h"
 #include "AudioParameter.h"
-#include "DocumentConfigurations/MetasoundPresetDocumentConfiguration.h"
+#include "DocumentConfigurations/MetasoundFrontendDocumentConfiguration.h"
+#include "DocumentConfigurations/MetasoundFrontendPresetDocumentConfiguration.h"
 #include "Interfaces/MetasoundFrontendInterfaceBindingRegistry.h"
 #include "Interfaces/MetasoundFrontendInterfaceRegistry.h"
 #include "MetasoundAssetBase.h"
@@ -18,7 +19,6 @@
 #include "MetasoundFrontendController.h"
 #include "MetasoundFrontendDataTypeRegistry.h"
 #include "MetasoundFrontendDocumentCache.h"
-#include "MetasoundFrontendDocumentConfiguration.h"
 #include "MetasoundFrontendDocument.h"
 #include "MetasoundFrontendDocumentIdGenerator.h"
 #include "MetasoundFrontendDocumentModifyDelegates.h"
@@ -33,6 +33,7 @@
 #include "MetasoundTrace.h"
 #include "MetasoundVariableNodes.h"
 #include "NodeTemplates/MetasoundFrontendNodeTemplateInput.h"
+#include "UObject/UObjectIterator.h"
 
 #include UE_INLINE_GENERATED_CPP_BY_NAME(MetasoundFrontendDocumentBuilder)
 
@@ -810,13 +811,19 @@ const FMetasoundFrontendDocument& UMetaSoundBuilderDocument::GetConstDocument()
 	return Document;
 }
 
-#if WITH_EDITORONLY_DATA
+#if WITH_EDITOR
+TStructView<FMetaSoundFrontendDocumentConfiguration> FMetaSoundFrontendDocumentBuilder::GetDocumentConfiguration() const
+{
+	TInstancedStruct<FMetaSoundFrontendDocumentConfiguration>& Config = GetDocumentChecked().Configuration;
+	return Config;
+}
+#endif // WITH_EDITOR
+
 TConstStructView<FMetaSoundFrontendDocumentConfiguration> FMetaSoundFrontendDocumentBuilder::GetConstDocumentConfiguration() const
 {
 	const TInstancedStruct<FMetaSoundFrontendDocumentConfiguration>& Config = GetConstDocumentChecked().Configuration;
 	return Config;
 }
-#endif // WITH_EDITORONLY_DATA
 
 EMetasoundFrontendClassAccessFlags UMetaSoundBuilderDocument::GetDefaultAccessFlags() const
 {
@@ -1281,7 +1288,7 @@ const FMetasoundFrontendNode* FMetaSoundFrontendDocumentBuilder::AddGraphInput(F
 
 		DocumentDelegates->InterfaceDelegates.OnInputAdded.Broadcast(NewIndex);
 #if WITH_EDITORONLY_DATA
-		if (FMetaSoundFrontendDocumentConfiguration* DocConfig = Document.Configuration.GetMutablePtr())
+		if (FDocumentVertexConfiguration* DocConfig = Document.Configuration.GetMutablePtr<FDocumentVertexConfiguration>())
 		{
 PRAGMA_DISABLE_EXPERIMENTAL_WARNINGS
 			DocConfig->AddVertexMetadata(*this, NewNode->GetID());
@@ -1401,7 +1408,7 @@ const FMetasoundFrontendNode* FMetaSoundFrontendDocumentBuilder::AddGraphOutput(
 		DocumentDelegates->InterfaceDelegates.OnOutputAdded.Broadcast(NewIndex);
 
 #if WITH_EDITORONLY_DATA
-		if (FMetaSoundFrontendDocumentConfiguration* DocConfig = Document.Configuration.GetMutablePtr())
+		if (FDocumentVertexConfiguration* DocConfig = Document.Configuration.GetMutablePtr<FDocumentVertexConfiguration>())
 		{
 PRAGMA_DISABLE_EXPERIMENTAL_WARNINGS
 			DocConfig->AddVertexMetadata(*this, NewOutput.NodeID);
@@ -2365,8 +2372,6 @@ void FMetaSoundFrontendDocumentBuilder::ClearDocument(TSharedRef<Metasound::Fron
 	GraphClass.GetDefaultInterface().SetOutputStyle({ });
 #endif // WITH_EDITOR
 
-	GraphClass.PresetOptions.InputsInheritingDefault.Empty();
-
 	// Removing graph pages is not necessary when editor only data is not available as graph mutation
 	// is only supported in builds with editor data loaded. Otherwise, anything calling ClearDocument
 	// should only be a transient, non serialized asset graph which does not support page mutation.
@@ -2515,7 +2520,6 @@ bool FMetaSoundFrontendDocumentBuilder::ConformObjectToDocument()
 	return GetDocumentInterfaceChecked().ConformObjectToDocument();
 }
 
-#if WITH_EDITORONLY_DATA
 bool FMetaSoundFrontendDocumentBuilder::ConfigureDocument(bool bLogFailure)
 {
 	using namespace Metasound::Frontend;
@@ -2534,20 +2538,34 @@ bool FMetaSoundFrontendDocumentBuilder::ConfigureDocument(bool bLogFailure)
 		if (DocConfig->ConfigureDocument(*this, bLogFailure))
 		{
 			GetDocumentInterfaceChecked().ConformObjectToDocument();
+
+#if WITH_EDITORONLY_DATA
+			FInputNodeTemplate::GetChecked().Inject(*this);
+#endif // WITH_EDITORONLY_DATA
+
 			bInDocConfigScope = false;
 			return true;
 		}
 
 		TSharedRef<FDocumentModifyDelegates> ModifyDelegates = MakeShared<FDocumentModifyDelegates>(*DocumentDelegates);
 		ClearDocument(ModifyDelegates);
-	}
 
-	SetIsGraphEditable(true);
+		// Add required interfaces to ensure document is left in valid state (ex. MetaSound sources
+		// need the source interface to properly initialize a null-rendering MetaSound generator.)
+		TArray<FMetasoundFrontendVersion> InitVersions = ISearchEngine::Get().FindUClassDefaultInterfaceVersions(GetBuilderClassPath());
+		FModifyInterfaceOptions Options({ }, MoveTemp(InitVersions));
+		ModifyInterfaces(MoveTemp(Options));
+	}
+#if WITH_EDITORONLY_DATA
+	else
+	{
+		SetIsGraphEditable(true);
+	}
+#endif // WITH_EDITORONLY_DATA
 
 	bInDocConfigScope = false;
 	return false;
 }
-#endif // WITH_EDITORONLY_DATA
 
 bool FMetaSoundFrontendDocumentBuilder::ContainsDependencyOfType(EMetasoundFrontendClassType ClassType) const
 {
@@ -2574,12 +2592,12 @@ bool FMetaSoundFrontendDocumentBuilder::ConvertFromPreset()
 
 	if (IsPreset())
 	{
-		GetDocumentChecked().RootGraph.PresetOptions = { };
+		// Don't process the null configuration in order to "freeze" the current document state and not clear it.
+		GetDocumentChecked().Configuration.Reset();
 
-#if WITH_EDITOR
-		FMetasoundFrontendGraphStyle& Style = FindBuildGraphChecked().Style;
-		Style.bIsGraphEditable = true;
-#endif // WITH_EDITOR
+#if WITH_EDITORONLY_DATA
+		SetIsGraphEditable(true);
+#endif // WITH_EDITORONLY_DATA?
 
 		DocumentDelegates->OnDocumentConfigurationChanged.Broadcast(FDocumentConfigChangedArgs { });
 		return true;
@@ -2598,28 +2616,13 @@ bool FMetaSoundFrontendDocumentBuilder::ConvertToPreset(const FMetaSoundFrontend
 	using namespace Metasound;
 	using namespace Metasound::Frontend;
 
-	TSharedRef<FDocumentModifyDelegates> ModifyDelegatesRef = ModifyDelegates.IsValid() ? ModifyDelegates->AsShared() : MakeShared<FDocumentModifyDelegates>(InReferencedDocumentBuilder.GetConstDocumentChecked());
-	ClearDocument(ModifyDelegatesRef);
+	using FPresetConfigStruct = TInstancedStruct<FMetaSoundFrontendPresetDocumentConfiguration>;
+	FPresetConfigStruct NewConfig = FPresetConfigStruct::Make<FMetaSoundFrontendPresetDocumentConfiguration>();
+	NewConfig.GetMutable().Parent = &InReferencedDocumentBuilder.CastDocumentObjectChecked<UObject>();
+	SetDocumentConfiguration(MoveTemp(NewConfig));
+	ConfigureDocument(true /* bLogFailure */);
 
-	FMetasoundFrontendGraphClass& PresetAssetRootGraph = GetDocumentChecked().RootGraph;
-	// Mark all inputs as inherited by default
-	{
-		PresetAssetRootGraph.PresetOptions.InputsInheritingDefault.Reset();
-		auto GetInputName = [](const FMetasoundFrontendClassInput& Input) { return Input.Name; };
-		Algo::Transform(PresetAssetRootGraph.GetDefaultInterface().Inputs, PresetAssetRootGraph.PresetOptions.InputsInheritingDefault, GetInputName);
-		PresetAssetRootGraph.PresetOptions.bIsPreset = true;
-	}
-
-	// Apply root graph transform
-	FRebuildPresetRootGraph RebuildPresetRootGraph(InReferencedDocumentBuilder);
-	if (RebuildPresetRootGraph.Transform(*this))
-	{
-		DocumentInterface->ConformObjectToDocument();
-		DocumentDelegates->OnDocumentConfigurationChanged.Broadcast(FDocumentConfigChangedArgs { });
-		return true;
-	}
-
-	return false;
+	return true;
 }
 
 TOptional<Metasound::FAnyDataReference> FMetaSoundFrontendDocumentBuilder::CreateDataReference(
@@ -2634,18 +2637,6 @@ TOptional<Metasound::FAnyDataReference> FMetaSoundFrontendDocumentBuilder::Creat
 	return IDataTypeRegistry::Get().CreateDataReference(DataType, AccessType, InLiteral, InOperatorSettings);
 };
 
-void FMetaSoundFrontendDocumentBuilder::SetPresetFlags(bool bIsPreset)
-{
-	FMetasoundFrontendGraphClass& RootGraph = GetDocumentChecked().RootGraph;
-	RootGraph.PresetOptions.bIsPreset = bIsPreset;
-	RootGraph.IterateGraphPages([&](FMetasoundFrontendGraph& PresetAssetGraph)
-	{
-#if WITH_EDITORONLY_DATA
-		PresetAssetGraph.Style.bIsGraphEditable = !bIsPreset;
-#endif // WITH_EDITORONLY_DATA
-	});
-}
-
 #if WITH_EDITORONLY_DATA
 bool FMetaSoundFrontendDocumentBuilder::DiffAgainstRegistryInterface(const FGuid& InNodeID, const FGuid& InPageID, bool bInUseHighestMinorVersion, Metasound::Frontend::FClassInterfaceUpdates& OutInterfaceUpdates, bool bInForceRegenerateClassInterfaceOverride) const
 {
@@ -3748,12 +3739,6 @@ TArray<const FMetasoundFrontendNode*> FMetaSoundFrontendDocumentBuilder::GetGrap
 
 const TSet<FName>* FMetaSoundFrontendDocumentBuilder::GetGraphInputsInheritingDefault() const
 {
-	FMetasoundFrontendGraphClassPresetOptions& PresetOptions = GetDocumentChecked().RootGraph.PresetOptions;
-	if (PresetOptions.bIsPreset)
-	{
-		return &PresetOptions.InputsInheritingDefault;
-	}
-
 	return nullptr;
 }
 
@@ -3776,39 +3761,15 @@ FMetasoundAssetBase& FMetaSoundFrontendDocumentBuilder::GetMetasoundAsset() cons
 FMetasoundAssetBase* FMetaSoundFrontendDocumentBuilder::GetReferencedPresetAsset() const
 {
 	using namespace Metasound::Frontend;
-	if (!IsPreset())
-	{
-		return nullptr;
-	}
 
-	// Find the single external node which is the referenced preset asset, 
-	// and find the asset with its registry key 
-	auto FindExternalNode = [this](const FMetasoundFrontendNode& Node)
-	{
-		const FMetasoundFrontendClass* Class = FindDependency(Node.ClassID);
-		check(Class);
-		return Class->Metadata.GetType() == EMetasoundFrontendClassType::External;
-	};
-	const FMetasoundFrontendNode* Node = FindConstBuildGraphChecked().Nodes.FindByPredicate(FindExternalNode);
-	if (Node != nullptr)
+	if (const FMetaSoundFrontendPresetDocumentConfiguration* PresetConfig = GetConstDocumentConfigurationAs<FMetaSoundFrontendPresetDocumentConfiguration>())
 	{
-		const FMetasoundFrontendClass* NodeClass = FindDependency(Node->ClassID);
-		check(NodeClass);
-		const FMetaSoundAssetKey NodeAssetKey(NodeClass->Metadata);
-		const TArray<FMetasoundAssetBase*> ReferencedAssets = GetMetasoundAsset().GetReferencedAssets();
-		for (FMetasoundAssetBase* RefAsset : ReferencedAssets)
+		if (UObject* Parent = PresetConfig->Parent.GetObject())
 		{
-			TScriptInterface<IMetaSoundDocumentInterface> RefDocInterface = RefAsset->GetOwningAsset();
-			if (RefDocInterface.GetObject() != nullptr)
-			{
-				const FMetaSoundAssetKey AssetKey(RefDocInterface->GetConstDocument().RootGraph.Metadata);
-				if (AssetKey == NodeAssetKey)
-				{
-					return RefAsset;
-				}
-			}
+			return IMetaSoundAssetManager::GetChecked().GetAsAsset(*Parent);
 		}
 	}
+
 	return nullptr;
 }
 
@@ -4225,7 +4186,7 @@ void FMetaSoundFrontendDocumentBuilder::InitDocument(const FMetasoundFrontendDoc
 		// Add default interfaces for given UClass
 		{
 			TArray<FMetasoundFrontendVersion> InitVersions = ISearchEngine::Get().FindUClassDefaultInterfaceVersions(GetBuilderClassPath());
-			FModifyInterfaceOptions Options({ }, InitVersions);
+			FModifyInterfaceOptions Options({ }, MoveTemp(InitVersions));
 			ModifyInterfaces(MoveTemp(Options));
 		}
 
@@ -4500,7 +4461,9 @@ bool FMetaSoundFrontendDocumentBuilder::IsInterfaceDeclared(const FMetasoundFron
 
 bool FMetaSoundFrontendDocumentBuilder::IsPreset() const
 {
-	return GetConstDocumentChecked().RootGraph.PresetOptions.bIsPreset;
+	using namespace Metasound::Frontend;
+
+	return IsDocumentConfiguredBy<FPresetDocumentConfiguration>();
 }
 
 Metasound::Frontend::EInvalidEdgeReason FMetaSoundFrontendDocumentBuilder::IsValidEdge(const FMetasoundFrontendEdge& InEdge, const FGuid* InPageID) const
@@ -4987,6 +4950,11 @@ void FMetaSoundFrontendDocumentBuilder::Reload(TSharedPtr<Metasound::Frontend::F
 
 #if WITH_EDITORONLY_DATA
 bool FMetaSoundFrontendDocumentBuilder::RemoveGraphInputDefault(FName InputName, const FGuid& InPageID, bool bClearInheritsDefault)
+{
+	return RemoveGraphInputDefault(InputName, InPageID);
+}
+
+bool FMetaSoundFrontendDocumentBuilder::RemoveGraphInputDefault(FName InputName, const FGuid& InPageID)
 {
 	using namespace Metasound;
 
@@ -5002,14 +4970,6 @@ bool FMetaSoundFrontendDocumentBuilder::RemoveGraphInputDefault(FName InputName,
 		if (bRemovedDefault)
 		{
 			DocumentDelegates->InterfaceDelegates.OnInputDefaultChanged.Broadcast(Index);
-
-			if (bClearInheritsDefault)
-			{
-				constexpr bool bInputInheritsDefault = false;
-				constexpr bool bForceUpdate = true;
-				SetGraphInputInheritsDefault(InputName, bInputInheritsDefault, bForceUpdate);
-			}
-
 			return true;
 		}
 	}
@@ -5262,9 +5222,9 @@ bool FMetaSoundFrontendDocumentBuilder::RemoveGraphInput(FName InputName, bool b
 		if (bNodesRemoved)
 		{
 #if WITH_EDITORONLY_DATA
-			if (FMetaSoundFrontendDocumentConfiguration* DocConfig = Document.Configuration.GetMutablePtr())
+			if (FDocumentVertexConfiguration* DocConfig = Document.Configuration.GetMutablePtr<FDocumentVertexConfiguration>())
 			{
-				DocConfig->RemoveMetadata(NodeID);
+				DocConfig->RemoveVertexMetadata(NodeID);
 			}
 #endif // WITH_EDITORONLY_DATA
 
@@ -5287,10 +5247,6 @@ bool FMetaSoundFrontendDocumentBuilder::RemoveGraphInput(FName InputName, bool b
 			Document.Metadata.ModifyContext.AddMemberIDModified(NodeID);
 #endif // WITH_EDITORONLY_DATA
 
-			constexpr bool bInputInheritsDefault = false;
-			constexpr bool bForceUpdate = true;
-			SetGraphInputInheritsDefault(InputName, bInputInheritsDefault, bForceUpdate);
-
 			const bool bDependencyReferenced = IsDependencyReferenced(ClassID);
 			if (bDependencyReferenced || RemoveDependency(ClassID))
 			{
@@ -5836,6 +5792,7 @@ FMetaSoundFrontendDocumentBuilder::FNodeInstanceReplacementData FMetaSoundFronte
 bool FMetaSoundFrontendDocumentBuilder::ResetGraphInputDefault(FName InputName)
 {
 	using namespace Metasound;
+	using namespace Metasound::Frontend;
 
 	auto NameMatchesInput = [&InputName](const FMetasoundFrontendClassInput& Input) { return Input.Name == InputName; };
 	FMetasoundFrontendDocument& Document = GetDocumentChecked();
@@ -5849,10 +5806,13 @@ bool FMetaSoundFrontendDocumentBuilder::ResetGraphInputDefault(FName InputName)
 
 		DocumentDelegates->InterfaceDelegates.OnInputDefaultChanged.Broadcast(Index);
 
-		// Set the input as inheriting default for presets
-		// (No-ops if MetaSound isn't preset or is already set to inherit default).
-		constexpr bool bInputInheritsDefault = true;
-		SetGraphInputInheritsDefault(InputName, bInputInheritsDefault);
+		if (FDocumentVertexConfiguration* Config = Document.Configuration.GetMutablePtr<FDocumentVertexConfiguration>())
+		{
+			if (FMetaSoundFrontendDocumentConfigurationVertexMetadata* VertexMetadata = Config->FindVertexMetadata(Input.NodeID))
+			{
+				VertexMetadata->OnDefaultReset(*this);
+			}
+		}
 
 		Document.Metadata.ModifyContext.AddMemberIDModified(Input.NodeID);
 		return true;
@@ -5927,7 +5887,6 @@ bool FMetaSoundFrontendDocumentBuilder::SetBuildPageID(const FGuid& InBuildPageI
 }
 #endif // WITH_EDITORONLY_DATA
 
-#if WITH_EDITORONLY_DATA
 void FMetaSoundFrontendDocumentBuilder::SetDocumentConfiguration(TInstancedStruct<FMetaSoundFrontendDocumentConfiguration> InConfig)
 {
 	using namespace Metasound::Frontend;
@@ -5936,6 +5895,7 @@ void FMetaSoundFrontendDocumentBuilder::SetDocumentConfiguration(TInstancedStruc
 	DocumentDelegates->OnDocumentConfigurationChanged.Broadcast({ });
 }
 
+#if WITH_EDITORONLY_DATA
 bool FMetaSoundFrontendDocumentBuilder::SetGraphInputAdvancedDisplay(const FName InputName, const bool InAdvancedDisplay)
 {
 	FMetasoundFrontendDocument& Document = GetDocumentChecked();
@@ -6138,6 +6098,7 @@ bool FMetaSoundFrontendDocumentBuilder::SetGraphInputDataType(FName InputName, F
 bool FMetaSoundFrontendDocumentBuilder::SetGraphInputDefault(FName InputName, FMetasoundFrontendLiteral InDefaultLiteral, const FGuid* InPageID)
 {
 	using namespace Metasound;
+	using namespace Metasound::Frontend;
 
 	auto NameMatchesInput = [&InputName](const FMetasoundFrontendClassInput& Input) { return Input.Name == InputName; };
 	FMetasoundFrontendDocument& Document = GetDocumentChecked();
@@ -6165,10 +6126,10 @@ bool FMetaSoundFrontendDocumentBuilder::SetGraphInputDefault(FName InputName, FM
 			}
 			DocumentDelegates->InterfaceDelegates.OnInputDefaultChanged.Broadcast(Index);
 
-			// Set the input as no longer inheriting default
-			constexpr bool bInputInheritsDefault = false;
-			constexpr bool bForceUpdate = true;
-			SetGraphInputInheritsDefault(InputName, bInputInheritsDefault, bForceUpdate);
+			if (const FDocumentConfigurationVertexMetadata* ConfigVertexMetadata = FindConstVertexMetadata<FDocumentConfigurationVertexMetadata>(Input.NodeID))
+			{
+				const_cast<FMetaSoundFrontendDocumentConfigurationVertexMetadata*>(ConfigVertexMetadata)->OnDefaultUpdated(*this);
+			}
 
 			return true;
 		}
@@ -6181,6 +6142,7 @@ bool FMetaSoundFrontendDocumentBuilder::SetGraphInputDefault(FName InputName, FM
 bool FMetaSoundFrontendDocumentBuilder::SetGraphInputDefaults(FName InputName, TArray<FMetasoundFrontendClassInputDefault> Defaults)
 {
 	using namespace Metasound;
+	using namespace Metasound::Frontend;
 
 	auto NameMatchesInput = [&InputName](const FMetasoundFrontendClassInput& Input) { return Input.Name == InputName; };
 	FMetasoundFrontendDocument& Document = GetDocumentChecked();
@@ -6200,10 +6162,14 @@ bool FMetaSoundFrontendDocumentBuilder::SetGraphInputDefaults(FName InputName, T
 			Input.SetDefaults(MoveTemp(Defaults));
 			DocumentDelegates->InterfaceDelegates.OnInputDefaultChanged.Broadcast(Index);
 
-			// Set the input as no longer inheriting default
-			constexpr bool bInputInheritsDefault = false;
-			constexpr bool bForceUpdate = true;
-			SetGraphInputInheritsDefault(InputName, bInputInheritsDefault, bForceUpdate);
+			if (FDocumentVertexConfiguration* DocConfig = Document.Configuration.GetMutablePtr<FDocumentVertexConfiguration>())
+			{
+				if (FDocumentConfigurationVertexMetadata* VertexMetadata = DocConfig->FindVertexMetadata(Input.NodeID))
+				{
+					VertexMetadata->OnDefaultUpdated(*this);
+				}
+			}
+
 			return true;
 		}
 		UE_LOG(LogMetaSound, Error, TEXT("Attempting to set graph input of type '%s' with unsupported literal type(s)"), *Input.TypeName.ToString());
@@ -6250,38 +6216,12 @@ bool FMetaSoundFrontendDocumentBuilder::SetGraphInputDisplayName(FName InputName
 
 bool FMetaSoundFrontendDocumentBuilder::SetGraphInputInheritsDefault(FName InName, bool bInputInheritsDefault, bool bForceUpdate)
 {
-	FMetasoundFrontendGraphClassPresetOptions& PresetOptions = GetDocumentChecked().RootGraph.PresetOptions;
-	if (!PresetOptions.bIsPreset && !bForceUpdate)
-	{
-		return false;
-	}
-
-	bool bValueChanged;
-	if (bInputInheritsDefault)
-	{
-		bValueChanged = PresetOptions.InputsInheritingDefault.Add(InName).IsValidId();
-	}
-	else
-	{
-		bValueChanged = PresetOptions.InputsInheritingDefault.Remove(InName) > 0;
-	}
-
-	if (bValueChanged)
-	{
-		if (const int32* Index = DocumentCache->GetInterfaceCache().FindInputIndex(InName))
-		{
-			DocumentDelegates->InterfaceDelegates.OnInputInheritsDefaultChanged.Broadcast(*Index);
-		}
-	}
-
-	return bValueChanged;
+	return false;
 }
 
 bool FMetaSoundFrontendDocumentBuilder::SetGraphInputsInheritingDefault(TSet<FName>&& InNames)
 {
-	FMetasoundFrontendGraphClassPresetOptions& PresetOptions = GetDocumentChecked().RootGraph.PresetOptions;
-	PresetOptions.InputsInheritingDefault = MoveTemp(InNames);
-	return true;
+	return false;
 }
 
 bool FMetaSoundFrontendDocumentBuilder::SetGraphInputName(FName InputName, FName NewName)
diff --git a/Engine/Plugins/Runtime/Metasound/Source/MetasoundFrontend/Private/MetasoundFrontendDocumentVersioning.cpp b/Engine/Plugins/Runtime/Metasound/Source/MetasoundFrontend/Private/MetasoundFrontendDocumentVersioning.cpp
--- a/Engine/Plugins/Runtime/Metasound/Source/MetasoundFrontend/Private/MetasoundFrontendDocumentVersioning.cpp
+++ b/Engine/Plugins/Runtime/Metasound/Source/MetasoundFrontend/Private/MetasoundFrontendDocumentVersioning.cpp
@@ -7,7 +7,7 @@
 #include "Algo/Transform.h"
 #include "Algo/Unique.h"
 #include "CoreGlobals.h"
-#include "DocumentConfigurations/MetasoundPresetDocumentConfiguration.h"
+#include "DocumentConfigurations/MetasoundFrontendPresetDocumentConfiguration.h"
 #include "HAL/IConsoleManager.h"
 #include "Interfaces/MetasoundFrontendInterface.h"
 #include "Interfaces/MetasoundFrontendInterfaceRegistry.h"
@@ -530,8 +530,8 @@ PRAGMA_ENABLE_DEPRECATION_WARNINGS
 			void TransformInternal(FDocumentHandle InDocument) const override
 			{
 				const FGuid NewAssetClassID = FGuid::NewGuid();
-				FMetasoundFrontendGraphClass Class = InDocument->GetRootGraphClass();
-				Class.Metadata.SetClassName(FMetasoundFrontendClassName({ }, FName(*NewAssetClassID.ToString()), { }));
+				const FMetasoundFrontendGraphClass& Class = InDocument->GetRootGraphClass();
+				const_cast<FMetasoundFrontendGraphClass&>(Class).Metadata.SetClassName(FMetasoundFrontendClassName({ }, FName(*NewAssetClassID.ToString()), { }));
 			}
 		};
 
@@ -959,6 +959,107 @@ PRAGMA_ENABLE_DEPRECATION_WARNINGS
 			}
 		};
 
+		/** Versions document from 1.16 to 1.17. */
+		class FVersionDocument_1_17 : public FVersionDocumentTransform
+		{
+		public:
+			virtual ~FVersionDocument_1_17() = default;
+
+			FMetasoundFrontendVersionNumber GetTargetVersion() const override
+			{
+				return { 1, 17 };
+			}
+
+		private:
+			static TScriptInterface<IMetaSoundDocumentInterface> FindReferencedMetaSound_Legacy(const FMetaSoundFrontendDocumentBuilder& Builder)
+			{
+				// Find the single external node which is the parent preset asset,
+				// and find the asset with its registry key
+				FMetaSoundAssetKey NodeAssetKey;
+				auto FindExternalNode = [&Builder, &NodeAssetKey](const FMetasoundFrontendNode& Node)
+				{
+					const FMetasoundFrontendClass* Class = Builder.FindDependency(Node.ClassID);
+					check(Class);
+					if (Class->Metadata.GetType() == EMetasoundFrontendClassType::External)
+					{
+						NodeAssetKey = FMetaSoundAssetKey(Class->Metadata);
+						return true;
+					}
+
+					return false;
+				};
+
+				const FMetasoundFrontendNode* Node = Builder.FindConstBuildGraphChecked().Nodes.FindByPredicate(FindExternalNode);
+				if (Node != nullptr)
+				{
+					const TArray<FMetasoundAssetBase*> ReferencedAssets = Builder.GetMetasoundAsset().GetReferencedAssets();
+
+					if (ReferencedAssets.Num() != 1)
+					{
+						METASOUND_VERSIONING_LOG(Warning, TEXT("More than one asset reference found when migrating preset to configuration for asset '%s'"), *Builder.GetDebugName());
+					}
+
+					for (FMetasoundAssetBase* RefAsset : ReferencedAssets)
+					{
+						TScriptInterface<IMetaSoundDocumentInterface> RefDocInterface = RefAsset->GetOwningAsset();
+						if (RefDocInterface.GetObject() != nullptr)
+						{
+							const FMetaSoundAssetKey AssetKey(RefDocInterface->GetConstDocument().RootGraph.Metadata);
+							if (AssetKey == NodeAssetKey)
+							{
+								return RefDocInterface;
+							}
+						}
+					}
+				}
+
+				METASOUND_VERSIONING_LOG(Error, TEXT("Failed to find MetaSound parent for preset with key '%s' on asset '%s'"), *NodeAssetKey.ToString(), *Builder.GetDebugName());
+				return { };
+			}
+
+		public:
+			void TransformInternal(FMetaSoundFrontendDocumentBuilder& OutBuilder) const override
+			{
+				const FMetasoundFrontendDocument& Document = OutBuilder.GetConstDocumentChecked();
+
+PRAGMA_DISABLE_DEPRECATION_WARNINGS
+				const FMetasoundFrontendGraphClassPresetOptions& PresetOptions = Document.RootGraph.PresetOptions;
+				bool& bIsPreset = const_cast<bool&>(PresetOptions.bIsPreset);
+PRAGMA_ENABLE_DEPRECATION_WARNINGS
+
+				if (bIsPreset)
+				{
+					bIsPreset = false;
+
+					using FPresetConfigStruct = TInstancedStruct<FMetaSoundFrontendPresetDocumentConfiguration>;
+					FPresetConfigStruct PresetConfig = FPresetConfigStruct::Make<FMetaSoundFrontendPresetDocumentConfiguration>();
+					PresetConfig.GetMutable().Parent = FindReferencedMetaSound_Legacy(OutBuilder);
+
+					for (const FMetasoundFrontendClassInput& Input : Document.RootGraph.GetDefaultInterface().Inputs)
+					{
+						// Some inputs don't support inheriting a default (ex. triggers), so ignore if no metadata is added.
+						if (FPresetVertexMetadata* VertexMetadata = PresetConfig.GetMutable().AddVertexMetadata<FPresetVertexMetadata>(OutBuilder, Input.NodeID))
+						{
+PRAGMA_DISABLE_DEPRECATION_WARNINGS
+							const bool bOverrideInheritedDefault = !PresetOptions.InputsInheritingDefault.Contains(Input.Name);
+							VertexMetadata->bOverrideInheritedDefault = bOverrideInheritedDefault;
+PRAGMA_ENABLE_DEPRECATION_WARNINGS
+						}
+					}
+
+					{
+						using FConfigStruct = TInstancedStruct<FMetaSoundFrontendDocumentConfiguration>;
+						FConfigStruct& ConfigRef = const_cast<FConfigStruct&>(Document.Configuration);
+						ConfigRef = MoveTemp(PresetConfig);
+					}
+				}
+
+				// There was a bug in the past where inputs inheriting defaults wasn't getting cleared when converting from preset,
+				// so just always clear the deprecated set here either way.
+				const_cast<TSet<FName>&>(PresetOptions.InputsInheritingDefault).Empty();
+			}
+		};
+
 		bool VersionBuilderDocument(FMetaSoundFrontendDocumentBuilder& Builder)
 		{
 			UObject& DocObject = Builder.CastDocumentObjectChecked<UObject>();
@@ -971,6 +1072,7 @@ PRAGMA_ENABLE_DEPRECATION_WARNINGS
 			bWasUpdated |= FVersionDocument_1_14().Transform(Builder);
 			bWasUpdated |= FVersionDocument_1_15().Transform(Builder);
 			bWasUpdated |= FVersionDocument_1_16().Transform(Builder);
+			bWasUpdated |= FVersionDocument_1_17().Transform(Builder);
 
 			return bWasUpdated;
 		}
@@ -1023,7 +1125,7 @@ PRAGMA_ENABLE_DEPRECATION_WARNINGS
 
 	FMetasoundFrontendVersionNumber FVersioningManager::GetMaxDocumentVersion()
 	{
-		return FMetasoundFrontendVersionNumber { 1, 16 };
+		return FMetasoundFrontendVersionNumber { 1, 17 };
 	}
 
 	FMetasoundFrontendVersionNumber FVersioningManager::GetPageMigrationVersion()
diff --git a/Engine/Plugins/Runtime/Metasound/Source/MetasoundFrontend/Private/MetasoundFrontendGraphController.cpp b/Engine/Plugins/Runtime/Metasound/Source/MetasoundFrontend/Private/MetasoundFrontendGraphController.cpp
--- a/Engine/Plugins/Runtime/Metasound/Source/MetasoundFrontend/Private/MetasoundFrontendGraphController.cpp
+++ b/Engine/Plugins/Runtime/Metasound/Source/MetasoundFrontend/Private/MetasoundFrontendGraphController.cpp
@@ -376,6 +376,7 @@ namespace Metasound
 
 		bool FGraphController::RemoveVariable(const FGuid& InVariableID)
 		{
+PRAGMA_DISABLE_DEPRECATION_WARNINGS
 			if (FMetasoundFrontendGraphClass* GraphClass = GraphClassPtr.Get())
 			{
 				if (FMetasoundFrontendVariable* Variable = FindFrontendVariable(InVariableID))
@@ -419,6 +420,7 @@ namespace Metasound
 			}
 
 			return false;
+PRAGMA_ENABLE_DEPRECATION_WARNINGS
 		}
 
 		TArray<FVariableHandle> FGraphController::GetVariables()
@@ -782,10 +784,6 @@ namespace Metasound
 				GraphClass->GetDefaultInterface().Inputs.Reset();
 				GraphClass->GetDefaultInterface().Outputs.Reset();
 
-PRAGMA_DISABLE_DEPRECATION_WARNINGS
-				GraphClass->PresetOptions.InputsInheritingDefault.Reset();
-PRAGMA_ENABLE_DEPRECATION_WARNINGS
-
 				OwningDocument->ClearInterfaceVersions();
 				OwningDocument->RemoveUnreferencedDependencies();
 
@@ -868,54 +866,6 @@ PRAGMA_ENABLE_DEPRECATION_WARNINGS
 			return GetNodesByPredicate(IsInputNode);
 		}
 
-		const TSet<FName>& FGraphController::GetInputsInheritingDefault() const
-		{
-			if (const FMetasoundFrontendGraphClass* GraphClass = GraphClassPtr.Get())
-			{
-				if (GraphClass->PresetOptions.bIsPreset)
-				{
-					return GraphClass->PresetOptions.InputsInheritingDefault;
-				}
-			}
-
-			return Invalid::GetInvalidNameSet();
-		}
-
-		bool FGraphController::SetInputInheritsDefault(FName InName, bool bInputInheritsDefault)
-		{
-			if (bInputInheritsDefault)
-			{
-				if (FMetasoundFrontendGraphClass* GraphClass = GraphClassPtr.Get())
-				{
-					if (GraphClass->PresetOptions.bIsPreset)
-					{
-						return GraphClass->PresetOptions.InputsInheritingDefault.Add(InName).IsValidId();
-					}
-				}
-			}
-			else
-			{
-				if (FMetasoundFrontendGraphClass* GraphClass = GraphClassPtr.Get())
-				{
-					if (GraphClass->PresetOptions.bIsPreset)
-					{
-						return GraphClass->PresetOptions.InputsInheritingDefault.Remove(InName) > 0;
-					}
-				}
-			}
-
-			return false;
-		}
-
-		void FGraphController::SetInputsInheritingDefault(TSet<FName>&& InNames)
-		{
-			if (FMetasoundFrontendGraphClass* GraphClass = GraphClassPtr.Get())
-			{
-				GraphClass->PresetOptions.bIsPreset = true;
-				GraphClass->PresetOptions.InputsInheritingDefault = MoveTemp(InNames);
-			}
-		}
-
 		bool FGraphController::ContainsOutputVertex(const FVertexName& InName, const FName& InTypeName) const
 		{
 			if (const FMetasoundFrontendGraphClass* GraphClass = GraphClassPtr.Get())
@@ -1538,19 +1488,6 @@ PRAGMA_ENABLE_DEPRECATION_WARNINGS
 		}
 #endif // WITH_EDITOR
 
-		// This can be used to clear the current literal for a given input.
-		// @returns false if the input name couldn't be found.
-		bool FGraphController::ClearLiteralForInput(const FVertexName& InInputName, FGuid InVertexID)
-		{
-			if (FMetasoundFrontendClassInput* Desc = FindInputDescriptionWithName(InInputName))
-			{
-				FMetasoundFrontendLiteral& DefaultLiteral = Desc->FindDefaultChecked(Frontend::DefaultPageID);
-				DefaultLiteral.Clear();
-			}
-
-			return false;
-		}
-
 		FNodeHandle FGraphController::AddNode(const FNodeRegistryKey& InKey, FGuid InNodeGuid)
 		{
 			FConstClassAccessPtr Class = OwningDocument->FindOrAddClass(InKey);
@@ -1695,23 +1632,6 @@ PRAGMA_ENABLE_DEPRECATION_WARNINGS
 			return false;
 		}
 
-		const FMetasoundFrontendGraphClassPresetOptions& FGraphController::GetGraphPresetOptions() const
-		{
-			if (const FMetasoundFrontendGraphClass* GraphClass = GraphClassPtr.Get())
-			{
-				return GraphClass->PresetOptions;
-			}
-			return Invalid::GetInvalidGraphClassPresetOptions();
-		}
-
-		void FGraphController::SetGraphPresetOptions(const FMetasoundFrontendGraphClassPresetOptions& InPresetOptions)
-		{
-			if (FMetasoundFrontendGraphClass* GraphClass = GraphClassPtr.Get())
-			{
-				GraphClass->PresetOptions = InPresetOptions;
-			}
-		}
-
 		const FMetasoundFrontendClassMetadata& FGraphController::GetGraphMetadata() const
 		{
 			if (const FMetasoundFrontendGraphClass* GraphClass = GraphClassPtr.Get())
diff --git a/Engine/Plugins/Runtime/Metasound/Source/MetasoundFrontend/Private/MetasoundFrontendGraphController.h b/Engine/Plugins/Runtime/Metasound/Source/MetasoundFrontend/Private/MetasoundFrontendGraphController.h
--- a/Engine/Plugins/Runtime/Metasound/Source/MetasoundFrontend/Private/MetasoundFrontendGraphController.h
+++ b/Engine/Plugins/Runtime/Metasound/Source/MetasoundFrontend/Private/MetasoundFrontendGraphController.h
@@ -82,16 +82,15 @@ namespace Metasound
 			TArray<FNodeHandle> GetInputNodes() override;
 			TArray<FConstNodeHandle> GetConstInputNodes() const override;
 
-			virtual const TSet<FName>& GetInputsInheritingDefault() const override;
-			virtual bool SetInputInheritsDefault(FName InName, bool bInputInheritsDefault) override;
-			virtual void SetInputsInheritingDefault(TSet<FName>&& InNames) override;
-
 			virtual FVariableHandle AddVariable(const FName& InDataType) override;
 			virtual FVariableHandle FindVariable(const FGuid& InVariableID) override;
 			virtual FConstVariableHandle FindVariable(const FGuid& InVariableID) const override;
 			virtual FVariableHandle FindVariableContainingNode(const FGuid& InNodeID) override;
 			virtual FConstVariableHandle FindVariableContainingNode(const FGuid& InNodeID) const override;
+
+			UE_DEPRECATED(5.8, "Use the MetaSound Builder API")
 			virtual bool RemoveVariable(const FGuid& InVariableID) override;
+
 			virtual TArray<FVariableHandle> GetVariables() override;
 			virtual TArray<FConstVariableHandle> GetVariables() const override;
 			virtual FNodeHandle FindOrAddVariableMutatorNode(const FGuid& InVariableID) override;
@@ -186,10 +185,6 @@ namespace Metasound
 			void SetSortOrderIndexForOutput(const FVertexName& InName, int32 InSortOrderIndex) override;
 #endif // WITH_EDITOR
 
-			// This can be used to clear the current literal for a given input.
-			// @returns false if the input name couldn't be found.
-			bool ClearLiteralForInput(const FVertexName& InInputName, FGuid InVertexID) override;
-
 			FNodeHandle AddNode(const FNodeRegistryKey& InNodeClass, FGuid InNodeGuid) override;
 			FNodeHandle AddNode(const FMetasoundFrontendClassMetadata& InClassMetadata, FGuid InNodeGuid) override;
 			FNodeHandle AddDuplicateNode(const INodeController& InNode) override;
@@ -197,10 +192,9 @@ namespace Metasound
 
 			// Remove the node corresponding to this node handle.
 			// On success, invalidates the received node handle.
-			bool RemoveNode(INodeController& InNode) override;
 
-			const FMetasoundFrontendGraphClassPresetOptions& GetGraphPresetOptions() const override;
-			void SetGraphPresetOptions(const FMetasoundFrontendGraphClassPresetOptions& InPresetOptions) override;
+			UE_DEPRECATED(5.8, "Use the MetaSound Builder API")
+			bool RemoveNode(INodeController& InNode) override;
 
 			const FMetasoundFrontendClassMetadata& GetGraphMetadata() const override;
 			void SetGraphMetadata(const FMetasoundFrontendClassMetadata& InClassMetadata) override;
diff --git a/Engine/Plugins/Runtime/Metasound/Source/MetasoundFrontend/Private/MetasoundFrontendInvalidController.h b/Engine/Plugins/Runtime/Metasound/Source/MetasoundFrontend/Private/MetasoundFrontendInvalidController.h
--- a/Engine/Plugins/Runtime/Metasound/Source/MetasoundFrontend/Private/MetasoundFrontendInvalidController.h
+++ b/Engine/Plugins/Runtime/Metasound/Source/MetasoundFrontend/Private/MetasoundFrontendInvalidController.h
@@ -38,7 +38,10 @@ namespace Metasound
 			const FMetasoundFrontendClassInterface& GetInvalidClassInterface();
 			const FMetasoundFrontendNodeInterface& GetInvalidNodeInterface();
 			const FMetasoundFrontendClassMetadata& GetInvalidClassMetadata();
+
+			UE_DEPRECATED(5.8, "PresetOptions have been integrated into Document Configuration")
 			const FMetasoundFrontendGraphClassPresetOptions& GetInvalidGraphClassPresetOptions();
+
 			const FMetasoundFrontendGraphClass& GetInvalidGraphClass();
 			const TArray<FMetasoundFrontendClass>& GetInvalidClassArray();
 			const TArray<FMetasoundFrontendGraphClass>& GetInvalidGraphClassArray();
@@ -330,16 +333,11 @@ namespace Metasound
 			virtual TArray<FConstNodeHandle> GetConstOutputNodes() const override { return TArray<FConstNodeHandle>(); }
 			virtual TArray<FConstNodeHandle> GetConstInputNodes() const override { return TArray<FConstNodeHandle>(); }
 
-			virtual const TSet<FName>& GetInputsInheritingDefault() const { return Invalid::GetInvalidNameSet(); }
-			virtual bool SetInputInheritsDefault(FName InName, bool bDefaultIsInherited) { return false; }
-			virtual void SetInputsInheritingDefault(TSet<FName>&& InNames) { }
-
 			virtual FVariableHandle AddVariable(const FName& InDataTypeName) override { return IVariableController::GetInvalidHandle(); }
 			virtual FVariableHandle FindVariable(const FGuid& InVariableID) override { return IVariableController::GetInvalidHandle(); }
 			virtual FConstVariableHandle FindVariable(const FGuid& InVariableID) const { return IVariableController::GetInvalidHandle(); }
 			virtual FVariableHandle FindVariableContainingNode(const FGuid& InNodeID) override { return IVariableController::GetInvalidHandle(); }
 			virtual FConstVariableHandle FindVariableContainingNode(const FGuid& InNodeID) const override { return IVariableController::GetInvalidHandle(); }
-			virtual bool RemoveVariable(const FGuid& InVariableID) override { return false; }
 			virtual TArray<FVariableHandle> GetVariables() override { return TArray<FVariableHandle>(); }
 			virtual TArray<FConstVariableHandle> GetVariables() const override { return TArray<FConstVariableHandle>(); }
 			virtual FNodeHandle FindOrAddVariableMutatorNode(const FGuid& InVariableID) override { return INodeController::GetInvalidHandle(); }
@@ -422,25 +420,15 @@ namespace Metasound
 			virtual void SetSortOrderIndexForOutput(const FVertexName& InName, int32 InSortOrderIndex) { }
 #endif // WITH_EDITOR
 
-			// This can be used to clear the current literal for a given input.
-			// @returns false if the input name couldn't be found.
-			virtual bool ClearLiteralForInput(const FVertexName& InInputName, FGuid InVertexID) override { return false; }
-
 			virtual FNodeHandle AddNode(const FNodeRegistryKey& InNodeClass, FGuid InNodeGuid) override { return INodeController::GetInvalidHandle(); }
 			virtual FNodeHandle AddNode(const FMetasoundFrontendClassMetadata& InNodeClass, FGuid InNodeGuid) override { return INodeController::GetInvalidHandle(); }
 			virtual FNodeHandle AddDuplicateNode(const INodeController& InNode) override { return INodeController::GetInvalidHandle(); }
 			virtual FNodeHandle AddTemplateNode(const FNodeRegistryKey& InNodeClass, FMetasoundFrontendNodeInterface&& InNodeInterface, FGuid InNodeGuid) override { return INodeController::GetInvalidHandle(); }
 			virtual FNodeHandle AddTemplateNode(const INodeTemplate& InNodeTemplate, FNodeTemplateGenerateInterfaceParams Params, FGuid InNodeGuid = FGuid::NewGuid()) override { return INodeController::GetInvalidHandle(); }
 
-			// Remove the node corresponding to this node handle.
-			// On success, invalidates the received node handle.
+			UE_DEPRECATED(5.8, "Use the MetaSound Builder API")
 			virtual bool RemoveNode(INodeController& InNode) override { return false; }
 
-			// Returns the metadata for the current graph, including the name, description and author.
-			virtual const FMetasoundFrontendGraphClassPresetOptions& GetGraphPresetOptions() const override { return Invalid::GetInvalidGraphClassPresetOptions(); }
-
-			virtual void SetGraphPresetOptions(const FMetasoundFrontendGraphClassPresetOptions& InMetadata) override { }
-
 			// Returns the metadata for the current graph, including the name, description and author.
 			virtual const FMetasoundFrontendClassMetadata& GetGraphMetadata() const override { return Invalid::GetInvalidClassMetadata(); }
 
diff --git a/Engine/Plugins/Runtime/Metasound/Source/MetasoundFrontend/Private/MetasoundFrontendLiteral.cpp b/Engine/Plugins/Runtime/Metasound/Source/MetasoundFrontend/Private/MetasoundFrontendLiteral.cpp
--- a/Engine/Plugins/Runtime/Metasound/Source/MetasoundFrontend/Private/MetasoundFrontendLiteral.cpp
+++ b/Engine/Plugins/Runtime/Metasound/Source/MetasoundFrontend/Private/MetasoundFrontendLiteral.cpp
@@ -917,22 +917,19 @@ void FMetasoundFrontendLiteral::Empty()
 	AsUObject.Empty();
 }
 
-namespace Metasound
+namespace Metasound::Frontend
 {
-	namespace Frontend
+	// Convenience function to convert Metasound::ELiteralType to EMetasoundFrontendLiteralType, the UEnum used for metasound documents.
+	EMetasoundFrontendLiteralType GetMetasoundFrontendLiteralType(ELiteralType InLiteralType)
 	{
-		// Convenience function to convert Metasound::ELiteralType to EMetasoundFrontendLiteralType, the UEnum used for metasound documents.
-		EMetasoundFrontendLiteralType GetMetasoundFrontendLiteralType(ELiteralType InLiteralType)
-		{
-			static_assert(static_cast<uint8>(EMetasoundFrontendLiteralType::Invalid) == static_cast<uint8>(ELiteralType::Invalid), "Possible literal type enum mismatch");
-			return static_cast<EMetasoundFrontendLiteralType>(InLiteralType);
-		}
+		static_assert(static_cast<uint8>(EMetasoundFrontendLiteralType::Invalid) == static_cast<uint8>(ELiteralType::Invalid), "Possible literal type enum mismatch");
+		return static_cast<EMetasoundFrontendLiteralType>(InLiteralType);
+	}
 
-		// Convenience function to convert Metasound::ELiteralType to EMetasoundFrontendLiteralType, the UEnum used for metasound documents.
-		ELiteralType GetMetasoundLiteralType(EMetasoundFrontendLiteralType InLiteralType)
-		{
-			static_assert(static_cast<uint8>(EMetasoundFrontendLiteralType::Invalid) == static_cast<uint8>(ELiteralType::Invalid), "Possible literal type enum mismatch");
-			return static_cast<ELiteralType>(InLiteralType);
-		}
+	// Convenience function to convert Metasound::ELiteralType to EMetasoundFrontendLiteralType, the UEnum used for metasound documents.
+	ELiteralType GetMetasoundLiteralType(EMetasoundFrontendLiteralType InLiteralType)
+	{
+		static_assert(static_cast<uint8>(EMetasoundFrontendLiteralType::Invalid) == static_cast<uint8>(ELiteralType::Invalid), "Possible literal type enum mismatch");
+		return static_cast<ELiteralType>(InLiteralType);
 	}
-}
+} // namespace Metasound::Frontend
diff --git a/Engine/Plugins/Runtime/Metasound/Source/MetasoundFrontend/Private/MetasoundFrontendPages.cpp b/Engine/Plugins/Runtime/Metasound/Source/MetasoundFrontend/Private/MetasoundFrontendPages.cpp
--- a/Engine/Plugins/Runtime/Metasound/Source/MetasoundFrontend/Private/MetasoundFrontendPages.cpp
+++ b/Engine/Plugins/Runtime/Metasound/Source/MetasoundFrontend/Private/MetasoundFrontendPages.cpp
@@ -113,11 +113,10 @@ namespace Metasound::Frontend
 
 			const int32 NumInitDefaults = GraphInput.GetDefaults().Num();
 
-			// Remove unused graphs
+			// Remove unused inputs
 			for (const FGuid& PageIDToRemove : PageIDsToRemove)
 			{
-				constexpr bool bClearInheritsDefault = false;
-				const bool bDidRemove = InBuilder.RemoveGraphInputDefault(GraphInput.Name, PageIDToRemove, bClearInheritsDefault);
+				const bool bDidRemove = InBuilder.RemoveGraphInputDefault(GraphInput.Name, PageIDToRemove);
 				if (bDidRemove)
 				{
 					UE_LOG(LogMetaSound, Display, TEXT("%s: Removed Graph Input %s w/PageID '%s'"),
diff --git a/Engine/Plugins/Runtime/Metasound/Source/MetasoundFrontend/Private/MetasoundFrontendTransform.cpp b/Engine/Plugins/Runtime/Metasound/Source/MetasoundFrontend/Private/MetasoundFrontendTransform.cpp
--- a/Engine/Plugins/Runtime/Metasound/Source/MetasoundFrontend/Private/MetasoundFrontendTransform.cpp
+++ b/Engine/Plugins/Runtime/Metasound/Source/MetasoundFrontend/Private/MetasoundFrontendTransform.cpp
@@ -2,6 +2,7 @@
 #include "MetasoundFrontendTransform.h"
 
 #include "Algo/Transform.h"
+#include "DocumentConfigurations/MetasoundFrontendPresetDocumentConfiguration.h"
 #include "Interfaces/MetasoundFrontendInterface.h"
 #include "Interfaces/MetasoundFrontendInterfaceRegistry.h"
 #include "NodeTemplates/MetasoundFrontendNodeTemplateInput.h"
@@ -520,47 +521,27 @@ namespace Metasound
 		{
 			METASOUND_TRACE_CPUPROFILER_EVENT_SCOPE(FAutoUpdateRootGraph::Transform);
 			bool bDidEdit = false;
-
-			const bool bIsPreset = InOutBuilderToTransform.IsPreset();
-			const FMetasoundFrontendGraphClass& RootGraph = InOutBuilderToTransform.GetConstDocumentChecked().RootGraph;
 			
-			// If preset, rebuild root graph 
-			if (bIsPreset)
+			// If document config provided, run configure
+			if (TConstStructView<FMetaSoundFrontendDocumentConfiguration> DocConfig = InOutBuilderToTransform.GetConstDocumentConfiguration(); DocConfig.IsValid())
 			{
-				FMetasoundAssetBase* PresetReferencedMetaSoundAsset = InOutBuilderToTransform.GetReferencedPresetAsset();
-				if (!PresetReferencedMetaSoundAsset)
-				{
-					METASOUND_VERSIONING_LOG(Error, TEXT("Auto-Updating preset '%s' failed: Referenced class missing."), *DebugAssetPath);
-					return false;
-				}
-				TScriptInterface<IMetaSoundDocumentInterface> RefDocInterface = PresetReferencedMetaSoundAsset->GetOwningAsset();
-				const FMetaSoundFrontendDocumentBuilder& ReferenceBuilder = IDocumentBuilderRegistry::GetChecked().FindOrBeginBuilding(RefDocInterface);
-
-				bDidEdit |= FRebuildPresetRootGraph(ReferenceBuilder).Transform(InOutBuilderToTransform);
-				if (bDidEdit)
-				{
-					TScriptInterface<IMetaSoundDocumentInterface> PresetInterface(&InOutBuilderToTransform.CastDocumentObjectChecked<UObject>());
-					PresetInterface->ConformObjectToDocument();
-
-					InOutBuilderToTransform.RemoveUnusedDependencies();
-					InOutBuilderToTransform.SynchronizeDependencyMetadata();
-				}
+				InOutBuilderToTransform.ConfigureDocument(true /* bLogFailure */);
+				InOutBuilderToTransform.SynchronizeDependencyMetadata();
 			}
-			// Non preset
-			else
+
+			// Update external node dependencies until there are no more updates
+			// (A node could need a minor version update, then an update transform, 
+			// and then updated to another minor version which would require multiple passes).
+			// This should be run post document config as well as a config's implementing code may
+			// or may not be updated to generate newest dependencies.
+			bool bUpdated = true;
+			while (bUpdated)
 			{
-				// Update external node dependencies until there are no more updates
-				// (A node could need a minor version update, then an update transform, 
-				// and then updated to another minor version which would require multiple passes)
-				bool bUpdated = true;
-				while (bUpdated)
-				{
-					bUpdated = false;
-					bUpdated |= UpdateExternalDependencies(InOutBuilderToTransform);
-					bDidEdit |= bUpdated;
-				}
+				bUpdated = false;
+				bUpdated |= UpdateExternalDependencies(InOutBuilderToTransform);
+				bDidEdit |= bUpdated;
 			}
-				
+
 			return bDidEdit;
 		}
 
@@ -600,7 +581,7 @@ namespace Metasound
 
 						METASOUND_VERSIONING_LOG(Display, TEXT("Auto-Updating '%s' node class '%s': Newer version '%s' found."), *DebugAssetPath, *NodeClassKey.ClassName.ToString(), *UpdateVersion.ToString());
 						DependenciesToUpdate.Add(FDependencyPair(NodeClassKey, ClassInfoWithHighestMinorVersion.ToRegistryKey()));
-						bDidEdit |= true;						
+						bDidEdit |= true;
 					}
 					else if (AutoUpdateReason == EAutoUpdateEligibility::Eligible_InterfaceChange)
 					{
@@ -686,30 +667,22 @@ namespace Metasound
 		bool FRebuildPresetRootGraph::Transform(FMetaSoundFrontendDocumentBuilder& InOutBuilderToTransform) const
 		{
 			METASOUND_TRACE_CPUPROFILER_EVENT_SCOPE(Metasound::Frontend::FRebuildPresetRootGraph::Transform);
-			
-			// Callers of this transform should check that the graph is supposed to
-			// be managed externally and the parent builder is valid before calling this transform. If a scenario
-			// arises where this transform is used outside of AutoUpdate, then this
-			// early exist should be removed as it's mostly here to protect against
-			// accidental manipulation of metasound graphs.
-			if (!InOutBuilderToTransform.IsPreset() || !ParentBuilder)
+
+			// This function is to be moved to the preset doc configuration and should no longer be directly called outside that context.
+			const FMetaSoundFrontendPresetDocumentConfiguration* PresetConfig = InOutBuilderToTransform.GetConstDocumentConfigurationAs<FMetaSoundFrontendPresetDocumentConfiguration>();
+			if (!PresetConfig || !ParentBuilder)
 			{
 				return false;
 			}
 	
-			const FMetasoundFrontendDocument& DocumentToTransform = InOutBuilderToTransform.GetConstDocumentChecked();
-			const FMetasoundFrontendDocument& ParentDocument = ParentBuilder->GetConstDocumentChecked();
-			
 			// Determine the inputs and outputs needed in the wrapping graph. Also
 			// cache any exiting literals that have been set on the wrapping graph.
-			TSet<FName> InputsInheritingDefault;
-			TArray<FMetasoundFrontendClassInput> ClassInputs = GenerateRequiredClassInputs(InOutBuilderToTransform, InputsInheritingDefault);
+			TArray<FMetasoundFrontendClassInput> ClassInputs = GenerateRequiredClassInputs(InOutBuilderToTransform);
 			TArray<FMetasoundFrontendClassOutput> ClassOutputs = GenerateRequiredClassOutputs(InOutBuilderToTransform);
 
 #if WITH_EDITORONLY_DATA
 			// Cache off member metadata so it be can be readded if necessary after the graph is cleared 
-			FMemberIDToMetadataMap CachedMemberMetadata;
-			CacheMemberMetadata(InOutBuilderToTransform, InputsInheritingDefault, CachedMemberMetadata);
+			FMemberIDToMetadataMap CachedMemberMetadata = CacheMemberMetadata(*PresetConfig, InOutBuilderToTransform);
 #endif // WITH_EDITORONLY_DATA
 
 			FGuid PresetNodeID;
@@ -722,6 +695,7 @@ namespace Metasound
 			{
 				// This ID was originally being set to FGuid::NewGuid. 
 				// If you were reliant on that ID, please resave the asset so it is serialized with a valid ID
+				const FMetasoundFrontendDocument& DocumentToTransform = InOutBuilderToTransform.GetConstDocumentChecked();
 				PresetNodeID = DocumentToTransform.RootGraph.ID;
 			}
 
@@ -729,11 +703,14 @@ namespace Metasound
 			TSharedRef<FDocumentModifyDelegates> ModifyDelegates = MakeShared<FDocumentModifyDelegates>(InOutBuilderToTransform.GetDocumentDelegates());
 			InOutBuilderToTransform.ClearDocument(ModifyDelegates);
 
-			InOutBuilderToTransform.SetPresetFlags();
+#if WITH_EDITORONLY_DATA
+			InOutBuilderToTransform.SetIsGraphEditable(false);
+#endif // WITH_EDITORONLY_DATA
 
 			// Ensure preset interfaces match those found in referenced graph.  Referenced graph is assumed to be
 			// well-formed (i.e. all inputs/outputs/environment variables declared by interfaces are present, and
 			// of proper name & data type). Pass in parent builder to copy member ids so added member ids can be consistent
+			const FMetasoundFrontendDocument& ParentDocument = ParentBuilder->GetConstDocumentChecked();
 			const TSet<FMetasoundFrontendVersion>& RefInterfaceVersions = ParentDocument.Interfaces;
 			for (const FMetasoundFrontendVersion& Version : RefInterfaceVersions)
 			{
@@ -741,8 +718,7 @@ namespace Metasound
 			}
 
 			// Add referenced node
-			const FMetasoundFrontendClassMetadata& ParentClassMetadata = ParentDocument.RootGraph.Metadata;
-			const FMetasoundFrontendNode* ReferencedNode = InOutBuilderToTransform.AddNodeByClassName(ParentClassMetadata.GetClassName(), ParentClassMetadata.GetVersion().Major, PresetNodeID);
+			const FMetasoundFrontendNode* ReferencedNode = InOutBuilderToTransform.AddGraphNode(ParentDocument.RootGraph, PresetNodeID);
 			check(ReferencedNode);
 			
 #if WITH_EDITOR
@@ -752,8 +728,7 @@ namespace Metasound
 #endif // WITH_EDITOR
 
 			// Connect parent graph to referenced graph
-			InOutBuilderToTransform.SetGraphInputsInheritingDefault(MoveTemp(InputsInheritingDefault));
-			AddAndConnectInputs(ClassInputs, InOutBuilderToTransform, PresetNodeID);
+			AddAndConnectInputs(*PresetConfig, ClassInputs, InOutBuilderToTransform, PresetNodeID);
 			AddAndConnectOutputs(ClassOutputs, InOutBuilderToTransform, PresetNodeID);
 
 #if WITH_EDITORONLY_DATA
@@ -779,9 +754,9 @@ namespace Metasound
 			}
 		}
 		
-		void FRebuildPresetRootGraph::CacheMemberMetadata(FMetaSoundFrontendDocumentBuilder& InOutBuilderToTransform, const TSet<FName>& InInputsInheritingDefault, FMemberIDToMetadataMap& InOutCachedMemberMetadata) const
+		FRebuildPresetRootGraph::FMemberIDToMetadataMap FRebuildPresetRootGraph::CacheMemberMetadata(const FMetaSoundFrontendPresetDocumentConfiguration& PresetConfig, FMetaSoundFrontendDocumentBuilder& InOutBuilderToTransform) const
 		{
-			const FMetasoundFrontendDocument& ParentDocument = ParentBuilder->GetConstDocumentChecked();
+			FMemberIDToMetadataMap CachedMemberMetadata;
 
 			auto CreateNewLiteral = [&](UMetaSoundFrontendMemberMetadata* TemplateObject, FGuid MemberID) -> UMetaSoundFrontendMemberMetadata*
 			{
@@ -795,6 +770,7 @@ namespace Metasound
 				return nullptr;
 			};
 			
+			const FMetasoundFrontendDocument& ParentDocument = ParentBuilder->GetConstDocumentChecked();
 			for (const FMetasoundFrontendClassInput& ParentClassInput : ParentDocument.RootGraph.GetDefaultInterface().Inputs)
 			{
 				UMetaSoundFrontendMemberMetadata* MemberMetadata = nullptr;
@@ -808,14 +784,15 @@ namespace Metasound
 					// If the input vertex already exists in the parent graph,
 					// check if parent should be used or not from set of managed
 					// input names.
-					if (InInputsInheritingDefault.Contains(ParentClassInput.Name))
+					if (const FPresetVertexMetadata* VertexMetadata = PresetConfig.FindConstVertexMetadata<FPresetVertexMetadata>(MemberID);
+						VertexMetadata && !VertexMetadata->bOverrideInheritedDefault)
 					{
 						UMetaSoundFrontendMemberMetadata* ParentMetadata = ParentBuilder->FindMemberMetadata(ParentClassInput.NodeID);
 						MemberMetadata = CreateNewLiteral(ParentMetadata, MemberID);
 					}
 					else
 					{
-						// Use existing defaults 
+						// Use existing defaults
 						MemberMetadata = InOutBuilderToTransform.FindMemberMetadata(GraphInput->NodeID);
 					}
 				}
@@ -828,23 +805,22 @@ namespace Metasound
 
 				if (MemberMetadata)
 				{
-					InOutCachedMemberMetadata.Emplace(MemberID, MemberMetadata);
+					CachedMemberMetadata.Emplace(MemberID, MemberMetadata);
 				}
 			}
+
+			return CachedMemberMetadata;
 		}
 #endif // WITH_EDITORONLY_DATA
 
-		void FRebuildPresetRootGraph::AddAndConnectInputs(const TArray<FMetasoundFrontendClassInput>& InClassInputs, FMetaSoundFrontendDocumentBuilder& InOutBuilderToTransform, const FGuid& InReferencedNodeID) const
+		void FRebuildPresetRootGraph::AddAndConnectInputs(const FMetaSoundFrontendPresetDocumentConfiguration& PresetConfig, const TArray<FMetasoundFrontendClassInput>& InClassInputs, FMetaSoundFrontendDocumentBuilder& InOutBuilderToTransform, const FGuid& InReferencedNodeID) const
 		{
-			// Add inputs and space appropriately
-			const INodeTemplate* InputTemplate = INodeTemplateRegistry::Get().FindTemplate(FInputNodeTemplate::ClassName);
-			check(InputTemplate);
-
 			TArray<const FMetasoundFrontendVertex*> ReferencedNodeInputVertices;
+
+#if WITH_EDITORONLY_DATA
 			TArray<const FMetasoundFrontendNode*> InputTemplateNodes;
-			const TSet<FName>* InputsInheritingDefault = InOutBuilderToTransform.GetGraphInputsInheritingDefault();
-			check(InputsInheritingDefault);
-			
+#endif // WITH_EDITORONLY_DATA
+
 			for (const FMetasoundFrontendClassInput& ClassInput : InClassInputs)
 			{
 				const FName InputName = ClassInput.Name;
@@ -856,41 +832,27 @@ namespace Metasound
 				}
 				else
 				{
-					// Defaults must be set as either the parent default or existing default may be different than the registered interface default
-					// Setting defaults will set the input as non inheriting even if we're just copying the value 
-					// so cache off bool and set back if necessary
-					const bool bInheritsDefault = InputsInheritingDefault->Contains(InputName);
 					InOutBuilderToTransform.SetGraphInputDefaults(InputName, ClassInput.GetDefaults());
-					if (bInheritsDefault)
-					{
-						InOutBuilderToTransform.SetGraphInputInheritsDefault(InputName, /*bInputInheritsDefault=*/true);
-					}
 				}
 
-				const FGuid InputNodeID = InputNode->GetID();
 				check(InputNode);
-				const FMetasoundFrontendVertex* InputNodeOutputVertex = InOutBuilderToTransform.FindNodeOutput(InputNode->GetID(), InputName);
+				const FGuid InputNodeID = InputNode->GetID();
+				const FMetasoundFrontendVertex* InputNodeOutputVertex = InOutBuilderToTransform.FindNodeOutput(InputNodeID, InputName);
 				check(InputNodeOutputVertex);
 				const FMetasoundFrontendVertex* ReferencedNodeInputVertex = InOutBuilderToTransform.FindNodeInput(InReferencedNodeID, InputName);
 				check(ReferencedNodeInputVertex);
 				
 				ReferencedNodeInputVertices.Add(ReferencedNodeInputVertex);
 
-				// If not editor-only data, just make connection directly between input and output node
-#if !WITH_EDITORONLY_DATA
-				InOutBuilderToTransform.AddEdge(FMetasoundFrontendEdge
-				{
-					InputNodeID,
-					InputNodeOutputVertex->VertexID,
-					InReferencedNodeID,
-					ReferencedNodeInputVertex->VertexID
-				});
-
 				// If editor-only data, add template node and connections
-#else // WITH_EDITORONLY_DATA
+#if WITH_EDITORONLY_DATA
 				// template node takes on data type of concrete input node's output type
 				const FName DataType = InputNode->Interface.Outputs.Last().TypeName;
 				
+				// Add inputs and space appropriately
+				const INodeTemplate* InputTemplate = INodeTemplateRegistry::Get().FindTemplate(FInputNodeTemplate::ClassName);
+				check(InputTemplate);
+
 				FNodeTemplateGenerateInterfaceParams Params{ { }, { DataType } };
 				const FMetasoundFrontendNode* TemplateNode = InOutBuilderToTransform.AddNodeByTemplate(*InputTemplate, MoveTemp(Params));
 
@@ -902,7 +864,7 @@ namespace Metasound
 				InOutBuilderToTransform.AddEdge(FMetasoundFrontendEdge
 				{
 					InputNodeID,
-					InputNodeOutputVertex->VertexID, 
+					InputNodeOutputVertex->VertexID,
 					TemplateNode->GetID(),
 					TemplateNodeInputVertexID
 				});
@@ -913,7 +875,17 @@ namespace Metasound
 					InReferencedNodeID,
 					ReferencedNodeInputVertex->VertexID
 				});
-#endif  // WITH_EDITORONLY_DATA
+#else // !WITH_EDITORONLY_DATA
+				// If not editor-only data, just make connection directly between input and output node
+				InOutBuilderToTransform.AddEdge(FMetasoundFrontendEdge
+					{
+						InputNodeID,
+						InputNodeOutputVertex->VertexID,
+						InReferencedNodeID,
+						ReferencedNodeInputVertex->VertexID
+					});
+
+#endif  // !WITH_EDITORONLY_DATA
 			}
 
 #if WITH_EDITOR
@@ -925,10 +897,11 @@ namespace Metasound
 			// Sort vertices on referenced node
 			// then use that order to order connected input nodes (which share the same names)
 			auto GetInputDisplayName = [&InOutBuilderToTransform, &InReferencedNodeID](const FMetasoundFrontendVertex& Vertex)
-			{
-				return InOutBuilderToTransform.GetNodeInputDisplayName(InReferencedNodeID, Vertex.Name);
-			};
+				{
+					return InOutBuilderToTransform.GetNodeInputDisplayName(InReferencedNodeID, Vertex.Name);
+				};
 			Style.SortVertices(ReferencedNodeInputVertices, GetInputDisplayName);
+
 			InOutBuilderToTransform.SetInputStyle(MoveTemp(Style));
 
 			// Set editor node locations
@@ -1003,7 +976,7 @@ namespace Metasound
 #endif // WITH_EDITOR
 		}
 
-		TArray<FMetasoundFrontendClassInput> FRebuildPresetRootGraph::GenerateRequiredClassInputs(FMetaSoundFrontendDocumentBuilder& InDocumentToTransformBuilder, TSet<FName>& OutInputsInheritingDefault) const
+		TArray<FMetasoundFrontendClassInput> FRebuildPresetRootGraph::GenerateRequiredClassInputs(FMetaSoundFrontendDocumentBuilder& InDocumentToTransformBuilder) const
 		{
 			TArray<FMetasoundFrontendClassInput> ClassInputs;
 			check(ParentBuilder);
@@ -1012,9 +985,8 @@ namespace Metasound
 			const FMetasoundFrontendDocument& ParentDocument = ParentBuilder->GetConstDocumentChecked();
 			const FMetasoundFrontendDocument& DocumentToTransform = InDocumentToTransformBuilder.GetConstDocumentChecked();
 
-			const TSet<FName>* ExistingInputsInheritingDefaultPtr = InDocumentToTransformBuilder.GetGraphInputsInheritingDefault();
-			check(ExistingInputsInheritingDefaultPtr);
-			TSet<FName> InputsInheritingDefault = *ExistingInputsInheritingDefaultPtr;
+			const FMetaSoundFrontendPresetDocumentConfiguration* PresetConfig = InDocumentToTransformBuilder.GetConstDocumentConfigurationAs<FMetaSoundFrontendPresetDocumentConfiguration>();
+			check(PresetConfig);
 
 			// Iterate through all input nodes of referenced graph
 			for (const FMetasoundFrontendClassInput& ParentClassInput : ParentDocument.RootGraph.GetDefaultInterface().Inputs)
@@ -1048,40 +1020,33 @@ namespace Metasound
 				const FMetasoundFrontendClassInput* GraphInput = InDocumentToTransformBuilder.FindGraphInput(NodeName);
 				if (GraphInput && GraphInput->TypeName == ParentClassInput.TypeName)
 				{
+					bool bInherited = false;
 					// If the input vertex already exists in the parent graph,
 					// check if parent should be used or not from set of managed
 					// input names.
-					if (InputsInheritingDefault.Contains(NodeName))
+					if (const FPresetVertexMetadata* VertexMetadata = PresetConfig->FindConstVertexMetadata<FPresetVertexMetadata>(GraphInput->NodeID))
 					{
-						InheritDefaultsFromGraph(*ParentBuilder);
+						if (!VertexMetadata->bOverrideInheritedDefault)
+						{
+							InheritDefaultsFromGraph(*ParentBuilder);
+							bInherited = true;
+						}
 					}
-					else
+
+					// Use existing defaults
+					if (!bInherited)
 					{
-						// Use existing defaults 
 						InheritDefaultsFromGraph(InDocumentToTransformBuilder);
 					}
 				}
 				else
 				{
 					InheritDefaultsFromGraph(*ParentBuilder);
-					// Add this to the list of inheriting as the
-					// type no longer matches or it no longer exists
-					InputsInheritingDefault.Add(NodeName);
 				}
+
 				ClassInputs.Add(MoveTemp(NewClassInput));
 			}
 
-			OutInputsInheritingDefault.Reset();
-			Algo::TransformIf(ClassInputs, OutInputsInheritingDefault,
-				[&InputsInheritingDefault](const FMetasoundFrontendClassInput& Input)
-				{
-					return InputsInheritingDefault.Contains(Input.Name);
-				},
-				[](const FMetasoundFrontendClassInput& Input)
-				{
-					return Input.Name;
-				});
-			
 			return ClassInputs;
 		}
 		
diff --git a/Engine/Plugins/Runtime/Metasound/Source/MetasoundFrontend/Public/DocumentConfigurations/MetasoundFrontendDocumentVertexConfiguration.h b/Engine/Plugins/Runtime/Metasound/Source/MetasoundFrontend/Public/DocumentConfigurations/MetasoundFrontendDocumentVertexConfiguration.h
new file mode 100644
--- /dev/null
+++ b/Engine/Plugins/Runtime/Metasound/Source/MetasoundFrontend/Public/DocumentConfigurations/MetasoundFrontendDocumentVertexConfiguration.h
@@ -0,0 +1,146 @@
+// Copyright Epic Games, Inc. All Rights Reserved.
+#pragma once
+
+#include "MetasoundFrontendDocumentConfiguration.h"
+
+#include "MetasoundFrontendDocumentVertexConfiguration.generated.h"
+
+#define UE_API METASOUNDFRONTEND_API
+
+// Forward Declarations
+struct FMetaSoundFrontendDocumentBuilder;
+
+
+// Vertex metadata displayed on MetaSound inputs & outputs pertaining to a given configuration.
+USTRUCT(BlueprintType)
+struct FMetaSoundFrontendDocumentConfigurationVertexMetadata
+{
+	GENERATED_BODY()
+
+private:
+	UPROPERTY()
+	FGuid NodeID;
+
+public:
+	FMetaSoundFrontendDocumentConfigurationVertexMetadata() = default;
+	FMetaSoundFrontendDocumentConfigurationVertexMetadata(const FGuid& InNodeID)
+		: NodeID(InNodeID)
+	{
+	}
+
+	virtual ~FMetaSoundFrontendDocumentConfigurationVertexMetadata() = default;
+
+	UE_EXPERIMENTAL(5.8, "Document Configuration is in active development and subject to change")
+	const FGuid& GetNodeID() const { return NodeID; }
+
+#if WITH_EDITOR
+	// If true, default values are directly editable if the parent configuration lists member editing as enabled.
+	UE_EXPERIMENTAL(5.8, "Document Configuration is in active development and subject to change")
+	virtual bool IsDefaultEditable() const { return true; }
+#endif // WITH_EDITOR
+
+	UE_EXPERIMENTAL(5.8, "Document Configuration is in active development and subject to change")
+	virtual void OnDefaultReset(const FMetaSoundFrontendDocumentBuilder& Builder) { }
+
+	UE_EXPERIMENTAL(5.8, "Document Configuration is in active development and subject to change")
+	virtual void OnDefaultUpdated(const FMetaSoundFrontendDocumentBuilder& Builder) { }
+
+#if WITH_EDITOR
+	UE_EXPERIMENTAL(5.8, "Document Configuration is in active development and subject to change")
+	virtual void OnPropertyChanged(const FPropertyChangedEvent& InEvent, FMetaSoundFrontendDocumentBuilder& OutBuilder) { }
+#endif // WITH_EDITOR
+};
+
+// Configuration of a MetaSound that stores vertex (input or output) Metadata to display alongside native properties in MetaSound vertex editors.
+USTRUCT(BlueprintType, meta = (Hidden))
+struct FMetaSoundFrontendDocumentVertexConfiguration : public FMetaSoundFrontendDocumentConfiguration
+{
+	GENERATED_BODY()
+
+	virtual ~FMetaSoundFrontendDocumentVertexConfiguration() = default;
+
+	// Creates and stores new VertexMetadata with the provided NodeID. Overwrites if metadata already exists with
+	// default constructed VertexMetadata if configuration supports Metadata.
+	UE_INTERNAL FMetaSoundFrontendDocumentConfigurationVertexMetadata* AddVertexMetadata(const FMetaSoundFrontendDocumentBuilder& Builder, const FGuid& NodeID);
+
+	// Creates and stores new VertexMetadata with the provided type and name, asserting that the type created by the overloaded configuration
+	// matches the template type. Overwrites if metadata already exists with default constructed VertexMetadata. Returns null if configuration
+	// does not support Metadata.
+	template <typename TConfigVertexMetadata>
+	UE_INTERNAL TConfigVertexMetadata* AddVertexMetadata(const FMetaSoundFrontendDocumentBuilder& Builder, const FGuid& NodeID)
+	{
+		if (TInstancedStruct<FMetaSoundFrontendDocumentConfigurationVertexMetadata>* Entry = VertexMetadata.Find(NodeID))
+		{
+			return &Entry->GetMutable<TConfigVertexMetadata>();
+		}
+
+PRAGMA_DISABLE_EXPERIMENTAL_WARNINGS
+		TInstancedStruct<FMetaSoundFrontendDocumentConfigurationVertexMetadata> NewMetadata = ConstructVertexMetadata(Builder, NodeID);
+PRAGMA_ENABLE_EXPERIMENTAL_WARNINGS
+
+		if (NewMetadata.IsValid())
+		{
+			return &VertexMetadata.Add(NodeID, MoveTemp(NewMetadata)).GetMutable<TConfigVertexMetadata>();
+		}
+
+		return nullptr;
+	}
+
+	// Returns whether or not there is existing Metadata with the given NodeID.
+	UE_API bool ContainsVertexMetadata(const FGuid& NodeID) const;
+
+	// Helper function to returns VertexMetadata with the provided class type and name and cast to underlying config type.
+	template <typename TConfigVertexMetadata = FMetaSoundFrontendDocumentConfigurationVertexMetadata>
+	UE_INTERNAL TConfigVertexMetadata* FindVertexMetadata(const FGuid& NodeID)
+	{
+		if (TInstancedStruct<FMetaSoundFrontendDocumentConfigurationVertexMetadata>* Entry = VertexMetadata.Find(NodeID))
+		{
+			return Entry->GetMutablePtr<TConfigVertexMetadata>();
+		}
+
+		return nullptr;
+	}
+
+	// Helper function to returns VertexMetadata with the provided class type and name and cast to underlying config type.
+	template <typename TConfigVertexMetadata = FMetaSoundFrontendDocumentConfigurationVertexMetadata>
+	const TConfigVertexMetadata* FindConstVertexMetadata(const FGuid& NodeID) const
+	{
+		if (const TInstancedStruct<FMetaSoundFrontendDocumentConfigurationVertexMetadata>* Entry = VertexMetadata.Find(NodeID))
+		{
+			return Entry->GetPtr<TConfigVertexMetadata>();
+		}
+
+		return nullptr;
+	}
+
+	const UScriptStruct* GetVertexMetadataStruct(const FGuid& NodeID) const
+	{
+		if (const TInstancedStruct<FMetaSoundFrontendDocumentConfigurationVertexMetadata>* Entry = VertexMetadata.Find(NodeID))
+		{
+			return Entry->GetScriptStruct();
+		}
+
+		return nullptr;
+	}
+
+	// Removes VertexMetadata with the provided NodeID. Returns if Metadata was found and removed or not.
+	UE_EXPERIMENTAL(5.8, "Document Configuration is in active development and subject to change")
+	UE_API bool RemoveVertexMetadata(const FGuid& NodeID);
+
+protected:
+	UE_EXPERIMENTAL(5.8, "Document Configuration is in active development and subject to change")
+	UE_API virtual TInstancedStruct<FMetaSoundFrontendDocumentConfigurationVertexMetadata> ConstructVertexMetadata(const FMetaSoundFrontendDocumentBuilder& Builder, const FGuid& NodeID) const;
+
+private:
+	UPROPERTY()
+	TMap<FGuid, TInstancedStruct<FMetaSoundFrontendDocumentConfigurationVertexMetadata>> VertexMetadata;
+};
+
+// Convenience, less verbose aliases
+namespace Metasound::Frontend
+{
+	using FDocumentVertexConfiguration = FMetaSoundFrontendDocumentVertexConfiguration;
+	using FDocumentConfigurationVertexMetadata = FMetaSoundFrontendDocumentConfigurationVertexMetadata;
+} // namespace Metasound::Frontend;
+
+#undef UE_API
diff --git a/Engine/Plugins/Runtime/Metasound/Source/MetasoundFrontend/Public/DocumentConfigurations/MetasoundFrontendPresetDocumentConfiguration.h b/Engine/Plugins/Runtime/Metasound/Source/MetasoundFrontend/Public/DocumentConfigurations/MetasoundFrontendPresetDocumentConfiguration.h
new file mode 100644
--- /dev/null
+++ b/Engine/Plugins/Runtime/Metasound/Source/MetasoundFrontend/Public/DocumentConfigurations/MetasoundFrontendPresetDocumentConfiguration.h
@@ -0,0 +1,78 @@
+// Copyright Epic Games, Inc. All Rights Reserved.
+#pragma once
+
+#include "MetasoundDocumentInterface.h"
+#include "MetasoundFrontendDocumentVertexConfiguration.h"
+#include "UObject/NameTypes.h"
+
+#include "MetasoundFrontendPresetDocumentConfiguration.generated.h"
+
+#define UE_API METASOUNDFRONTEND_API
+
+// Forward Declarations
+class UMetaSoundFrontendMemberMetadata;
+enum class EMetasoundFrontendClassType : uint8;
+
+
+USTRUCT(BlueprintType, DisplayName = "Preset")
+struct FMetaSoundFrontendPresetDocumentConfiguration final : public FMetaSoundFrontendDocumentVertexConfiguration
+{
+	GENERATED_BODY()
+
+	virtual ~FMetaSoundFrontendPresetDocumentConfiguration() = default;
+
+#if WITH_EDITORONLY_DATA
+	UE_INTERNAL UE_API virtual bool ConfigureDocument(FMetaSoundFrontendDocumentBuilder& OutBuilder, bool bLogFailure) override;
+#endif // WITH_EDITORONLY_DATA
+
+#if WITH_EDITOR
+	UE_EXPERIMENTAL(5.8, "Document Configuration is in active development and subject to change")
+	virtual TArray<const UClass*> GetSupportedAssetActionClasses() const override;
+
+	UE_EXPERIMENTAL(5.8, "Document Configuration is in active development and subject to change")
+	virtual bool IsMemberEditingEnabled() const override { return true; }
+
+	UE_EXPERIMENTAL(5.8, "Document Configuration is in active development and subject to change")
+	virtual bool IsGraphEditorVisible() const override { return true; }
+
+	UE_EXPERIMENTAL(5.8, "Document Configuration is in active development and subject to change")
+	virtual void OnAssetInitialized(TArray<UObject*> SelectedObjects, FMetaSoundFrontendDocumentBuilder& OutBuilder) override;
+
+	UE_EXPERIMENTAL(5.8, "Document Configuration is in active development and subject to change")
+	UE_API virtual void OnPropertyChanged(const FPropertyChangedEvent& InEvent, FMetaSoundFrontendDocumentBuilder& OutBuilder) override;
+#endif // WITH_EDITOR
+
+protected:
+	UE_INTERNAL UE_API virtual TInstancedStruct<FMetaSoundFrontendDocumentConfigurationVertexMetadata> ConstructVertexMetadata(const FMetaSoundFrontendDocumentBuilder& Builder, const FGuid& NodeID) const override;
+
+public:
+	UPROPERTY(EditAnywhere, BlueprintReadOnly, Category = "Preset Options")
+	TScriptInterface<IMetaSoundDocumentInterface> Parent;
+};
+
+USTRUCT(BlueprintType)
+struct FMetaSoundFrontendPresetConfigurationVertexMetadata : public FMetaSoundFrontendDocumentConfigurationVertexMetadata
+{
+	GENERATED_BODY()
+
+	virtual ~FMetaSoundFrontendPresetConfigurationVertexMetadata() = default;
+
+#if WITH_EDITOR
+	virtual bool IsDefaultEditable() const override;
+#endif // WITH_EDITOR
+
+	virtual void OnDefaultReset(const FMetaSoundFrontendDocumentBuilder& Builder) override;
+	virtual void OnDefaultUpdated(const FMetaSoundFrontendDocumentBuilder& Builder) override;
+
+	// Enables overriding the input's inherited default value otherwise provided by the parent graph.
+	// Setting to true disables the configuration utilizing the input's default value if updated on the parent asset.
+	UPROPERTY(EditAnywhere, BlueprintReadOnly, Category = PresetOptions)
+	bool bOverrideInheritedDefault = false;
+};
+
+namespace Metasound::Frontend
+{
+	using FPresetDocumentConfiguration = FMetaSoundFrontendPresetDocumentConfiguration;
+	using FPresetVertexMetadata = FMetaSoundFrontendPresetConfigurationVertexMetadata;
+} // namespace Metasound::Frontend;
+#undef UE_API
\ No newline at end of file
diff --git a/Engine/Plugins/Runtime/Metasound/Source/MetasoundFrontend/Public/MetasoundFrontendController.h b/Engine/Plugins/Runtime/Metasound/Source/MetasoundFrontend/Public/MetasoundFrontendController.h
--- a/Engine/Plugins/Runtime/Metasound/Source/MetasoundFrontend/Public/MetasoundFrontendController.h
+++ b/Engine/Plugins/Runtime/Metasound/Source/MetasoundFrontend/Public/MetasoundFrontendController.h
@@ -594,11 +594,11 @@ namespace Metasound
 			/** Returns the ClassID associated with this graph. */
 			virtual FGuid GetClassID() const = 0;
 
-			/** Return the preset options for the current graph. */
-			virtual const FMetasoundFrontendGraphClassPresetOptions& GetGraphPresetOptions() const = 0;
+			UE_DEPRECATED(5.8, "PresetOptions have been integrated into Document Configuration")
+			virtual const FMetasoundFrontendGraphClassPresetOptions& GetGraphPresetOptions() const;
 
-			/** Sets the preset options for the current graph. */
-			virtual void SetGraphPresetOptions(const FMetasoundFrontendGraphClassPresetOptions& InPresetOptions) = 0;
+			UE_DEPRECATED(5.8, "PresetOptions have been integrated into Document Configuration")
+			virtual void SetGraphPresetOptions(const FMetasoundFrontendGraphClassPresetOptions& InPresetOptions);
 
 			/** Return the metadata for the current graph. */
 			virtual const FMetasoundFrontendClassMetadata& GetGraphMetadata() const = 0;
@@ -653,17 +653,14 @@ namespace Metasound
 			/** Returns all input nodes in the graph. */
 			virtual TArray<FConstNodeHandle> GetConstInputNodes() const = 0;
 
-			/** Returns a set of all input names that are managed */
-			virtual const TSet<FName>& GetInputsInheritingDefault() const = 0;
+			UE_DEPRECATED(5.8, "PresetOptions have been integrated into Document Configuration")
+			virtual const TSet<FName>& GetInputsInheritingDefault() const;
 
-			/** If true, adds an item to the set of all input names
-			  * that are managed.
-			  * Returns true if successfully added/removed, false if not.
-			  */
-			virtual bool SetInputInheritsDefault(FName InName, bool bDefaultIsInherited) = 0;
+			UE_DEPRECATED(5.8, "PresetOptions have been integrated into Document Configuration")
+			virtual bool SetInputInheritsDefault(FName InName, bool bDefaultIsInherited);
 
-			/** Sets managed input names */
-			virtual void SetInputsInheritingDefault(TSet<FName>&& InNames) = 0;
+			UE_DEPRECATED(5.8, "PresetOptions have been integrated into Document Configuration")
+			virtual void SetInputsInheritingDefault(TSet<FName>&& InNames);
 
 			/** Adds a new variable to the graph */
 			virtual FVariableHandle AddVariable(const FName& InDataTypeName) = 0;
@@ -694,11 +691,8 @@ namespace Metasound
 			 */
 			virtual FConstVariableHandle FindVariableContainingNode(const FGuid& InNodeID) const = 0;
 
-			/** Removes the variable with the given ID. 
-			 *
-			 * @param InVariableID - ID of existing variable.
-			 */
-			virtual bool RemoveVariable(const FGuid& InVariableID) = 0;
+			UE_DEPRECATED(5.8, "Use the MetaSound Builder API")
+			virtual bool RemoveVariable(const FGuid& InVariableID) { return false; }
 
 			/* Returns an array of all variables associated with the graph. */
 			virtual TArray<FVariableHandle> GetVariables() = 0;
@@ -907,11 +901,8 @@ namespace Metasound
 			  */
 			virtual void UpdateInterfaceChangeID() = 0;
 
-			/** Clear the current literal for a given input.
-			 *
-			 * @return True on success, false on failure.
-			 */
-			virtual bool ClearLiteralForInput(const FVertexName& InInputName, FGuid InVertexID) = 0;
+			UE_DEPRECATED(5.8, "Use Builder API instead")
+			virtual bool ClearLiteralForInput(const FVertexName& InInputName, FGuid InVertexID) { return false;  }
 
 			/** Add a new node to this graph from the node registry.
 			 *
diff --git a/Engine/Plugins/Runtime/Metasound/Source/MetasoundFrontend/Public/MetasoundFrontendDocument.h b/Engine/Plugins/Runtime/Metasound/Source/MetasoundFrontend/Public/MetasoundFrontendDocument.h
--- a/Engine/Plugins/Runtime/Metasound/Source/MetasoundFrontend/Public/MetasoundFrontendDocument.h
+++ b/Engine/Plugins/Runtime/Metasound/Source/MetasoundFrontend/Public/MetasoundFrontendDocument.h
@@ -4,11 +4,11 @@
 #include "Algo/Transform.h"
 #include "Containers/Array.h"
 #include "Containers/Map.h"
+#include "DocumentConfigurations/MetasoundFrontendDocumentConfiguration.h"
 #include "HAL/IConsoleManager.h"
 #include "IAudioParameterInterfaceRegistry.h"
 #include "Internationalization/Text.h"
 #include "MetasoundAccessPtr.h"
-#include "MetasoundFrontendDocumentConfiguration.h"
 #include "MetasoundFrontendLiteral.h"
 #include "MetasoundOperatorData.h"
 #include "MetasoundNodeInterface.h"
@@ -1889,15 +1889,13 @@ struct FMetasoundFrontendGraphClassPresetOptions
 {
 	GENERATED_BODY()
 
-	// Whether or not graph class is a preset or not.
-	UPROPERTY()
+#if WITH_EDITORONLY_DATA
+	UPROPERTY(meta = (DeprecatedProperty, DeprecationMessage = "5.8 - Preset options now serialized as Document Configuration. See 'FMetaSoundFrontendPresetDocumentConfiguration'"))
 	bool bIsPreset = false;
 
-	// Names of all inputs inheriting default values from the referenced node. All input names
-	// in this set have their default value set on update when registered with the Frontend Class
-	// Registry.  Omitted inputs remain using the pre-existing, serialized default values.
-	UPROPERTY()
+	UPROPERTY(meta = (DeprecatedProperty, DeprecationMessage = "5.8 - Preset options now serialized as Document Configuration. See 'FMetaSoundFrontendPresetDocumentConfiguration'"))
 	TSet<FName> InputsInheritingDefault;
+#endif // WITH_EDITORONLY_DATA
 };
 
 USTRUCT()
@@ -1952,11 +1950,10 @@ private:
 	TArray<FMetasoundFrontendGraph> PagedGraphs;
 
 public:
-	UPROPERTY()
+#if WITH_EDITORONLY_DATA
+	UPROPERTY(meta = (DeprecationMessage = "5.8 - PresetOptions is now characterized as configuration vertex metadata on MetaSound Document."))
 	FMetasoundFrontendGraphClassPresetOptions PresetOptions;
 
-public:
-#if WITH_EDITORONLY_DATA
 	UE_API const FMetasoundFrontendGraph& AddGraphPage(const FGuid& InPageID, bool bDuplicateLastGraph = true, bool bSetAsBuildGraph = true);
 
 	// Removes the page associated with the given PageID.  Returns true if removed, false if not.
@@ -2062,10 +2059,8 @@ public:
 	UPROPERTY()
 	TArray<FMetasoundFrontendClass> Dependencies;
 
-#if WITH_EDITORONLY_DATA
-	UPROPERTY()
+	UPROPERTY(EditAnywhere, Category = CustomView)
 	TInstancedStruct<FMetaSoundFrontendDocumentConfiguration> Configuration;
-#endif // WITH_EDITORONLY_DATA
 
 	uint32 GetNextIdCounter() const
 	{
diff --git a/Engine/Plugins/Runtime/Metasound/Source/MetasoundFrontend/Public/MetasoundFrontendDocumentBuilder.h b/Engine/Plugins/Runtime/Metasound/Source/MetasoundFrontend/Public/MetasoundFrontendDocumentBuilder.h
--- a/Engine/Plugins/Runtime/Metasound/Source/MetasoundFrontend/Public/MetasoundFrontendDocumentBuilder.h
+++ b/Engine/Plugins/Runtime/Metasound/Source/MetasoundFrontend/Public/MetasoundFrontendDocumentBuilder.h
@@ -2,6 +2,7 @@
 #pragma once
 
 #include "Containers/Map.h"
+#include "DocumentConfigurations/MetasoundFrontendDocumentVertexConfiguration.h"
 #include "Interfaces/MetasoundFrontendInterfaceRegistry.h"
 #include "MetasoundFrontendDocumentCacheInterface.h"
 #include "MetasoundDataReference.h"
@@ -275,15 +276,16 @@ PRAGMA_ENABLE_EXPERIMENTAL_WARNINGS
 
 	// Clears document completely of all graph page data (nodes, edges, & member metadata), dependencies,
 	// interfaces, member metadata, preset state, etc. Leaves ClassMetadata intact. Reloads the builder state,
-	// so external delegates must be relinked if desired.
+	// so external delegates must be relinked if desired. NOTE: Configuration is not cleared and must be managed
+	// by separate API.
 	UE_API void ClearDocument(TSharedRef<Metasound::Frontend::FDocumentModifyDelegates> ModifyDelegates);
 
 #if WITH_EDITORONLY_DATA
 	UE_API bool ClearMemberMetadata(const FGuid& InMemberID);
+#endif // WITH_EDITORONLY_DATA
 
-	UE_EXPERIMENTAL(5.8, "Document Configuration is in active development and subject to change")
+	UE_EXPERIMENTAL(5.8, "Configurator is in active development and subject to change")
 	UE_API bool ConfigureDocument(bool bLogFailure = false);
-#endif // WITH_EDITORONLY_DATA
 
 	UE_API bool ConformObjectToDocument();
 
@@ -291,9 +293,13 @@ PRAGMA_ENABLE_EXPERIMENTAL_WARNINGS
 	UE_API bool ContainsEdge(const FMetasoundFrontendEdge& InEdge, const FGuid* InPageID = nullptr) const;
 	UE_API bool ContainsNode(const FGuid& InNodeID, const FGuid* InPageID = nullptr) const;
 
+	UE_DEPRECATED(5.8, "Use SetDocumentConfiguration instead")
 	UE_API bool ConvertFromPreset();
-	UE_DEPRECATED(5.7, "Use the ConvertToPreset which takes in a builder")
+
+	UE_DEPRECATED(5.8, "Use SetDocumentConfiguration instead")
 	UE_API bool ConvertToPreset(const FMetasoundFrontendDocument& InReferencedDocument, TSharedPtr<Metasound::Frontend::FDocumentModifyDelegates> ModifyDelegates = { });
+
+	UE_DEPRECATED(5.8, "Use SetDocumentConfiguration instead")
 	UE_API bool ConvertToPreset(const FMetaSoundFrontendDocumentBuilder& InReferencedDocumentBuilder, TSharedPtr<Metasound::Frontend::FDocumentModifyDelegates> ModifyDelegates = { });
 	
 
@@ -391,23 +397,58 @@ PRAGMA_ENABLE_EXPERIMENTAL_WARNINGS
 
 	UE_API const FMetasoundFrontendDocument& GetConstDocumentChecked() const;
 
-#if WITH_EDITORONLY_DATA
+	// Convenience template function for attempting to get the underlying mutable
+	// vertex configuration metadata. If the provided vertex configuration type
+	// is not that which the current document is set to or no configuration is
+	// set, returns null.
+	template<typename TVertexMetadataType>
+	UE_INTERNAL TVertexMetadataType* FindVertexMetadata(const FGuid& InNodeID) const
+	{
+		using namespace Metasound::Frontend;
+
+		if (FDocumentVertexConfiguration* Config = GetDocumentChecked().Configuration.GetMutablePtr<FDocumentVertexConfiguration>())
+		{
+			return Config->FindVertexMetadata<TVertexMetadataType>(InNodeID);
+		}
+
+		return nullptr;
+	}
+
+	// Convenience template function for attempting to get the underlying const
+	// vertex configuration metadata. If the provided vertex configuration type
+	// is not that which the current document is set to or no configuration is
+	// set, returns null.
+	template<typename TVertexMetadataType>
+	UE_INTERNAL const TVertexMetadataType* FindConstVertexMetadata(const FGuid& InNodeID) const
+	{
+		using namespace Metasound::Frontend;
+
+		if (const FDocumentVertexConfiguration* Config = GetConstDocumentConfiguration().GetPtr<FDocumentVertexConfiguration>())
+		{
+			return Config->FindConstVertexMetadata<TVertexMetadataType>(InNodeID);
+		}
+
+		return nullptr;
+	}
+
+#if WITH_EDITOR
+	// Returns struct view of document configuration if set. Internal use only to avoid
+	// writing against struct members in invalid or dangerous, non-editor contexts (use
+	// 'SetDocumentConfiguration' instead for most script, gameplay or runtime applications).
+	UE_INTERNAL UE_API TStructView<FMetaSoundFrontendDocumentConfiguration> GetDocumentConfiguration() const;
+#endif // WITH_EDITOR
+
 	// Returns struct view of document configuration if set.
-	UE_EXPERIMENTAL(5.8, "Document Configuration is in active development and subject to change")
 	UE_API TConstStructView<FMetaSoundFrontendDocumentConfiguration> GetConstDocumentConfiguration() const;
 
 	// Convenience template function for attempting to get the underlying const pointer of the set configuration.
 	// If the provided configuration type is not that which the current document is set to or no configuration
 	// is set, returns null.
 	template<typename TConfigType>
-	UE_EXPERIMENTAL(5.8, "Document Configuration is in active development and subject to change")
 	const TConfigType* GetConstDocumentConfigurationAs() const
 	{
-PRAGMA_DISABLE_EXPERIMENTAL_WARNINGS
 		return GetConstDocumentConfiguration().GetPtr<TConfigType>();
-PRAGMA_ENABLE_EXPERIMENTAL_WARNINGS
 	}
-#endif // WITH_EDITORONLY_DATA
 
 	UE_API const IMetaSoundDocumentInterface& GetConstDocumentInterfaceChecked() const;
 	UE_API const FString GetDebugName() const;
@@ -453,15 +494,15 @@ PRAGMA_ENABLE_EXPERIMENTAL_WARNINGS
 	UE_API const FTopLevelAssetPath& GetHintPath() const;
 
 	UE_API FMetasoundAssetBase& GetMetasoundAsset() const;
-	
-	// Get the asset referenced by this builder's preset asset, nullptr if builder is not a preset.
+
+	UE_DEPRECATED(5.8, "Use DocumentConfiguration API to access configuration as type FMetaSoundFrontendPresetDocumentConfiguration and use the 'Parent' property directly.")
 	UE_API FMetasoundAssetBase* GetReferencedPresetAsset() const;
 
 	UE_API int32 GetTransactionCount() const;
 
 	UE_API TArray<const FMetasoundFrontendNode*> GetGraphInputTemplateNodes(FName InInputName, const FGuid* InPageID = nullptr);
 
-	// If graph is set to be a preset, returns set of graph input names inheriting default data from the referenced graph. If not a preset, returns null.
+	UE_DEPRECATED(5.8, "Input inheritance data is now stored on MetaSound Preset document configuration metadata.")
 	UE_API const TSet<FName>* GetGraphInputsInheritingDefault() const;
 
 	UE_API EMetasoundFrontendVertexAccessType GetNodeInputAccessType(const FGuid& InNodeID, const FGuid& InVertexID, const FGuid* InPageID = nullptr) const;
@@ -506,7 +547,6 @@ PRAGMA_ENABLE_EXPERIMENTAL_WARNINGS
 
 	UE_API bool IsDependencyReferenced(const FGuid& InClassID) const;
 
-#if WITH_EDITORONLY_DATA
 	// Determines if the builder's document configuration is of the given type
 	template<typename TConfigType>
 	const bool IsDocumentConfiguredBy() const
@@ -518,7 +558,6 @@ PRAGMA_ENABLE_EXPERIMENTAL_WARNINGS
 
 		return false;
 	}
-#endif // WITH_EDITORONLY_DATA
 
 	UE_API bool IsNodeInputConnected(const FGuid& InNodeID, const FGuid& InVertexID, const FGuid* InPageID = nullptr) const;
 	UE_API bool IsNodeInputConnectionUserModifiable(const FGuid& InNodeID, const FGuid& InVertexID, const FGuid* InPageID = nullptr) const;
@@ -589,7 +628,10 @@ PRAGMA_ENABLE_EXPERIMENTAL_WARNINGS
 	UE_API void Reload(TSharedPtr<Metasound::Frontend::FDocumentModifyDelegates> Delegates = {}, bool bPrimeCache = false);
 
 #if WITH_EDITORONLY_DATA
-	UE_API bool RemoveGraphInputDefault(FName InputName, const FGuid& InPageID, bool bClearInheritsDefault = true);
+	UE_DEPRECATED(5.8, "Input inheritance data is now cached in doc config.  No need to call directly here. Use overload without bool parameter.")
+	UE_API bool RemoveGraphInputDefault(FName InputName, const FGuid& InPageID, bool bClearInheritsDefault);
+
+	UE_API bool RemoveGraphInputDefault(FName InputName, const FGuid& InPageID);
 #endif // WITH_EDITORONLY_DATA
 
 	UE_API bool RemoveNodeInputDefault(const FGuid& InNodeID, const FGuid& InVertexID, const FGuid* InPageID = nullptr);
@@ -610,7 +652,6 @@ PRAGMA_ENABLE_EXPERIMENTAL_WARNINGS
 
 #if WITH_EDITOR
 	UE_API void SetAuthor(const FString& InAuthor);
-
 #endif // WITH_EDITOR
 
 #if WITH_EDITORONLY_DATA
@@ -624,11 +665,11 @@ PRAGMA_ENABLE_EXPERIMENTAL_WARNINGS
 	UE_API bool SetBuildPageID(const FGuid& InBuildPageID, bool bBroadcastDelegate = true);
 #endif // WITH_EDITORONLY_DATA
 
-#if WITH_EDITORONLY_DATA
 PRAGMA_DISABLE_EXPERIMENTAL_WARNINGS
 	UE_API void SetDocumentConfiguration(TInstancedStruct<FMetaSoundFrontendDocumentConfiguration> InConfig);
 PRAGMA_ENABLE_EXPERIMENTAL_WARNINGS
 
+#if WITH_EDITORONLY_DATA
 	// If page ID is provided, updates the graph's editable flag to the given value. If no page ID is provided,
 	// sets *all* paged graph fields to the given value.
 	UE_API void SetIsGraphEditable(bool bIsEditable, const FGuid* PageID = nullptr) const;
@@ -659,10 +700,10 @@ PRAGMA_ENABLE_EXPERIMENTAL_WARNINGS
 	UE_API bool SetGraphInputDisplayName(FName InputName, FText DisplayName);
 #endif // WITH_EDITORONLY_DATA
 
-	// Sets whether or not graph input inherits default.  By default, updates only if graph is marked as a preset.
-	// Optionally, if bForceUpdate is set, updates inheritance even if not a preset (primarily for clearing flag
-	// if non-preset has unnecessary data).
+	UE_DEPRECATED(5.8, "Input inheritance information moved to preset document configuration")
 	UE_API bool SetGraphInputInheritsDefault(FName InName, bool bInputInheritsDefault, bool bForceUpdate = false);
+
+	UE_DEPRECATED(5.8, "Input inheritance information moved to preset document configuration")
 	UE_API bool SetGraphInputsInheritingDefault(TSet<FName>&& InNames);
 
 	// Sets a given graph input's name to a new name. Succeeds if the graph output exists and the new name is set (or is the same as the old name).
@@ -750,8 +791,6 @@ PRAGMA_ENABLE_EXPERIMENTAL_WARNINGS
 
 	UE_API bool SetNodeInputDefault(const FGuid& InNodeID, const FGuid& InVertexID, const FMetasoundFrontendLiteral& InLiteral, const FGuid* InPageID = nullptr);
 	
-	UE_INTERNAL void SetPresetFlags(bool bIsPreset=true);
-
 	// Sets the document's version number.  Should only be called by document versioning.
 	UE_API void SetVersionNumber(const FMetasoundFrontendVersionNumber& InDocumentVersionNumber);
 
@@ -876,10 +915,9 @@ private:
 	void IterateNodesConnectedWithVertex(const FMetasoundFrontendVertexHandle& Vertex, TFunctionRef<void(const FMetasoundFrontendEdge&, FMetasoundFrontendNode&)> NodeIndexIterFunc, const FGuid& InPageID);
 
 	const FTopLevelAssetPath GetBuilderClassPath() const;
-	FMetasoundFrontendDocument& GetDocumentChecked() const;
-	IMetaSoundDocumentInterface& GetDocumentInterfaceChecked() const;
-
-	bool ModifyInterfacesInternal(Metasound::Frontend::FModifyInterfaceOptions&& InOptions, const FMetaSoundFrontendDocumentBuilder* InBuilder = nullptr);
+	
+	UE_API FMetasoundFrontendDocument& GetDocumentChecked() const;
+	UE_API IMetaSoundDocumentInterface& GetDocumentInterfaceChecked() const;
 
 	void RemoveSwapDependencyInternal(int32 Index);
 
diff --git a/Engine/Plugins/Runtime/Metasound/Source/MetasoundFrontend/Public/MetasoundFrontendDocumentBuilderRegistry.h b/Engine/Plugins/Runtime/Metasound/Source/MetasoundFrontend/Public/MetasoundFrontendDocumentBuilderRegistry.h
--- a/Engine/Plugins/Runtime/Metasound/Source/MetasoundFrontend/Public/MetasoundFrontendDocumentBuilderRegistry.h
+++ b/Engine/Plugins/Runtime/Metasound/Source/MetasoundFrontend/Public/MetasoundFrontendDocumentBuilderRegistry.h
@@ -1,6 +1,7 @@
 // Copyright Epic Games, Inc. All Rights Reserved.
 #pragma once
 
+#include "UObject/Interface.h"
 #include "MetasoundBuilderInterface.h"
 #include "MetasoundDocumentInterface.h"
 #include "MetasoundFrontendDocumentModifyDelegates.h"
@@ -8,7 +9,6 @@
 #include "MetasoundOperatorSettings.h"
 #include "Templates/Function.h"
 #include "Templates/UniquePtr.h"
-#include "UObject/Interface.h"
 #include "UObject/TopLevelAssetPath.h"
 
 
@@ -29,6 +29,11 @@ namespace Metasound::Frontend
 		virtual ~IDocumentBuilderRegistry() = default;
 
 #if WITH_EDITORONLY_DATA
+		// Given the provided builder, removes paged data within the associated document for a cooked build.
+		// This function removes graphs and input defaults which are not to ever be used by a given cook
+		// platform, allowing users to optimize away data and scale the amount of memory required for
+		// initial load of input UObjects and graph topology, which can also positively effect runtime
+		// performance as well, etc. Returns true if builder modified the document, false if not.
 		UE_DEPRECATED(5.7, "Use StripUnusedPages(...) instead.")
 		virtual bool CookPages(FName CookPlatformName, FMetaSoundFrontendDocumentBuilder& Builder) const = 0;
 #endif // WITH_EDITORONLY_DATA
diff --git a/Engine/Plugins/Runtime/Metasound/Source/MetasoundFrontend/Public/MetasoundFrontendGraphLayer.h b/Engine/Plugins/Runtime/Metasound/Source/MetasoundFrontend/Public/MetasoundFrontendGraphLayer.h
--- a/Engine/Plugins/Runtime/Metasound/Source/MetasoundFrontend/Public/MetasoundFrontendGraphLayer.h
+++ b/Engine/Plugins/Runtime/Metasound/Source/MetasoundFrontend/Public/MetasoundFrontendGraphLayer.h
@@ -12,8 +12,6 @@
 
 #define UE_API METASOUNDFRONTEND_API
 
-struct UE_EXPERIMENTAL(5.7, "GraphLayers are currently being tested for Config generation, but may expand functionality and change in the future.") FMetasoundFrontendGraphLayer;
-
 USTRUCT()
 struct FMetasoundFrontendGraphLayer
 {
diff --git a/Engine/Plugins/Runtime/Metasound/Source/MetasoundFrontend/Public/MetasoundFrontendLiteral.h b/Engine/Plugins/Runtime/Metasound/Source/MetasoundFrontend/Public/MetasoundFrontendLiteral.h
--- a/Engine/Plugins/Runtime/Metasound/Source/MetasoundFrontend/Public/MetasoundFrontendLiteral.h
+++ b/Engine/Plugins/Runtime/Metasound/Source/MetasoundFrontend/Public/MetasoundFrontendLiteral.h
@@ -4,6 +4,7 @@
 
 #include "Algo/Transform.h"
 #include "AudioParameterControllerInterface.h"
+#include "Containers/ContainersFwd.h"
 #include "MetasoundLiteral.h"
 
 #include <type_traits>
@@ -178,31 +179,43 @@ private:
 
 namespace Metasound::Frontend
 {
-	template <typename ObjectType = UObject>
-	FMetasoundFrontendLiteral MakeFrontendLiteral(TConstArrayView<TObjectPtr<ObjectType>> ObjectPtrs)
-	{
-		TArray<UObject*> Objects;
-		Algo::Transform(ObjectPtrs, Objects,
-			[](const TObjectPtr<ObjectType> Obj) { return Obj.Get(); });
+	// Convenience function to convert Metasound::ELiteralType to EMetasoundFrontendLiteralType.
+	UE_API EMetasoundFrontendLiteralType GetMetasoundFrontendLiteralType(ELiteralType InLiteralType);
 
+	// Convenience function to convert EMetasoundFrontendLiteralType to Metasound::ELiteralType.
+	UE_API ELiteralType GetMetasoundLiteralType(EMetasoundFrontendLiteralType InLiteralType);
+
+	template<typename ... ArgTypes>
+	FMetasoundFrontendLiteral MakeFrontendLiteral(ArgTypes&& ... InArgs)
+	{
 		FMetasoundFrontendLiteral Value;
-		Value.Set(MoveTemp(Objects));
+		if constexpr (sizeof...(InArgs) == 0)
+		{
+			Value.Set(FMetasoundFrontendLiteral::FDefault{ });
+		}
+		else if constexpr (sizeof...(InArgs) == 1)
+		{
+			Value.Set(Forward<ArgTypes>(InArgs)...);
+		}
+		else
+		{
+			static_assert(false, "MakeFrontendLiteral: Invalid num arguments");
+		}
+
 		return Value;
 	}
 
-	template <typename LiteralType>
-	FMetasoundFrontendLiteral MakeFrontendLiteral(LiteralType InitValue)
+	template <typename ObjectType = UObject>
+	FMetasoundFrontendLiteral MakeFrontendLiteral(TConstArrayView<TObjectPtr<ObjectType>> ObjectPtrs)
 	{
+		TArray<UObject*> Objects;
+		auto ObjPtrToPtr = [](const TObjectPtr<ObjectType> Obj) { return Obj.Get(); };
+		Algo::Transform(ObjectPtrs, Objects, ObjPtrToPtr);
+
 		FMetasoundFrontendLiteral Value;
-		Value.Set(MoveTemp(InitValue));
+		Value.Set(MoveTemp(Objects));
 		return Value;
 	}
-
-	// Convenience function to convert Metasound::ELiteralType to EMetasoundFrontendLiteralType.
-	UE_API EMetasoundFrontendLiteralType GetMetasoundFrontendLiteralType(ELiteralType InLiteralType);
-
-	// Convenience function to convert EMetasoundFrontendLiteralType to Metasound::ELiteralType.
-	UE_API ELiteralType GetMetasoundLiteralType(EMetasoundFrontendLiteralType InLiteralType);
 }
 
 #undef UE_API
diff --git a/Engine/Plugins/Runtime/Metasound/Source/MetasoundFrontend/Public/MetasoundFrontendTransform.h b/Engine/Plugins/Runtime/Metasound/Source/MetasoundFrontend/Public/MetasoundFrontendTransform.h
--- a/Engine/Plugins/Runtime/Metasound/Source/MetasoundFrontend/Public/MetasoundFrontendTransform.h
+++ b/Engine/Plugins/Runtime/Metasound/Source/MetasoundFrontend/Public/MetasoundFrontendTransform.h
@@ -188,17 +188,16 @@ namespace Metasound::Frontend
 		UE_API bool Transform(FMetaSoundFrontendDocumentBuilder& InOutBuilderToTransform) const;
 
 	private:
-
-		// Get the class inputs needed for this preset. Input literals set on 
-		// the preset graph will be used if they are set and are marked as inheriting
-		// the default from the referenced graph.
-		TArray<FMetasoundFrontendClassInput> GenerateRequiredClassInputs(FMetaSoundFrontendDocumentBuilder& InDocumentToTransformBuilder, TSet<FName>& OutInputsInheritingDefault) const;
+		// Get the class inputs needed for this preset. Input literals set on
+		// the preset graph will be used if they are set as overriding
+		// the default from the parent graph.
+		TArray<FMetasoundFrontendClassInput> GenerateRequiredClassInputs(FMetaSoundFrontendDocumentBuilder& InDocumentToTransformBuilder) const;
 
 		// Get the class Outputs needed for this preset.
 		TArray<FMetasoundFrontendClassOutput> GenerateRequiredClassOutputs(FMetaSoundFrontendDocumentBuilder& InDocumentToTransformBuilder) const;
 		
 		// Add inputs to parent graph and connect to wrapped graph node.
-		void AddAndConnectInputs(const TArray<FMetasoundFrontendClassInput>& InClassInputs, FMetaSoundFrontendDocumentBuilder& InOutBuilderToTransform, const FGuid& InReferencedNodeID) const;
+		void AddAndConnectInputs(const FMetaSoundFrontendPresetDocumentConfiguration& PresetConfig, const TArray<FMetasoundFrontendClassInput>& InClassInputs, FMetaSoundFrontendDocumentBuilder& InOutBuilderToTransform, const FGuid& InReferencedNodeID) const;
 
 		// Add outputs to parent graph and connect to wrapped graph node.
 		void AddAndConnectOutputs(const TArray<FMetasoundFrontendClassOutput>& InClassOutputs, FMetaSoundFrontendDocumentBuilder& InOutBuilderToTransform, const FGuid& InReferencedNodeID) const;
@@ -206,7 +205,7 @@ namespace Metasound::Frontend
 #if WITH_EDITORONLY_DATA
 		using FMemberIDToMetadataMap = TMap<FGuid, TObjectPtr<UMetaSoundFrontendMemberMetadata>>;
 		void AddMemberMetadata(const FMemberIDToMetadataMap& InCachedMemberMetadata, FMetaSoundFrontendDocumentBuilder& InOutBuilderToTransform) const;
-		void CacheMemberMetadata(FMetaSoundFrontendDocumentBuilder& InOutBuilderToTransform, const TSet<FName>& InInputsInheritingDefault, FMemberIDToMetadataMap& InOutCachedMemberMetadata) const;
+		FMemberIDToMetadataMap CacheMemberMetadata(const FMetaSoundFrontendPresetDocumentConfiguration& PresetConfig, FMetaSoundFrontendDocumentBuilder& InOutBuilderToTransform) const;
 #endif // WITH_EDITORONLY_DATA
 
 		const FMetaSoundFrontendDocumentBuilder* ParentBuilder = nullptr;
