Fix for Mac build (namespace disambiguation)
#rnx

#Synced-CL 48741549

diff --git a/Engine/Plugins/Runtime/Metasound/Source/MetasoundEditor/Private/MetasoundAssetDefinitions.cpp b/Engine/Plugins/Runtime/Metasound/Source/MetasoundEditor/Private/MetasoundAssetDefinitions.cpp
--- a/Engine/Plugins/Runtime/Metasound/Source/MetasoundEditor/Private/MetasoundAssetDefinitions.cpp
+++ b/Engine/Plugins/Runtime/Metasound/Source/MetasoundEditor/Private/MetasoundAssetDefinitions.cpp
@@ -903,7 +903,7 @@ namespace MenuExtension_MetaSoundTemplate
 
 			// MetaSound Source Action Registration
 			{
-				const FSlateIcon AssetIcon = Style::CreateSlateIcon("ClassIcon.MetasoundSource");
+				const FSlateIcon AssetIcon = Metasound::Editor::Style::CreateSlateIcon("ClassIcon.MetasoundSource");
 
 				RegisterDocumentConfigurationActions<UMetaSoundSource, UMetaSoundSourceFactory>(AssetIcon);
 
