MetaSound DocConfig UX feedback/fixes missed in last local checkpoint resolve + BugHawk cleanup
#rnx
#Synced-CL 48836534

diff --git a/Engine/Plugins/Runtime/Metasound/Source/MetasoundEditor/Private/MetasoundAssetDefinitions.cpp b/Engine/Plugins/Runtime/Metasound/Source/MetasoundEditor/Private/MetasoundAssetDefinitions.cpp
--- a/Engine/Plugins/Runtime/Metasound/Source/MetasoundEditor/Private/MetasoundAssetDefinitions.cpp
+++ b/Engine/Plugins/Runtime/Metasound/Source/MetasoundEditor/Private/MetasoundAssetDefinitions.cpp
@@ -844,7 +844,7 @@ namespace MenuExtension_MetaSoundTemplate
 			}
 
 			// Don't show structs marked as hidden
-			if (Struct->FindMetaData("Hidden") != nullptr)
+			if (Struct->HasMetaData("Hidden"))
 			{
 				continue;
 			}
diff --git a/Engine/Plugins/Runtime/Metasound/Source/MetasoundEditor/Private/MetasoundEditor.cpp b/Engine/Plugins/Runtime/Metasound/Source/MetasoundEditor/Private/MetasoundEditor.cpp
--- a/Engine/Plugins/Runtime/Metasound/Source/MetasoundEditor/Private/MetasoundEditor.cpp
+++ b/Engine/Plugins/Runtime/Metasound/Source/MetasoundEditor/Private/MetasoundEditor.cpp
@@ -117,7 +117,7 @@ namespace Metasound
 		{
 			static bool bEnableLiveAudition = false;
 			static bool bIgnoreDocConfigVisibilitySettings = false;
-
+			static bool bShowGraphOverride = false;
 			static const FText LoadingAssetsFormat = LOCTEXT("MetaSoundScanInProgressNotificationText", "Loading {0} MetaSound Assets...");
 
 			FAutoConsoleVariableRef CVarMetaSoundEditorEnableLiveAudition(
@@ -1846,6 +1846,11 @@ namespace Metasound
 
 			// This asset 
 			const FMetaSoundFrontendDocumentBuilder& DocBuilder = Builder->GetConstBuilder();
+			if (!DocBuilder.IsValid())
+			{
+				return nullptr;
+			}
+
 			const FMetasoundFrontendGraphClass& ThisGraphClass = DocBuilder.GetConstDocumentChecked().RootGraph;
 			const FNodeRegistryKey ThisRegistryKey(ThisGraphClass);
 
@@ -2307,7 +2312,7 @@ namespace Metasound
 					if (UScriptStruct* Struct = Cast<UScriptStruct>(*It))
 					{
 						if (Struct->IsChildOf(FMetaSoundFrontendDocumentConfiguration::StaticStruct())
-							&& Struct != FMetaSoundFrontendDocumentConfiguration::StaticStruct())
+							&& !Struct->HasMetaData("Hidden"))
 						{
 							CreateConfigEntry(Struct);
 						}
@@ -2343,7 +2348,28 @@ namespace Metasound
 							LOCTEXT("ConfigurationSettingsMenu", "Configuration"),
 							LOCTEXT("ConfigurationSettingsMenu_Tooltip", "Settings related to procedurally-generated MetaSound Configurations"),
 							GetExportStatusImage(),
-							false
+							false /* bInSimpleComboBox*/,
+							NAME_None /* InTutorialHighlightName */,
+							TAttribute<EVisibility>::CreateSPLambda(AsShared(), [this]()
+							{
+								if (Builder.IsValid())
+								{
+									const FMetaSoundFrontendDocumentBuilder& DocBuilder = Builder->GetConstBuilder();
+									if (DocBuilder.IsValid())
+									{
+										using namespace Frontend;
+
+										//  Show option if override is required.  If the config requires it always visible, hide.
+										TConstStructView<FDocumentConfiguration> Config = DocBuilder.GetConstDocumentConfiguration();
+										if (Config.IsValid() && !Config.GetPtr()->IsGraphEditorVisible())
+										{
+											return EVisibility::Visible;
+										}
+									}
+								}
+
+								return EVisibility::Collapsed;
+							})
 						);
  					}
  					ToolbarBuilder.EndSection();
@@ -2417,7 +2443,7 @@ namespace Metasound
 							ToolbarBuilder.AddToolBarButton(
 								FEditorCommands::Get().ToggleShowGraph,
 								NAME_None,
-								TAttribute<FText>::CreateSPLambda(this, [this]() { return AssetEditorPrivate::bIgnoreDocConfigVisibilitySettings
+								TAttribute<FText>::CreateSPLambda(this, [this]() { return AssetEditorPrivate::bShowGraphOverride
 									? LOCTEXT("ShowGraphLabel", "Hide Graph")
 									: LOCTEXT("ShowGraphLabel", "Show Graph");
 								}),
@@ -2485,11 +2511,6 @@ namespace Metasound
 				Commands.TogglePlayback,
 				FExecuteAction::CreateSP(this, &FEditor::TogglePlayback));
 
-			ToolkitCommands->MapAction(
-				Commands.ToggleShowGraph,
-				FExecuteAction::CreateSP(this, &FEditor::ToggleShowHideGraph),
-				FCanExecuteAction::CreateSP(this, &FEditor::CanToggleShowHideGraph));
-
 			ToolkitCommands->MapAction(
 				FGenericCommands::Get().Undo,
 				FExecuteAction::CreateSP(this, &FEditor::UndoGraphAction));
@@ -2907,7 +2928,8 @@ namespace Metasound
 			{
 				if (const FMetaSoundFrontendDocumentBuilder& DocBuilder = Builder->GetBuilder(); ensure(DocBuilder.IsValid()))
 				{
-					return DocBuilder.GetConstDocumentConfiguration().IsValid();
+					TConstStructView<FMetaSoundFrontendDocumentConfiguration> Config = DocBuilder.GetConstDocumentConfiguration();
+					return Config.IsValid() && !Config.GetPtr()->IsGraphEditorVisible();
 				}
 			}
 
@@ -2917,7 +2939,7 @@ namespace Metasound
 
 		void FEditor::ToggleShowHideGraph()
 		{
-			AssetEditorPrivate::bIgnoreDocConfigVisibilitySettings = !AssetEditorPrivate::bIgnoreDocConfigVisibilitySettings;
+			AssetEditorPrivate::bShowGraphOverride = !AssetEditorPrivate::bShowGraphOverride;
 			UpdateEnabledTabs();
 		}
 
@@ -5368,8 +5390,9 @@ PRAGMA_ENABLE_DEPRECATION_WARNINGS
 				{
 					if (!AssetEditorPrivate::bIgnoreDocConfigVisibilitySettings)
 					{
-						ResolveTabEnablement(TabNamesPrivate::Find, Config->IsGraphEditorVisible());
-						ResolveTabEnablement(TabNamesPrivate::GraphCanvas, Config->IsGraphEditorVisible());
+						const bool bShowGraph = Config->IsGraphEditorVisible() || AssetEditorPrivate::bShowGraphOverride;
+						ResolveTabEnablement(TabNamesPrivate::Find, bShowGraph);
+						ResolveTabEnablement(TabNamesPrivate::GraphCanvas, bShowGraph);
 						ResolveTabEnablement(TabNamesPrivate::Interfaces, Config->IsInterfaceEditingEnabled());
 						ResolveTabEnablement(TabNamesPrivate::Members, Config->IsMemberEditingEnabled());
 
diff --git a/Engine/Plugins/Runtime/Metasound/Source/MetasoundEditor/Private/MetasoundEditor.h b/Engine/Plugins/Runtime/Metasound/Source/MetasoundEditor/Private/MetasoundEditor.h
--- a/Engine/Plugins/Runtime/Metasound/Source/MetasoundEditor/Private/MetasoundEditor.h
+++ b/Engine/Plugins/Runtime/Metasound/Source/MetasoundEditor/Private/MetasoundEditor.h
@@ -460,7 +460,7 @@ namespace Metasound::Editor
 
 		UUserWidget* CreateDocConfigWidget();
 
-		/** For teardown and regeneration of user preset widget on world change
+		/** For teardown and regeneration of user document configuration widget on world change
 			Based on UEditorUtilityWidgetBlueprint::ChangeTabWorld */
 		void ChangeDocConfigWidgetWorld(UWorld* World, EMapChangeType MapChangeType);
 
diff --git a/Engine/Plugins/Runtime/Metasound/Source/MetasoundEditor/Public/MetasoundEditorSubsystem.h b/Engine/Plugins/Runtime/Metasound/Source/MetasoundEditor/Public/MetasoundEditorSubsystem.h
--- a/Engine/Plugins/Runtime/Metasound/Source/MetasoundEditor/Public/MetasoundEditorSubsystem.h
+++ b/Engine/Plugins/Runtime/Metasound/Source/MetasoundEditor/Public/MetasoundEditorSubsystem.h
@@ -84,7 +84,7 @@ public:
 		EMetaSoundBuilderResult& OutResult);
 
 	// Initialize the UObject asset with an optional MetaSound document configuration to be applied.
-	// Optionally clears the existing MetaSound document (for the case of duplicated assets)
+	// Clears the existing MetaSound document (for the case of duplicated assets).
 	UE_API void InitAsset(UObject& InNewMetaSound, FInitAssetArgs InitAssetArgs = { });
 
 	UE_DEPRECATED(5.8, "Use overload that is passed optional initialization arguments. Call SetSoundWaveSettingsFromTemplate directly if desired.")
diff --git a/Engine/Plugins/Runtime/Metasound/Source/MetasoundEngine/Private/DocumentConfigurations/MetasoundRandomizerDocumentConfiguration.cpp b/Engine/Plugins/Runtime/Metasound/Source/MetasoundEngine/Private/DocumentConfigurations/MetasoundRandomizerDocumentConfiguration.cpp
--- a/Engine/Plugins/Runtime/Metasound/Source/MetasoundEngine/Private/DocumentConfigurations/MetasoundRandomizerDocumentConfiguration.cpp
+++ b/Engine/Plugins/Runtime/Metasound/Source/MetasoundEngine/Private/DocumentConfigurations/MetasoundRandomizerDocumentConfiguration.cpp
@@ -288,7 +288,7 @@ void FMetaSoundRandomizerDocumentConfiguration::OnAssetInitialized(TArray<UObjec
 {
 	Sounds.Reset();
 	Algo::TransformIf(SelectedObjects, Sounds,
-		[](const UObject* Object) { return Object->IsA<USoundWave>(); }, // Has to match type exactly so procedural sounds, MetaSounds, etc. aren't used
+		[](const UObject* Object) { return Object->GetClass() == USoundWave::StaticClass(); }, // Has to match type exactly so procedural sounds, SourceBuses, MetaSounds, etc. aren't used
 		[](UObject* Object) { return TObjectPtr<USoundWave>(CastChecked<USoundWave>(Object)); });
 }
 
diff --git a/Engine/Plugins/Runtime/Metasound/Source/MetasoundEngine/Public/DocumentConfigurations/MetasoundRandomizerDocumentConfiguration.h b/Engine/Plugins/Runtime/Metasound/Source/MetasoundEngine/Public/DocumentConfigurations/MetasoundRandomizerDocumentConfiguration.h
--- a/Engine/Plugins/Runtime/Metasound/Source/MetasoundEngine/Public/DocumentConfigurations/MetasoundRandomizerDocumentConfiguration.h
+++ b/Engine/Plugins/Runtime/Metasound/Source/MetasoundEngine/Public/DocumentConfigurations/MetasoundRandomizerDocumentConfiguration.h
@@ -41,7 +41,7 @@ public:
 	UPROPERTY(EditAnywhere, BlueprintReadOnly, meta = (Category = "Randomizer Options"))
 	bool bIsOneShot = false;
 
-	UPROPERTY(EditAnywhere, BlueprintReadOnly, meta = (Category = "Randomizer Options"))
+	UPROPERTY(EditAnywhere, BlueprintReadOnly, meta = (Category = "Randomizer Options", DisallowedClasses = "/Script/MetasoundEngine.MetaSoundSource, /Script/Engine.SoundSourceBus"))
 	TArray<TObjectPtr<USoundWave>> Sounds;
 
 	UPROPERTY(EditAnywhere, BlueprintReadOnly, meta = (Category = "Randomizer Options", ClampMin = "-36.0", ClampMax = "36.0", UIMin = "-12.0", UIMax = "12.0"))
diff --git a/Engine/Plugins/Runtime/Metasound/Source/MetasoundEngine/Public/MetasoundAssetSubsystem.h b/Engine/Plugins/Runtime/Metasound/Source/MetasoundEngine/Public/MetasoundAssetSubsystem.h
--- a/Engine/Plugins/Runtime/Metasound/Source/MetasoundEngine/Public/MetasoundAssetSubsystem.h
+++ b/Engine/Plugins/Runtime/Metasound/Source/MetasoundEngine/Public/MetasoundAssetSubsystem.h
@@ -166,7 +166,7 @@ PRAGMA_ENABLE_DEPRECATION_WARNINGS
 		UE_API virtual void AddReferencedObjects(FReferenceCollector& Collector) override;
 		virtual FString GetReferencerName() const override { return TEXT("FMetaSoundAssetManager"); }
 
-		Frontend::FMetaSoundAssetRegistrationOptions GetRegistrationOptions() const;
+		UE_INTERNAL virtual Frontend::FMetaSoundAssetRegistrationOptions GetRegistrationOptions() const override;
 
 	private:
 		FMetaSoundAssetKey AddOrUpdateFromObjectInternal(const UObject& InObject);
diff --git a/Engine/Plugins/Runtime/Metasound/Source/MetasoundFrontend/Private/DocumentConfigurations/MetasoundFrontendPresetDocumentConfiguration.cpp b/Engine/Plugins/Runtime/Metasound/Source/MetasoundFrontend/Private/DocumentConfigurations/MetasoundFrontendPresetDocumentConfiguration.cpp
--- a/Engine/Plugins/Runtime/Metasound/Source/MetasoundFrontend/Private/DocumentConfigurations/MetasoundFrontendPresetDocumentConfiguration.cpp
+++ b/Engine/Plugins/Runtime/Metasound/Source/MetasoundFrontend/Private/DocumentConfigurations/MetasoundFrontendPresetDocumentConfiguration.cpp
@@ -2,6 +2,11 @@
 #include "DocumentConfigurations/MetasoundFrontendPresetDocumentConfiguration.h"
 
 #include "Algo/AnyOf.h"
+
+#if WITH_EDITOR
+#include "MetasoundAssetManager.h"
+#endif // WITH_EDITOR
+
 #include "MetasoundDataTypeGetTraits.h"
 #include "MetasoundDocumentInterface.h"
 #include "MetasoundFrontendDocument.h"
@@ -9,7 +14,6 @@
 #include "MetasoundFrontendDocumentBuilderRegistry.h"
 #include "MetasoundTrigger.h"
 
-
 #define LOCTEXT_NAMESPACE "MetaSound"
 
 namespace Metasound::Frontend::PresetDocConfigPrivate
@@ -121,6 +125,18 @@ void FMetaSoundFrontendPresetDocumentConfiguration::OnAssetInitialized(TArray<UO
 
 void FMetaSoundFrontendPresetDocumentConfiguration::OnPropertyChanged(const FPropertyChangedEvent& InEvent, FMetaSoundFrontendDocumentBuilder& OutBuilder)
 {
+	using namespace Metasound::Frontend;
+
+	if (InEvent.GetPropertyName() == GET_MEMBER_NAME_CHECKED(FMetaSoundFrontendPresetDocumentConfiguration, Parent))
+	{
+		if (Parent.GetObject() != nullptr)
+		{
+			if (FMetasoundAssetBase* MetaSoundAsset = IMetaSoundAssetManager::GetChecked().GetAsAsset(*Parent.GetObject()))
+			{
+				MetaSoundAsset->UpdateAndRegisterForExecution(IMetaSoundAssetManager::GetChecked().GetRegistrationOptions());
+			}
+		}
+	}
 	ConfigureDocument(OutBuilder, false);
 }
 
diff --git a/Engine/Plugins/Runtime/Metasound/Source/MetasoundFrontend/Public/MetasoundAssetManager.h b/Engine/Plugins/Runtime/Metasound/Source/MetasoundFrontend/Public/MetasoundAssetManager.h
--- a/Engine/Plugins/Runtime/Metasound/Source/MetasoundFrontend/Public/MetasoundAssetManager.h
+++ b/Engine/Plugins/Runtime/Metasound/Source/MetasoundFrontend/Public/MetasoundAssetManager.h
@@ -292,6 +292,9 @@ PRAGMA_ENABLE_DEPRECATION_WARNINGS
 		virtual bool GetReferencedPresetHierarchy(FMetasoundAssetBase& InAsset, TArray<FMetasoundAssetBase*>& OutReferencedAssets) const = 0;
 #endif // WITH_EDITOR
 
+		// Returns default registration options for MetaSound assets
+		UE_INTERNAL virtual Frontend::FMetaSoundAssetRegistrationOptions GetRegistrationOptions() const = 0;
+
 		// Returns whether or not the given ClassMetadata is both valid asset ClassMetadata and is registered.
 		virtual bool IsAssetClass(const FMetasoundFrontendClassMetadata& ClassMetadata) const = 0;
 
