Fix static assert that is not jiving with certain builds
#rnx

#Synced-CL 48836534

diff --git a/Engine/Plugins/Runtime/Metasound/Source/MetasoundFrontend/Public/MetasoundFrontendLiteral.h b/Engine/Plugins/Runtime/Metasound/Source/MetasoundFrontend/Public/MetasoundFrontendLiteral.h
--- a/Engine/Plugins/Runtime/Metasound/Source/MetasoundFrontend/Public/MetasoundFrontendLiteral.h
+++ b/Engine/Plugins/Runtime/Metasound/Source/MetasoundFrontend/Public/MetasoundFrontendLiteral.h
@@ -199,7 +199,7 @@ namespace Metasound::Frontend
 		}
 		else
 		{
-			static_assert(false, "MakeFrontendLiteral: Invalid num arguments");
+			static_assert(sizeof...(InArgs) <= 1, "MakeFrontendLiteral: Invalid num arguments");
 		}
 
 		return Value;
