Remove visibility restriction added to wrong MetaSound Editor doc config section (config vs. graph toggle). Follow-up CL in review (awaiting editor tools approval for minor underlying button block change) to show/hide the toggle option.
#rnx
#Synced-CL 48836534

diff --git a/Engine/Plugins/Runtime/Metasound/Source/MetasoundEditor/Private/MetasoundEditor.cpp b/Engine/Plugins/Runtime/Metasound/Source/MetasoundEditor/Private/MetasoundEditor.cpp
--- a/Engine/Plugins/Runtime/Metasound/Source/MetasoundEditor/Private/MetasoundEditor.cpp
+++ b/Engine/Plugins/Runtime/Metasound/Source/MetasoundEditor/Private/MetasoundEditor.cpp
@@ -2347,29 +2347,7 @@ namespace Metasound
 							FOnGetContent::CreateSP(this, &FEditor::CreateConfigurationMenuOptions),
 							LOCTEXT("ConfigurationSettingsMenu", "Configuration"),
 							LOCTEXT("ConfigurationSettingsMenu_Tooltip", "Settings related to procedurally-generated MetaSound Configurations"),
-							GetExportStatusImage(),
-							false /* bInSimpleComboBox*/,
-							NAME_None /* InTutorialHighlightName */,
-							TAttribute<EVisibility>::CreateSPLambda(AsShared(), [this]()
-							{
-								if (Builder.IsValid())
-								{
-									const FMetaSoundFrontendDocumentBuilder& DocBuilder = Builder->GetConstBuilder();
-									if (DocBuilder.IsValid())
-									{
-										using namespace Frontend;
-
-										//  Show option if override is required.  If the config requires it always visible, hide.
-										TConstStructView<FDocumentConfiguration> Config = DocBuilder.GetConstDocumentConfiguration();
-										if (Config.IsValid() && !Config.GetPtr()->IsGraphEditorVisible())
-										{
-											return EVisibility::Visible;
-										}
-									}
-								}
-
-								return EVisibility::Collapsed;
-							})
+							GetExportStatusImage()
 						);
  					}
  					ToolbarBuilder.EndSection();
