- Fix-up MetaSound Editor callsites looking at old bIsPreset field in document model
- Move property migration on MetaSound member metadata off PostLoad to specialized virtual called during versioning to avoid mutating fields during async versioning

#rnx
#Synced-CL 48890236

diff --git a/Engine/Plugins/Runtime/Metasound/Source/MetasoundEditor/Private/MetasoundEditor.cpp b/Engine/Plugins/Runtime/Metasound/Source/MetasoundEditor/Private/MetasoundEditor.cpp
--- a/Engine/Plugins/Runtime/Metasound/Source/MetasoundEditor/Private/MetasoundEditor.cpp
+++ b/Engine/Plugins/Runtime/Metasound/Source/MetasoundEditor/Private/MetasoundEditor.cpp
@@ -1485,6 +1485,8 @@ namespace Metasound
 
 		const FSlateBrush* FEditor::GetDefaultTabIcon() const
 		{
+			using namespace Frontend;
+
 			FString IconName = TEXT("MetasoundEditor");
 			if (IsPlaying())
 			{
@@ -1503,7 +1505,8 @@ namespace Metasound
 				}
 
 				TScriptInterface<const IMetaSoundDocumentInterface> DocInterface = MetaSound;
-				if (DocInterface->GetConstDocument().RootGraph.PresetOptions.bIsPreset)
+				const UStruct* Struct = DocInterface->GetConstDocument().Configuration.GetScriptStruct();
+				if (Struct && Struct->IsChildOf(FPresetDocumentConfiguration::StaticStruct()))
 				{
 					IconName += TEXT(".Preset");
 				}
@@ -1852,14 +1855,14 @@ namespace Metasound
 			}
 
 			const FMetasoundFrontendGraphClass& ThisGraphClass = DocBuilder.GetConstDocumentChecked().RootGraph;
-			const FNodeRegistryKey ThisRegistryKey(ThisGraphClass);
+			const FNodeClassRegistryKey ThisRegistryKey(ThisGraphClass);
 
 			// Referenced assets
 			FMetasoundAssetBase& ThisAsset = DocBuilder.GetMetasoundAsset();
 			TArray<FMetasoundAssetBase*> ReferencedAssets;
 			IMetaSoundAssetManager::GetChecked().GetReferencedPresetHierarchy(ThisAsset, ReferencedAssets);
 
-			TArray<FNodeRegistryKey> ReferencedRegistryKeys;
+			TSet<FNodeClassRegistryKey> ReferencedRegistryKeys;
 			Algo::Transform(ReferencedAssets, ReferencedRegistryKeys, [](const FMetasoundAssetBase* ReferencedAsset)
 			{
 				if (ensure(ReferencedAsset))
@@ -1867,9 +1870,9 @@ namespace Metasound
 					TScriptInterface<const IMetaSoundDocumentInterface> RefDocInterface = ReferencedAsset->GetOwningAsset();
 					const FMetasoundFrontendDocument& ReferencedDocument = RefDocInterface->GetConstDocument();
 					const FMetasoundFrontendGraphClass& ReferencedGraphClass = ReferencedDocument.RootGraph;
-					return FNodeRegistryKey(ReferencedGraphClass);
+					return FNodeClassRegistryKey(ReferencedGraphClass);
 				}
-				return FNodeRegistryKey();
+				return FNodeClassRegistryKey();
 			});
 
 			// Filter whether the widget supports this MetaSound
diff --git a/Engine/Plugins/Runtime/Metasound/Source/MetasoundEditor/Private/MetasoundEditorGraphMemberDefaults.cpp b/Engine/Plugins/Runtime/Metasound/Source/MetasoundEditor/Private/MetasoundEditorGraphMemberDefaults.cpp
--- a/Engine/Plugins/Runtime/Metasound/Source/MetasoundEditor/Private/MetasoundEditorGraphMemberDefaults.cpp
+++ b/Engine/Plugins/Runtime/Metasound/Source/MetasoundEditor/Private/MetasoundEditorGraphMemberDefaults.cpp
@@ -6,6 +6,7 @@
 #include "Algo/Transform.h"
 #include "AudioDefines.h"
 #include "AudioParameterControllerInterface.h"
+#include "DocumentConfigurations/MetasoundFrontendPresetDocumentConfiguration.h"
 #include "EdGraph/EdGraphNode.h"
 #include "GraphEditorSettings.h"
 #include "MetasoundDataReference.h"
@@ -614,26 +615,38 @@ void UMetasoundEditorGraphMemberDefaultFloat::InitDefault(const FGuid& InPageID)
 	MemberDefaultsPrivate::InitDefault<float>(*this, InPageID);
 }
 
-void UMetasoundEditorGraphMemberDefaultFloat::PostLoad()
+#if WITH_EDITORONLY_DATA
+bool UMetasoundEditorGraphMemberDefaultFloat::MigrateProperties()
 {
 	using namespace Metasound::Editor;
 	using namespace Metasound::Frontend;
 
-	Super::PostLoad();
+	bool bMigrated = false;
 
 	FMetasoundAssetBase& MetaSound = FGraphBuilder::GetOutermostMetaSoundChecked(*this);
 	const FMetasoundFrontendDocument& Document = MetaSound.GetConstDocumentChecked();
-	const bool bIsPreset = Document.RootGraph.PresetOptions.bIsPreset;
 
+	// Property migration occurs before document versioning, so check if either is true to
+	// ensure widget type is cleared irrespective of document version.
+	bool bIsPreset = Document.RootGraph.PresetOptions.bIsPreset;
+	if (!bIsPreset)
+	{
+		const UStruct* Struct = Document.Configuration.GetScriptStruct();
+		bIsPreset = Struct && Struct->IsChildOf(FPresetDocumentConfiguration::StaticStruct());
+	}
 	// Set default input widget type from settings. Widgets do not apply to presets
 	if (bIsPreset)
 	{
 		if (WidgetType != EMetasoundMemberDefaultWidget::None)
 		{
 			WidgetType = EMetasoundMemberDefaultWidget::None;
+			bMigrated = true;
 		}
 	}
+
+	return bMigrated;
 }
+#endif // WITH_EDITORONLY_DATA
 
 void UMetasoundEditorGraphMemberDefaultFloat::Initialize()
 {
diff --git a/Engine/Plugins/Runtime/Metasound/Source/MetasoundEditor/Public/MetasoundEditorGraphMemberDefaults.h b/Engine/Plugins/Runtime/Metasound/Source/MetasoundEditor/Public/MetasoundEditorGraphMemberDefaults.h
--- a/Engine/Plugins/Runtime/Metasound/Source/MetasoundEditor/Public/MetasoundEditorGraphMemberDefaults.h
+++ b/Engine/Plugins/Runtime/Metasound/Source/MetasoundEditor/Public/MetasoundEditorGraphMemberDefaults.h
@@ -373,7 +373,11 @@ public:
 	virtual bool TryFindDefault(FMetasoundFrontendLiteral& OutLiteral, const FGuid* InPageID = nullptr) const override;
 	virtual void UpdatePreviewInstance(const Metasound::FVertexName& InParameterName, TScriptInterface<IAudioParameterControllerInterface>& InParameterInterface) const override;
 	virtual void PostEditChangeChainProperty(FPropertyChangedChainEvent& InPropertyChangedEvent) override;
-	virtual void PostLoad() override;
+
+#if WITH_EDITORONLY_DATA
+	virtual bool MigrateProperties() override;
+#endif // WITH_EDITORONLY_DATA
+
 	// Set Range to reasonable limit given current Default value
 	void SetInitialRange();
 
diff --git a/Engine/Plugins/Runtime/Metasound/Source/MetasoundEngine/Private/MetasoundAssetSubsystem.cpp b/Engine/Plugins/Runtime/Metasound/Source/MetasoundEngine/Private/MetasoundAssetSubsystem.cpp
--- a/Engine/Plugins/Runtime/Metasound/Source/MetasoundEngine/Private/MetasoundAssetSubsystem.cpp
+++ b/Engine/Plugins/Runtime/Metasound/Source/MetasoundEngine/Private/MetasoundAssetSubsystem.cpp
@@ -5,6 +5,7 @@
 #include "Algo/Transform.h"
 #include "AssetRegistry/AssetRegistryModule.h"
 #include "Async/Async.h"
+#include "DocumentConfigurations/MetasoundFrontendPresetDocumentConfiguration.h"
 #include "Engine/AssetManager.h"
 #include "HAL/CriticalSection.h"
 #include "Metasound.h"
@@ -763,11 +764,15 @@ PRAGMA_ENABLE_DEPRECATION_WARNINGS
 	
 	bool FMetaSoundAssetManager::GetReferencedPresetHierarchy(FMetasoundAssetBase& InAsset, TArray<FMetasoundAssetBase*>& OutReferencedAssets) const
 	{
+		using namespace Frontend;
+
 		OutReferencedAssets.Reset();
 
 		TScriptInterface<const IMetaSoundDocumentInterface> DocInterface = InAsset.GetOwningAsset();
 		check(DocInterface.GetObject());
-		if (!DocInterface->GetConstDocument().RootGraph.PresetOptions.bIsPreset)
+
+		const UStruct* Struct = DocInterface->GetConstDocument().Configuration.GetScriptStruct();
+		if (!Struct || !Struct->IsChildOf(FPresetDocumentConfiguration::StaticStruct()))
 		{
 			return false;
 		}
@@ -783,13 +788,15 @@ PRAGMA_ENABLE_DEPRECATION_WARNINGS
 
 	bool FMetaSoundAssetManager::GetReferencedPresetHierarchyInternal(FMetasoundAssetBase& InAsset, TArray<FMetasoundAssetBase*>& OutReferencedAssets) const
 	{
+		using namespace Frontend;
+
 		OutReferencedAssets.Add(&InAsset);
 
 		TScriptInterface<const IMetaSoundDocumentInterface> DocInterface = InAsset.GetOwningAsset();
-		const bool bIsPreset = DocInterface->GetConstDocument().RootGraph.PresetOptions.bIsPreset;
 
 		// Base case (found the first non preset referenced asset)
-		if (!bIsPreset)
+		const UStruct* Struct = DocInterface->GetConstDocument().Configuration.GetScriptStruct();
+		if (!Struct || !Struct->IsChildOf(FPresetDocumentConfiguration::StaticStruct()))
 		{
 			return true;
 		}
diff --git a/Engine/Plugins/Runtime/Metasound/Source/MetasoundFrontend/Private/MetasoundFrontendDocument.cpp b/Engine/Plugins/Runtime/Metasound/Source/MetasoundFrontend/Private/MetasoundFrontendDocument.cpp
--- a/Engine/Plugins/Runtime/Metasound/Source/MetasoundFrontend/Private/MetasoundFrontendDocument.cpp
+++ b/Engine/Plugins/Runtime/Metasound/Source/MetasoundFrontend/Private/MetasoundFrontendDocument.cpp
@@ -1713,6 +1713,17 @@ bool FMetasoundFrontendDocument::MigrateProperties()
 		InterfaceVersions.Reset();
 	}
 
+	for (TPair<FGuid, TObjectPtr<UMetaSoundFrontendMemberMetadata>>& Pair : Metadata.MemberMetadata)
+	{
+		if (UMetaSoundFrontendMemberMetadata* MemberMetadata = Pair.Value.Get())
+		{
+			const bool bMigratedMetadata = MemberMetadata->MigrateProperties();
+			if (bMigratedMetadata)
+			{
+				bMigrated = true;
+			}
+		}
+	}
 	return bMigrated;
 }
 #endif // WITH_EDITORONLY_DATA
diff --git a/Engine/Plugins/Runtime/Metasound/Source/MetasoundFrontend/Private/MetasoundFrontendDocumentVersioning.cpp b/Engine/Plugins/Runtime/Metasound/Source/MetasoundFrontend/Private/MetasoundFrontendDocumentVersioning.cpp
--- a/Engine/Plugins/Runtime/Metasound/Source/MetasoundFrontend/Private/MetasoundFrontendDocumentVersioning.cpp
+++ b/Engine/Plugins/Runtime/Metasound/Source/MetasoundFrontend/Private/MetasoundFrontendDocumentVersioning.cpp
@@ -729,7 +729,11 @@ PRAGMA_ENABLE_DEPRECATION_WARNINGS
 			{
 				FMetasoundFrontendGraphClass Class = InDocument->GetRootGraphClass();
 				FMetasoundFrontendGraphClassPresetOptions PresetOptions = Class.PresetOptions;
+
+PRAGMA_DISABLE_DEPRECATION_WARNINGS
 				Class.PresetOptions.bIsPreset = Class.Metadata.GetAndClearAutoUpdateManagesInterface_Deprecated();
+PRAGMA_ENABLE_DEPRECATION_WARNINGS
+
 				InDocument->SetRootGraphClass(MoveTemp(Class));
 			}
 		};
diff --git a/Engine/Plugins/Runtime/Metasound/Source/MetasoundFrontend/Public/MetasoundFrontendDocument.h b/Engine/Plugins/Runtime/Metasound/Source/MetasoundFrontend/Public/MetasoundFrontendDocument.h
--- a/Engine/Plugins/Runtime/Metasound/Source/MetasoundFrontend/Public/MetasoundFrontendDocument.h
+++ b/Engine/Plugins/Runtime/Metasound/Source/MetasoundFrontend/Public/MetasoundFrontendDocument.h
@@ -2006,6 +2006,10 @@ public:
 #if WITH_EDITORONLY_DATA
 	UPROPERTY()
 	FGuid MemberID;
+
+	// Enables migrating properties during asynchronous versioning safely. Favor using this
+	// over versioning properties on 'PostLoad' to avoid race conditions.
+	virtual bool MigrateProperties() { return false; }
 #endif // WITH_EDITORONLY_DATA
 };
 
