MetaSound graph editor asset drag drop improvements
When dragging compatible assets into a MetaSound graph, node palette menu is opened and users can
- Add Graph Input (with object set)
- Add Node (with object set on the node input)
- Add Node with connected Graph Input (with object set on the graph input and graph input node connected to the selected node)
Dragging multiple assets onto a pin creates a new connected input with the assets

#Synced-CL 48903607

diff --git a/Engine/Plugins/Runtime/Metasound/Source/MetasoundEditor/Private/MetasoundEditor.cpp b/Engine/Plugins/Runtime/Metasound/Source/MetasoundEditor/Private/MetasoundEditor.cpp
--- a/Engine/Plugins/Runtime/Metasound/Source/MetasoundEditor/Private/MetasoundEditor.cpp
+++ b/Engine/Plugins/Runtime/Metasound/Source/MetasoundEditor/Private/MetasoundEditor.cpp
@@ -5968,6 +5968,21 @@ PRAGMA_ENABLE_DEPRECATION_WARNINGS
 				FGraphBuilder::RegisterGraphWithFrontend(ParentMetasound, MoveTemp(RegOptions));
 			}
 		}
+
+		void FEditor::SetDroppedObjectsInfo(FDroppedObjectsInfo InInfo)
+		{
+			DroppedObjectsInfo = TOptional<FDroppedObjectsInfo>(MoveTemp(InInfo));
+		}
+
+		const TOptional<FDroppedObjectsInfo>& FEditor::GetDroppedObjectsInfo() const
+		{
+			return DroppedObjectsInfo;
+		}
+
+		void FEditor::ClearDroppedObjectsInfo()
+		{
+			DroppedObjectsInfo.Reset();
+		}
 	}
 }
 
diff --git a/Engine/Plugins/Runtime/Metasound/Source/MetasoundEditor/Private/MetasoundEditor.h b/Engine/Plugins/Runtime/Metasound/Source/MetasoundEditor/Private/MetasoundEditor.h
--- a/Engine/Plugins/Runtime/Metasound/Source/MetasoundEditor/Private/MetasoundEditor.h
+++ b/Engine/Plugins/Runtime/Metasound/Source/MetasoundEditor/Private/MetasoundEditor.h
@@ -119,7 +119,6 @@ namespace Metasound::Editor
 		COUNT
 	};
 
-
 	class FEditor : public IMetasoundEditor, public FNotifyHook, public FEditorUndoClient, public FTickableEditorObject
 	{
 	public:
@@ -257,6 +256,11 @@ namespace Metasound::Editor
 
 		void SetDelayedRename();
 
+		void SetDroppedObjectsInfo(FDroppedObjectsInfo InInfo);
+		const TOptional<FDroppedObjectsInfo>& GetDroppedObjectsInfo() const;
+
+		void ClearDroppedObjectsInfo();
+
 	protected:
 		// Callbacks for action tree
 		bool CanRenameOnActionNode(TWeakPtr<FGraphActionNode> InSelectedNode) const;
@@ -650,5 +654,8 @@ namespace Metasound::Editor
 		};
 
 		TSharedPtr<FDocumentListener> DocListener;
+
+		// Objects that are currently being dropped onto the editor
+		TOptional<FDroppedObjectsInfo> DroppedObjectsInfo;
 	};
 } // namespace Metasound::Editor
diff --git a/Engine/Plugins/Runtime/Metasound/Source/MetasoundEditor/Private/MetasoundEditorGraphBuilder.cpp b/Engine/Plugins/Runtime/Metasound/Source/MetasoundEditor/Private/MetasoundEditorGraphBuilder.cpp
--- a/Engine/Plugins/Runtime/Metasound/Source/MetasoundEditor/Private/MetasoundEditorGraphBuilder.cpp
+++ b/Engine/Plugins/Runtime/Metasound/Source/MetasoundEditor/Private/MetasoundEditorGraphBuilder.cpp
@@ -32,6 +32,7 @@
 #include "MetasoundFrontendDataTypeRegistry.h"
 #include "MetasoundFrontendDocumentAccessPtr.h"
 #include "MetasoundFrontendDocumentBuilder.h"
+#include "MetasoundFrontendDocumentIdGenerator.h"
 #include "MetasoundFrontendPages.h"
 #include "MetasoundFrontendQuery.h"
 #include "MetasoundFrontendQuerySteps.h"
@@ -1291,13 +1292,16 @@ PRAGMA_ENABLE_DEPRECATION_WARNINGS
 		FMetasoundFrontendClassInput FGraphBuilder::CreateUniqueClassInput(UObject& InMetaSound, const FCreateNodeVertexParams& InParams, const TArray<FMetasoundFrontendClassInputDefault>& InDefaultLiterals, const FName* InNameBase)
 		{
 			using namespace Engine;
+			using namespace Frontend;
 
 			FMetasoundFrontendClassInput ClassInput;
 			UMetaSoundBuilderBase& Builder = FDocumentBuilderRegistry::GetChecked().FindOrBeginBuilding(InMetaSound);
-			ClassInput.Name = Builder.GetConstBuilder().GenerateUniqueInputName(InNameBase ? InNameBase->ToString() : TEXT("Input"));
+			const FMetaSoundFrontendDocumentBuilder& DocBuilder = Builder.GetConstBuilder();
+			const FMetasoundFrontendDocument& Document = DocBuilder.GetConstDocumentChecked();
+			ClassInput.Name = DocBuilder.GenerateUniqueInputName(InNameBase ? InNameBase->ToString() : TEXT("Input"));
 			ClassInput.TypeName = InParams.DataType;
-			ClassInput.VertexID = FGuid::NewGuid();
-			ClassInput.NodeID = FGuid::NewGuid();
+			ClassInput.VertexID = FDocumentIDGenerator::Get().CreateVertexID(Document);
+			ClassInput.NodeID = FDocumentIDGenerator::Get().CreateNodeID(Document);
 
 			// Can be unset if attempting to mirror parameters from a reroute, so default to reference
 			ClassInput.AccessType = InParams.AccessType == EMetasoundFrontendVertexAccessType::Unset ? EMetasoundFrontendVertexAccessType::Reference : InParams.AccessType;
diff --git a/Engine/Plugins/Runtime/Metasound/Source/MetasoundEditor/Private/MetasoundEditorGraphSchema.cpp b/Engine/Plugins/Runtime/Metasound/Source/MetasoundEditor/Private/MetasoundEditorGraphSchema.cpp
--- a/Engine/Plugins/Runtime/Metasound/Source/MetasoundEditor/Private/MetasoundEditorGraphSchema.cpp
+++ b/Engine/Plugins/Runtime/Metasound/Source/MetasoundEditor/Private/MetasoundEditorGraphSchema.cpp
@@ -10,6 +10,7 @@
 #include "Editor/EditorEngine.h"
 #include "Framework/Commands/GenericCommands.h"
 #include "Framework/MultiBox/MultiBoxBuilder.h"
+#include "Framework/Application/SlateApplication.h"
 #include "GraphEditor.h"
 #include "GraphEditorActions.h"
 #include "HAL/IConsoleManager.h"
@@ -148,37 +149,51 @@ namespace Metasound::Editor
 
 		bool DataTypeSupportsAssetTypes(const Metasound::Frontend::FDataTypeRegistryInfo& InRegistryInfo, const TArray<FAssetData>& InAssets)
 		{
-			if (InRegistryInfo.PreferredLiteralType != Metasound::ELiteralType::UObjectProxy)
+			using namespace Metasound::Frontend;
+
+			if (!InRegistryInfo.IsDataTypeProxyParsable())
 			{
 				return false;
 			}
-
-			const IMetasoundEditorModule& EditorModule = FModuleManager::GetModuleChecked<IMetasoundEditorModule>("MetaSoundEditor");
-			return Algo::AnyOf(InAssets, [&EditorModule, &InRegistryInfo](const FAssetData& Asset)
+			const IDataTypeRegistry& DataTypeRegistry = IDataTypeRegistry::Get();
+			FDataTypeRegistryInfo TestDataType;
+			FName DataTypeName = InRegistryInfo.DataTypeName;
+			if (InRegistryInfo.bIsArrayType)
 			{
-				if (InRegistryInfo.ProxyGeneratorClass)
+				FString DataTypeString = DataTypeName.ToString();
+				if (DataTypeString.EndsWith(METASOUND_DATA_TYPE_NAME_ARRAY_TYPE_SPECIFIER))
 				{
-					if (const UClass* Class = Asset.GetClass())
+					DataTypeString.RemoveFromEnd(METASOUND_DATA_TYPE_NAME_ARRAY_TYPE_SPECIFIER);
+					DataTypeName = FName(DataTypeString);
+					if (!DataTypeRegistry.GetDataTypeInfo(DataTypeName, TestDataType))
 					{
-PRAGMA_DISABLE_DEPRECATION_WARNINGS
-						if (EditorModule.IsExplicitProxyClass(*InRegistryInfo.ProxyGeneratorClass))
-						{
-							return Class == InRegistryInfo.ProxyGeneratorClass;
-						}
-PRAGMA_ENABLE_DEPRECATION_WARNINGS
-						if (InRegistryInfo.bIsExplicit)
-						{
-							return Class == InRegistryInfo.ProxyGeneratorClass;
-						}
-						else
-						{
-							return Class->IsChildOf(InRegistryInfo.ProxyGeneratorClass);
-						}
+						return false;
 					}
 				}
-
+				else
+				{
+					return false;
+				}
+			}
+			// In type is not array type but InAssets is a multiasset array
+			else if (InAssets.Num() > 1)
+			{
 				return false;
-			});
+			}
+
+			const bool bSuccess = Algo::AllOf(InAssets, 
+				[&DataTypeRegistry, &TestDataType, &DataTypeName](const FAssetData& Asset) 
+				{
+					if (UClass* Class = Asset.GetClass())
+					{
+						if (DataTypeRegistry.GetDataTypeInfo(Class, TestDataType))
+						{
+							return TestDataType.DataTypeName == DataTypeName;
+						}
+					}
+					return false;
+				});
+			return bSuccess;
 		}
 
 		// Connects to first pin with the same DataType
@@ -315,6 +330,28 @@ PRAGMA_ENABLE_DEPRECATION_WARNINGS
 			}
 		}
 
+		void PostAddNode(UMetasoundEditorGraphNode* NewGraphNode, UObject& ParentMetasound, UMetasoundEditorGraph& MetaSoundGraph, const FVector2f& InLocation, bool bSelectNewNode = true)
+		{
+			using namespace Metasound::Frontend;
+
+			NewGraphNode->UpdateFrontendNodeLocation(FDeprecateSlateVector2D(InLocation));
+			NewGraphNode->SyncLocationFromFrontendNode();
+
+			// Ensure referencing graphs are updated
+			FMetaSoundAssetRegistrationOptions RegOptions = FGraphBuilder::GetDefaultRegistrationOptions();
+			RegOptions.bIgnoreIfLiveAuditioning = true;
+			FGraphBuilder::RegisterGraphWithFrontend(ParentMetasound, MoveTemp(RegOptions));
+
+			if (bSelectNewNode)
+			{
+				SchemaPrivate::SelectNodeInEditorForRename(MetaSoundGraph, *NewGraphNode);
+			}
+			else
+			{
+				SchemaPrivate::SelectNodeInEditor(MetaSoundGraph, *NewGraphNode);
+			}
+		}
+
 		UEdGraphNode* PromoteToVariable(const FName BaseName, UEdGraphPin& FromPin, const FName DataType, const FMetasoundFrontendClass& InVariableClass, const FVector2f& InLocation, bool bSelectNode)
 		{
 			using namespace Frontend;
@@ -359,9 +396,6 @@ PRAGMA_ENABLE_DEPRECATION_WARNINGS
 						UMetasoundEditorGraphVariableNode* NewGraphNode = FGraphBuilder::AddVariableNode(ParentMetaSound, NewVariableNode->GetID());
 						if (ensure(NewGraphNode))
 						{
-							NewGraphNode->UpdateFrontendNodeLocation(FDeprecateSlateVector2D(InLocation));
-							NewGraphNode->SyncLocationFromFrontendNode();
-
 							// Set the literal using the new value if connecting to an input
 							if (FromPin.Direction == EGPD_Input)
 							{
@@ -379,15 +413,7 @@ PRAGMA_ENABLE_DEPRECATION_WARNINGS
 							UEdGraphNode* EdGraphNode = CastChecked<UEdGraphNode>(NewGraphNode);
 							if (ensure(SchemaPrivate::TryConnectNewNodeToMatchingDataTypePin(*EdGraphNode, &FromPin)))
 							{
-								FGraphBuilder::RegisterGraphWithFrontend(ParentMetaSound);
-								if (bSelectNode)
-								{
-									SchemaPrivate::SelectNodeInEditorForRename(*MetaSoundGraph, *NewGraphNode);
-								}
-								else
-								{
-									SchemaPrivate::SelectNodeInEditor(*MetaSoundGraph, *NewGraphNode);
-								}
+								SchemaPrivate::PostAddNode(NewGraphNode, ParentMetaSound, *MetaSoundGraph, InLocation, bSelectNode);
 								return EdGraphNode;
 							}
 						}
@@ -603,7 +629,7 @@ PRAGMA_ENABLE_DEPRECATION_WARNINGS
 
 	namespace SchemaUtils
 	{
-		UEdGraphNode* PromoteToInput(UEdGraph* ParentGraph, UEdGraphPin* FromPin, const UE::Slate::FDeprecateVector2DParameter& InLocation, bool bSelectNewNode)
+		UEdGraphNode* PromoteToInput(UEdGraph* ParentGraph, UEdGraphPin* FromPin, const UE::Slate::FDeprecateVector2DParameter& InLocation, bool bSelectNewNode, TOptional<FMetasoundFrontendLiteral> InLiteral)
 		{
 			using namespace Metasound::Editor;
 			using namespace Metasound::Frontend;
@@ -637,7 +663,11 @@ PRAGMA_ENABLE_DEPRECATION_WARNINGS
 			};
 
 			TArray<FMetasoundFrontendClassInputDefault> DefaultLiterals;
-			if (const FMetasoundFrontendVertexLiteral* VertexLiteral = Builder.FindNodeInputDefault(InputVertexHandle.NodeID, InputVertexHandle.VertexID))
+			if (InLiteral.IsSet())
+			{
+				DefaultLiterals = InitDefaultLiterals(*InLiteral);
+			}
+			else if (const FMetasoundFrontendVertexLiteral* VertexLiteral = Builder.FindNodeInputDefault(InputVertexHandle.NodeID, InputVertexHandle.VertexID))
 			{
 				// Since a default page ID requires an associated value and no other nodes on any page would be connected to this new input, use
 				// the default page ID. If the user wants different behavior, when they connect the newly created input on a lower-indexed graph,
@@ -695,24 +725,10 @@ PRAGMA_ENABLE_DEPRECATION_WARNINGS
 					{
 						if (UMetasoundEditorGraphNode* NewGraphNode = FGraphBuilder::AddInputNode(ParentMetasound, NewTemplateNode->GetID()))
 						{
-							NewGraphNode->UpdateFrontendNodeLocation(FDeprecateSlateVector2D(InLocation));
-							NewGraphNode->SyncLocationFromFrontendNode();
 							UEdGraphNode* EdGraphNode = CastChecked<UEdGraphNode>(NewGraphNode);
-
 							if (ensure(SchemaPrivate::TryConnectNewNodeToMatchingDataTypePin(*EdGraphNode, FromPin)))
 							{
-								Frontend::FMetaSoundAssetRegistrationOptions RegOptions = FGraphBuilder::GetDefaultRegistrationOptions();
-								RegOptions.bIgnoreIfLiveAuditioning = true;
-								FGraphBuilder::RegisterGraphWithFrontend(ParentMetasound, MoveTemp(RegOptions));
-								if (bSelectNewNode)
-								{
-									SchemaPrivate::SelectNodeInEditorForRename(*MetasoundGraph, *NewGraphNode);
-								}
-								else
-								{
-									SchemaPrivate::SelectNodeInEditor(*MetasoundGraph, *NewGraphNode);
-								}
-
+								SchemaPrivate::PostAddNode(NewGraphNode, ParentMetasound, *MetasoundGraph, InLocation, bSelectNewNode);
 								return EdGraphNode;
 							}
 						}
@@ -762,22 +778,10 @@ PRAGMA_ENABLE_DEPRECATION_WARNINGS
 				{
 					if (UMetasoundEditorGraphOutputNode* NewGraphNode = FGraphBuilder::AddOutputNode(ParentMetasound, OutputNodeID))
 					{
-						NewGraphNode->UpdateFrontendNodeLocation(FDeprecateSlateVector2D(InLocation));
-						NewGraphNode->SyncLocationFromFrontendNode();
-
 						UEdGraphNode* EdGraphNode = CastChecked<UEdGraphNode>(NewGraphNode);
-
 						if (ensure(SchemaPrivate::TryConnectNewNodeToMatchingDataTypePin(*EdGraphNode, FromPin)))
 						{
-							FGraphBuilder::RegisterGraphWithFrontend(ParentMetasound);
-							if (bSelectNewNode)
-							{
-								SchemaPrivate::SelectNodeInEditorForRename(*MetasoundGraph, *NewGraphNode);
-							}
-							else
-							{
-								SchemaPrivate::SelectNodeInEditor(*MetasoundGraph, *NewGraphNode);
-							}
+							SchemaPrivate::PostAddNode(NewGraphNode, ParentMetasound, *MetasoundGraph, InLocation, bSelectNewNode);
 
 							return EdGraphNode;
 						}
@@ -856,6 +860,68 @@ PRAGMA_ENABLE_DEPRECATION_WARNINGS
 
 			return nullptr;
 		}
+
+		UEdGraphNode* AddNewGraphInput(FName DataType, FMetasoundFrontendLiteral DefaultLiteral, UEdGraph* ParentGraph, const UE::Slate::FDeprecateVector2DParameter& InLocation, bool bSelectNewNode)
+		{
+			using namespace Metasound::Editor;
+			using namespace Metasound::Frontend;
+			using namespace Metasound::Engine;
+
+			UMetasoundEditorGraph* MetasoundGraph = CastChecked<UMetasoundEditorGraph>(ParentGraph);
+			UObject& ParentMetasound = MetasoundGraph->GetMetasoundChecked();
+
+			const FScopedTransaction Transaction(LOCTEXT("AddNewGraphInputNode", "Add New MetaSound Graph Input"));
+			ParentMetasound.Modify();
+			ParentGraph->Modify();
+
+			FMetaSoundFrontendDocumentBuilder& Builder = FDocumentBuilderRegistry::GetChecked().FindOrBeginBuilding(&ParentMetasound);
+
+			const FCreateNodeVertexParams VertexParams =
+			{
+				DataType,
+				EMetasoundFrontendVertexAccessType::Reference
+			};
+
+			const FMetasoundFrontendNode* NewNode = nullptr;
+
+			TArray<FMetasoundFrontendClassInputDefault> DefaultLiterals;
+			const FGuid PageID = Builder.GetBuildPageID();
+			// Create default page literal if not default page
+			if (PageID != DefaultPageID)
+			{
+				FMetasoundFrontendLiteral DefaultPageLiteral;
+				DefaultPageLiteral.SetType(DefaultLiteral.GetType());
+				DefaultLiterals.Emplace(DefaultPageID, MoveTemp(DefaultPageLiteral));
+			}
+			DefaultLiterals.Emplace(PageID, DefaultLiteral);
+
+			// Name may be different than previous name because it may already exist in the graph,
+			// and CreateUniqueClassInput will add an identifier to the end to make it unique
+			FName NewName = TEXT("Input");
+			{
+				FMetasoundFrontendClassInput ClassInput = FGraphBuilder::CreateUniqueClassInput(ParentMetasound, VertexParams, DefaultLiterals, &NewName);
+				NewName = ClassInput.Name;
+				NewNode = Builder.AddGraphInput(MoveTemp(ClassInput));
+			}
+
+			if (ensure(NewNode))
+			{
+				UMetasoundEditorGraphInput* Input = MetasoundGraph->FindOrAddInput(NewNode->GetID());
+				if (ensure(Input))
+				{
+					if (const FMetasoundFrontendNode* NewTemplateNode = FInputNodeTemplate::CreateNode(Builder, NewName))
+					{
+						if (UMetasoundEditorGraphNode* NewGraphNode = FGraphBuilder::AddInputNode(ParentMetasound, NewTemplateNode->GetID()))
+						{
+							UEdGraphNode* EdGraphNode = CastChecked<UEdGraphNode>(NewGraphNode);
+							SchemaPrivate::PostAddNode(NewGraphNode, ParentMetasound, *MetasoundGraph, InLocation, bSelectNewNode);
+							return EdGraphNode;
+						}
+					}
+				}
+			}
+			return nullptr;
+		}
 	} // namespace SchemaUtils
 } // namespace Metasound::Editor
 
@@ -931,6 +997,7 @@ FString FMetasoundGraphSchemaAction_NewNode::GetUniqueStableFavoriteKey() const
 UEdGraphNode* FMetasoundGraphSchemaAction_NewNode::PerformAction(UEdGraph* ParentGraph, UEdGraphPin* FromPin, const FVector2f& Location, bool bSelectNewNode /* = true */)
 {
 	using namespace Metasound::Editor;
+	using namespace Metasound::Engine;
 	using namespace Metasound::Frontend;
 
 	const FScopedTransaction Transaction(LOCTEXT("AddNewNode", "Add New MetaSound Node"));
@@ -947,6 +1014,113 @@ UEdGraphNode* FMetasoundGraphSchemaAction_NewNode::PerformAction(UEdGraph* Paren
 		NewGraphNode->SyncLocationFromFrontendNode();
 		SchemaPrivate::TryConnectNewNodeToMatchingDataTypePin(*NewGraphNode, FromPin);
 		SchemaPrivate::SelectNodeInEditorForRename(*MetaSoundGraph, *NewGraphNode);
+
+		if (DroppedObjectsInfo.IsSet())
+		{
+			UMetaSoundBuilderBase& Builder = FDocumentBuilderRegistry::GetChecked().FindOrBeginBuilding(ParentMetasound);
+			FMetaSoundFrontendDocumentBuilder& DocBuilder = Builder.GetBuilder();
+
+			// Find the vertex to set the dropped objects on
+			// Just use the first matching vertex with the correct type for now 
+			auto FindFirstVertexMatchingType = [&DocBuilder](FGuid NodeID, FName TypeName, bool bFindInput /*If false, find output*/)
+			{
+				const TArray<const FMetasoundFrontendVertex*> Vertices = bFindInput ? 
+					DocBuilder.FindUserModifiableNodeInputs(NodeID) : DocBuilder.FindUserModifiableNodeOutputs(NodeID);
+				const FMetasoundFrontendVertex* const* Vertex = Vertices.FindByPredicate([&TypeName](const FMetasoundFrontendVertex* Vertex)
+				{
+					return Vertex->TypeName == TypeName;
+				});
+				return (Vertex && *Vertex) ? *Vertex : nullptr;
+			};
+			
+			const FMetasoundFrontendVertex* FirstMatchingVertex = FindFirstVertexMatchingType(NewGraphNode->GetNodeID(), DroppedObjectsInfo->TypeName, /*bFindInput*/true);
+			if (!ensure(FirstMatchingVertex))
+			{
+				UE_LOG(LogMetasoundEditor, Warning, TEXT("Failed to find matching vertex on node %s with type %s when setting dropped objects."), *NewGraphNode->GetDisplayName().ToString(), *DroppedObjectsInfo->TypeName.ToString());
+				return NewGraphNode;
+			}
+
+			// Set literal
+			FMetasoundFrontendLiteral DefaultLiteral;
+			if (DroppedObjectsInfo->DroppedObjects.Num() == 1)
+			{
+				// As single object
+				DefaultLiteral.Set(DroppedObjectsInfo->DroppedObjects[0]);
+			}
+			else
+			{
+				// As array
+				DefaultLiteral.Set(DroppedObjectsInfo->DroppedObjects);
+			}
+
+			// Add dropped objects to literal directly
+			// If more than one object, can only set on connected input because 
+			// because in graph array literal editing is not currently supported
+			if (DroppedObjectsInfo->DroppedObjects.Num() == 1 && DroppedObjectsInfo->ModifierKeys.IsShiftDown())
+			{
+				if (FirstMatchingVertex)
+				{
+					DocBuilder.SetNodeInputDefault(NewGraphNode->GetNodeID(), FirstMatchingVertex->VertexID, DefaultLiteral);
+				}
+			}
+			// Add with object(s) set on new connected input
+			else 
+			{
+				// Setup class input
+				TArray<FMetasoundFrontendClassInputDefault> DefaultLiterals;
+				const FGuid PageID = DocBuilder.GetBuildPageID();
+				// Create default page literal if not default page
+				if (PageID != DefaultPageID)
+				{
+					FMetasoundFrontendLiteral DefaultPageLiteral;
+					DefaultPageLiteral.SetType(DefaultLiteral.GetType());
+					DefaultLiterals.Emplace(DefaultPageID, MoveTemp(DefaultPageLiteral));
+				}
+				DefaultLiterals.Emplace(PageID, DefaultLiteral);
+				
+				const FCreateNodeVertexParams VertexParams =
+				{
+					FirstMatchingVertex->TypeName,
+					DocBuilder.GetNodeInputAccessType(NewGraphNode->GetNodeID(), FirstMatchingVertex->VertexID)
+				};
+				// Create class input and input node
+				const FMetasoundFrontendNode* InputNode;
+				FName NewName;
+				{
+					FMetasoundFrontendClassInput ClassInput = FGraphBuilder::CreateUniqueClassInput(ParentMetasound, VertexParams, DefaultLiterals, &FirstMatchingVertex->Name);
+					NewName = ClassInput.Name;
+					InputNode = DocBuilder.AddGraphInput(MoveTemp(ClassInput));
+				}
+
+				if (ensure(InputNode))
+				{
+					UMetasoundEditorGraphInput* Input = MetaSoundGraph->FindOrAddInput(InputNode->GetID());
+					if (ensure(Input))
+					{
+						// Add template node 
+						if (const FMetasoundFrontendNode* NewTemplateNode = FInputNodeTemplate::CreateNode(DocBuilder, NewName))
+						{
+							// Add ed graph input node 
+							if (UMetasoundEditorGraphNode* NewInputNode = FGraphBuilder::AddInputNode(ParentMetasound, NewTemplateNode->GetID()))
+							{
+								// Set node locations
+								NewInputNode->UpdateFrontendNodeLocation(FDeprecateSlateVector2D(Location));
+								NewInputNode->SyncLocationFromFrontendNode();
+
+								// Connect input node to external node
+								UEdGraphNode* EdGraphNode = CastChecked<UEdGraphNode>(NewInputNode);
+								UEdGraphPin* ExternalNodeInputPin = NewGraphNode->FindPin(FirstMatchingVertex->Name.ToString());
+								if (ensure(SchemaPrivate::TryConnectNewNodeToMatchingDataTypePin(*EdGraphNode, ExternalNodeInputPin)))
+								{
+									SchemaPrivate::PostAddNode(NewGraphNode, ParentMetasound, *MetaSoundGraph, FDeprecateSlateVector2D(Location) + DisplayStyle::NodeLayout::DefaultOffsetX, bSelectNewNode);
+									return EdGraphNode;
+								}
+							}
+						}
+					}
+				}
+			}
+		}
 		return NewGraphNode;
 	}
 
@@ -1018,6 +1192,189 @@ PRAGMA_ENABLE_DEPRECATION_WARNINGS
 	return nullptr;
 }
 
+FMetasoundGraphSchemaAction_NewGraphInput::FMetasoundGraphSchemaAction_NewGraphInput(FName InDataType, FMetasoundFrontendLiteral InLiteral)
+	: FMetasoundGraphSchemaAction_NewGraphMember(
+		InDataType, 
+		InLiteral,
+		FText(),
+		LOCTEXT("AddNewGraphInputName", "Add Graph Input"),
+		LOCTEXT("AddNewGraphInputTooltip", "Adds new graph input with given input default value"),
+		Metasound::Editor::EPrimaryContextGroup::Common)
+{
+}
+
+const FSlateBrush* FMetasoundGraphSchemaAction_NewGraphInput::GetIconBrush() const
+{
+	return &Metasound::Editor::Style::GetSlateBrushSafe("MetasoundEditor.Graph.Node.Class.Input");
+}
+
+const FLinearColor& FMetasoundGraphSchemaAction_NewGraphInput::GetIconColor() const
+{
+	if (const UMetasoundEditorSettings* EditorSettings = GetDefault<UMetasoundEditorSettings>())
+	{
+		return EditorSettings->InputNodeTitleColor;
+	}
+
+	return Super::GetIconColor();
+}
+
+UEdGraphNode* FMetasoundGraphSchemaAction_NewGraphInput::PerformAction(UEdGraph* ParentGraph, UEdGraphPin* FromPin, const FVector2f& InLocation, bool bSelectNewNode /* = true */)
+{
+	using namespace Metasound::Editor;
+	return SchemaUtils::AddNewGraphInput(DataType, DefaultLiteral, ParentGraph, InLocation, bSelectNewNode);
+}
+
+FMetasoundGraphSchemaAction_NewGraphOutput::FMetasoundGraphSchemaAction_NewGraphOutput(FName InDataType, FMetasoundFrontendLiteral InLiteral)
+	: FMetasoundGraphSchemaAction_NewGraphMember(
+		InDataType,
+		InLiteral,
+		FText(),
+		LOCTEXT("AddNewGraphOutputName", "Add Graph Output"),
+		LOCTEXT("AddNewGraphOutputTooltip", "Adds new graph output with given default value"),
+		Metasound::Editor::EPrimaryContextGroup::Common)
+{
+}
+
+const FSlateBrush* FMetasoundGraphSchemaAction_NewGraphOutput::GetIconBrush() const
+{
+	return &Metasound::Editor::Style::GetSlateBrushSafe("MetasoundEditor.Graph.Node.Class.Output");
+}
+
+const FLinearColor& FMetasoundGraphSchemaAction_NewGraphOutput::GetIconColor() const
+{
+	if (const UMetasoundEditorSettings* EditorSettings = GetDefault<UMetasoundEditorSettings>())
+	{
+		return EditorSettings->OutputNodeTitleColor;
+	}
+
+	return Super::GetIconColor();
+}
+
+UEdGraphNode* FMetasoundGraphSchemaAction_NewGraphOutput::PerformAction(UEdGraph* ParentGraph, UEdGraphPin* FromPin, const FVector2f& InLocation, bool bSelectNewNode /* = true */)
+{
+	using namespace Metasound::Editor;
+	using namespace Metasound::Frontend;
+	using namespace Metasound::Engine;
+
+	UMetasoundEditorGraph* MetasoundGraph = CastChecked<UMetasoundEditorGraph>(ParentGraph);
+	UObject& ParentMetaSound = MetasoundGraph->GetMetasoundChecked();
+
+	const FScopedTransaction Transaction(LOCTEXT("AddNewGraphOutputNode", "Add New MetaSound Graph Output"));
+	ParentMetaSound.Modify();
+	ParentGraph->Modify();
+
+	FMetaSoundFrontendDocumentBuilder& Builder = FDocumentBuilderRegistry::GetChecked().FindOrBeginBuilding(&ParentMetaSound);
+
+	const FCreateNodeVertexParams VertexParams =
+	{
+		DataType,
+		EMetasoundFrontendVertexAccessType::Reference
+	};
+
+	const FMetasoundFrontendNode* NewNode = nullptr;
+
+	// Name may be different than previous name because it may already exist in the graph,
+	// and CreateUniqueClassOutput will add an identifier to the end to make it unique
+	FName NewName = TEXT("Output");
+	const FMetasoundFrontendClassOutput ClassOutput = FGraphBuilder::CreateUniqueClassOutput(ParentMetaSound, VertexParams, &NewName);
+	NewName = ClassOutput.Name;
+
+	const FMetasoundFrontendNode* OutputNode = Builder.AddGraphOutput(ClassOutput);
+	if (ensure(OutputNode))
+	{
+		const FGuid OutputNodeID = OutputNode->GetID();
+		UMetasoundEditorGraphOutput* Output = MetasoundGraph->FindOrAddOutput(OutputNodeID);
+		if (ensure(Output))
+		{
+			if (UMetasoundEditorGraphOutputNode* NewGraphNode = FGraphBuilder::AddOutputNode(ParentMetaSound, OutputNodeID))
+			{
+				const FMetasoundFrontendVertex* Vertex = Builder.FindNodeInput(OutputNodeID, NewName);
+				if (ensure(Vertex))
+				{
+					Builder.SetNodeInputDefault(OutputNodeID, Vertex->VertexID, DefaultLiteral);
+				}
+
+				SchemaPrivate::PostAddNode(NewGraphNode, ParentMetaSound, *MetasoundGraph, InLocation, bSelectNewNode);
+
+				return NewGraphNode;
+			}
+		}
+	}
+
+	return nullptr;
+}
+
+FMetasoundGraphSchemaAction_NewGraphVariable::FMetasoundGraphSchemaAction_NewGraphVariable(FName InDataType, FMetasoundFrontendLiteral InLiteral)
+	: FMetasoundGraphSchemaAction_NewGraphMember(
+		InDataType,
+		InLiteral,
+		FText(),
+		LOCTEXT("AddNewGraphOutputName", "Add Graph Variable"),
+		LOCTEXT("AddNewGraphOutputTooltip", "Adds new graph variable with given default value"),
+		Metasound::Editor::EPrimaryContextGroup::Common)
+{
+}
+
+const FSlateBrush* FMetasoundGraphSchemaAction_NewGraphVariable::GetIconBrush() const
+{
+	return &Metasound::Editor::Style::GetSlateBrushSafe("MetasoundEditor.Graph.Node.Class.Variable");
+}
+
+const FLinearColor& FMetasoundGraphSchemaAction_NewGraphVariable::GetIconColor() const
+{
+	if (const UMetasoundEditorSettings* EditorSettings = GetDefault<UMetasoundEditorSettings>())
+	{
+		return EditorSettings->VariableNodeTitleColor;
+	}
+
+	return Super::GetIconColor();
+}
+
+UEdGraphNode* FMetasoundGraphSchemaAction_NewGraphVariable::PerformAction(UEdGraph* ParentGraph, UEdGraphPin* FromPin, const FVector2f& InLocation, bool bSelectNewNode /* = true */)
+{
+	using namespace Metasound::Editor;
+	using namespace Metasound::Frontend;
+	using namespace Metasound::Engine;
+
+	UMetasoundEditorGraph* MetasoundGraph = CastChecked<UMetasoundEditorGraph>(ParentGraph);
+	UObject& ParentMetaSound = MetasoundGraph->GetMetasoundChecked();
+
+	const FScopedTransaction Transaction(LOCTEXT("AddNewGraphVariableNode", "Add New MetaSound Graph Variable"));
+	ParentMetaSound.Modify();
+	ParentGraph->Modify();
+
+	FMetaSoundFrontendDocumentBuilder& Builder = FDocumentBuilderRegistry::GetChecked().FindOrBeginBuilding(&ParentMetaSound);
+
+	const FMetasoundFrontendNode* NewNode = nullptr;
+
+	const FName NodeName = Builder.GenerateUniqueVariableName(TEXT("Variable"));
+	FMetasoundFrontendClass VariableClass;
+	if (ensure(IDataTypeRegistry::Get().GetFrontendVariableAccessorClass(DataType, VariableClass)))
+	{
+		const FMetasoundFrontendVariable* FrontendVariable = Builder.AddGraphVariable(NodeName, DataType);
+		if (ensure(FrontendVariable))
+		{
+			UMetasoundEditorGraphVariable* Variable = MetasoundGraph->FindOrAddVariable(FrontendVariable->Name);
+			if (ensure(Variable))
+			{
+				const FMetasoundFrontendNode* NewVariableNode = Builder.AddGraphVariableNode(FrontendVariable->Name, VariableClass.Metadata.GetType());
+				if (ensure(NewVariableNode))
+				{
+					UMetasoundEditorGraphVariableNode* NewGraphNode = FGraphBuilder::AddVariableNode(ParentMetaSound, NewVariableNode->GetID());
+					if (ensure(NewGraphNode))
+					{
+						Builder.SetGraphVariableDefault(NodeName, DefaultLiteral);
+						SchemaPrivate::PostAddNode(NewGraphNode, ParentMetaSound, *MetasoundGraph, InLocation, bSelectNewNode);
+
+						return NewGraphNode;
+					}
+				}
+			}
+		}
+	}
+	return nullptr;
+}
+
 FMetasoundGraphSchemaAction_PromoteToInput::FMetasoundGraphSchemaAction_PromoteToInput()
 	: FMetasoundGraphSchemaAction_NodeWithMultipleOutputs(
 		FText(),
@@ -1539,7 +1896,46 @@ void UMetasoundEditorGraphSchema::GetGraphContextActions(FGraphContextMenuBuilde
 	FActionVertexFilters VertexFilters;
 	FConstGraphHandle GraphHandle = IGraphController::GetInvalidHandle();
 	EMetasoundFrontendVertexAccessType OutputAccessType = EMetasoundFrontendVertexAccessType::Unset;
-	if (const UEdGraphPin* FromPin = ContextMenuBuilder.FromPin)
+
+	TSharedPtr<FEditor> MetasoundEditor = FGraphBuilder::GetEditorForGraph(*ContextMenuBuilder.CurrentGraph);
+	check(MetasoundEditor.IsValid());
+	TOptional<FDroppedObjectsInfo> DroppedObjectsInfo = MetasoundEditor->GetDroppedObjectsInfo();
+	// If dropped assets, set vertex type filter and add new graph member options
+	if (DroppedObjectsInfo.IsSet() && !DroppedObjectsInfo->DroppedObjects.IsEmpty())
+	{
+		// Set filter type
+		VertexFilters.InputTypeName = DroppedObjectsInfo->TypeName;
+		// Value for superset of value and reference 
+		VertexFilters.InputAccessType = EMetasoundFrontendVertexAccessType::Value;
+
+		// Add member actions for drop + shift
+		if (DroppedObjectsInfo->ModifierKeys.IsShiftDown())
+		{
+			// Set default literal from dropped objects
+			FMetasoundFrontendLiteral DefaultLiteral;
+			if (DroppedObjectsInfo->DroppedObjects.Num() == 1)
+			{
+				// As single object
+				DefaultLiteral.Set(DroppedObjectsInfo->DroppedObjects[0]);
+			}
+			else
+			{
+				// As array
+				DefaultLiteral.Set(DroppedObjectsInfo->DroppedObjects);
+			}
+			FGraphActionMenuBuilder& ActionMenuBuilder = static_cast<FGraphActionMenuBuilder&>(ContextMenuBuilder);
+			ActionMenuBuilder.AddAction(MakeShared<FMetasoundGraphSchemaAction_NewGraphInput>(DroppedObjectsInfo->TypeName, DefaultLiteral));
+			ActionMenuBuilder.AddAction(MakeShared<FMetasoundGraphSchemaAction_NewGraphVariable>(DroppedObjectsInfo->TypeName, DefaultLiteral));
+			ActionMenuBuilder.AddAction(MakeShared<FMetasoundGraphSchemaAction_NewGraphOutput>(DroppedObjectsInfo->TypeName, DefaultLiteral));
+			if (DroppedObjectsInfo->DroppedObjects.Num() > 1)
+			{
+				// Function actions are not currently supported because 
+				// there is no editing UI for arrays on node inputs
+				return;
+			}
+		}
+	} 
+	else if (const UEdGraphPin* FromPin = ContextMenuBuilder.FromPin)
 	{
 		if (FromPin->Direction == EGPD_Input)
 		{
@@ -1609,7 +2005,6 @@ void UMetasoundEditorGraphSchema::GetGraphContextActions(FGraphContextMenuBuilde
 	}
 	else
 	{
-		TSharedPtr<FEditor> MetasoundEditor = FGraphBuilder::GetEditorForGraph(*ContextMenuBuilder.CurrentGraph);
 		if (MetasoundEditor.IsValid() && MetasoundEditor->CanPasteNodes())
 		{
 			TSharedPtr<FMetasoundGraphSchemaAction_Paste> NewAction = MakeShared<FMetasoundGraphSchemaAction_Paste>(FText::GetEmpty(), LOCTEXT("PasteHereAction", "Paste here"), FText::GetEmpty(), EPrimaryContextGroup::Common);
@@ -1628,7 +2023,8 @@ void UMetasoundEditorGraphSchema::GetGraphContextActions(FGraphContextMenuBuilde
 		}
 	}
 
-	GetFunctionActions(ContextMenuBuilder, VertexFilters, true /* bShowSelectedActions */, GraphHandle);
+	// Add function actions
+	GetFunctionActions(ContextMenuBuilder, VertexFilters, true /* bShowSelectedActions */, GraphHandle, ContextMenuBuilder.CurrentGraph);
 
 	// Variable and conversion actions are always by reference so are incompatible with constructor outputs 
 	if (OutputAccessType != EMetasoundFrontendVertexAccessType::Value)
@@ -2316,20 +2712,21 @@ void UMetasoundEditorGraphSchema::BreakSinglePinLink(UEdGraphPin* SourcePin, UEd
 	Super::BreakSinglePinLink(SourcePin, TargetPin);
 }
 
-void UMetasoundEditorGraphSchema::GetAssetsGraphHoverMessage(const TArray<FAssetData>& Assets, const UEdGraph* HoverGraph, FString& OutTooltipText, bool& OutOkIcon) const
+bool UMetasoundEditorGraphSchema::ValidateDroppedAssetsOnGraph(const TArray<FAssetData>& Assets, const UEdGraph* HoverGraph, FString* OutTooltipText) const
 {
 	using namespace Metasound;
 	using namespace Metasound::Editor;
+	using namespace Metasound::Frontend;
 
-	OutOkIcon = true;
+	bool bIsValid = true;
+	FString OutReasonText;
 
 	if (!HoverGraph)
 	{
-		OutOkIcon = false;
-		return;
+		return false;
 	}
 
-	OutTooltipText = LOCTEXT("AddMetaSoundToGraph", "Add MetaSound reference to Graph.").ToString();
+	OutReasonText = LOCTEXT("AddNonMetaSoundToGraph", "Add asset reference(s) to graph. \n Drop for node with new connected input \n Shift + Drop for node with asset(s) set on node \n Alt + Drop for new input").ToString();
 
 	const UMetasoundEditorGraph* MetaSoundGraph = CastChecked<UMetasoundEditorGraph>(HoverGraph);
 	const UObject& MetaSound = MetaSoundGraph->GetMetasoundChecked();
@@ -2338,123 +2735,167 @@ void UMetasoundEditorGraphSchema::GetAssetsGraphHoverMessage(const TArray<FAsset
 	check(MetaSoundAsset);
 
 	const TSharedPtr<IAssetReferenceFilter> AssetReferenceFilter = MakeAssetReferenceFilter(MetaSoundGraph);
-
+	// For non MetaSound assets, all assets must be of the same class to be valid, so keep track of type
+	UClass* CurrentNonMetaSoundClass = nullptr;
 	for (const FAssetData& Data : Assets)
 	{
 		UClass* AssetClass = Data.GetClass();
 		check(AssetClass);
-		const bool bIsRegisteredClass = IMetasoundUObjectRegistry::Get().IsRegisteredClass(*AssetClass);
-		if (!bIsRegisteredClass)
-		{
-			OutOkIcon = false;
-			OutTooltipText = LOCTEXT("NotAllMetaSounds", "Asset(s) must all be MetaSounds.").ToString();
-			break;
-		}
+		static const FString MixedTypeErrorText = LOCTEXT("NotAllSameType", "Asset(s) must all be MetaSound types or a single type.").ToString();
 
-		if (UObject* DroppedObject = Data.GetAsset())
+		// MetaSound specific checks (ex. can't cause loop)
+		const bool bIsRegisteredMetaSoundClass = IMetasoundUObjectRegistry::Get().IsRegisteredClass(*AssetClass);
+		if (bIsRegisteredMetaSoundClass)
 		{
-			FMetasoundAssetBase* DroppedMetaSoundAsset = IMetasoundUObjectRegistry::Get().GetObjectAsAssetBase(DroppedObject);
-			if (!DroppedMetaSoundAsset)
+			// Cannot drop MetaSound types mixed with other types 
+			if (CurrentNonMetaSoundClass)
 			{
-				OutOkIcon = false;
-				OutTooltipText = LOCTEXT("NotAllReferenceableMetaSounds", "Asset(s) not (all) MetaSound(s) that can be referenced.").ToString();
+				bIsValid = false;
+				OutReasonText = MixedTypeErrorText;
 				break;
 			}
-
-			const EMetasoundFrontendClassAccessFlags DroppedAccessFlags = DroppedMetaSoundAsset->GetConstDocumentChecked().RootGraph.Metadata.GetAccessFlags();
-			if (!EnumHasAnyFlags(DroppedAccessFlags, EMetasoundFrontendClassAccessFlags::Referenceable))
+			if (UObject* DroppedObject = Data.GetAsset())
 			{
-				OutOkIcon = false;
-				OutTooltipText = LOCTEXT("NotAllAccessibleMetaSounds", "Asset(s) do not contain access flag 'Reference' (see 'Access Flags'). Cannot reference provided MetaSound Asset.").ToString();
-				break;
-			}
+				FMetasoundAssetBase* DroppedMetaSoundAsset = IMetasoundUObjectRegistry::Get().GetObjectAsAssetBase(DroppedObject);
+				if (!DroppedMetaSoundAsset)
+				{
+					bIsValid = false;
+					OutReasonText = LOCTEXT("NotAllReferenceableMetaSounds", "Asset(s) not (all) MetaSound(s) that can be referenced.").ToString();
+					break;
+				}
 
-			if (MetaSoundAsset->AddingReferenceCausesLoop(*DroppedMetaSoundAsset))
-			{
-				OutOkIcon = false;
-				OutTooltipText = LOCTEXT("AddingMetaSoundWouldCauseLoop", "Cannot add asset(s) that would create a reference loop/loops.").ToString();
-				break;
-			}
+				if (MetaSoundAsset->AddingReferenceCausesLoop(*DroppedMetaSoundAsset))
+				{
+					bIsValid = false;
+					OutReasonText = LOCTEXT("AddingMetaSoundWouldCauseLoop", "Cannot add asset(s) that would create a reference loop/loops.").ToString();
+					break;
+				}
 
-			FText FailureReason;
-			if (AssetReferenceFilter.IsValid() && !AssetReferenceFilter->PassesFilter(Data, &FailureReason))
+				const EMetasoundFrontendClassAccessFlags DroppedAccessFlags = DroppedMetaSoundAsset->GetConstDocumentChecked().RootGraph.Metadata.GetAccessFlags();
+				if (!EnumHasAnyFlags(DroppedAccessFlags, EMetasoundFrontendClassAccessFlags::Referenceable))
+				{
+					bIsValid = false;
+					OutReasonText = LOCTEXT("NotAllAccessibleMetaSounds", "Asset(s) do not contain access flag 'Reference' (see 'Access Flags'). Cannot reference provided MetaSound Asset.").ToString();
+					break;
+				}
+
+				FText FailureReason;
+				if (AssetReferenceFilter.IsValid() && !AssetReferenceFilter->PassesFilter(Data, &FailureReason))
+				{
+					bIsValid = false;
+					OutReasonText = FailureReason.ToString();
+					break;
+				}
+				OutReasonText = LOCTEXT("AddMetaSoundToGraph", "Add MetaSound reference(s) to graph.").ToString();
+			}
+			else
 			{
-				OutOkIcon = false;
-				OutTooltipText = FailureReason.ToString();
+				bIsValid = false;
+				OutReasonText = LOCTEXT("AssetNotFound", "Asset(s) not found.").ToString();
 				break;
 			}
 		}
+		// For non MetaSounds, all assets must be of the same type
+		// and that type must be a registered proxy type
 		else
 		{
-			OutOkIcon = false;
-			OutTooltipText = LOCTEXT("AssetNotFound", "Asset(s) not found.").ToString();
-			break;
+			FDataTypeRegistryInfo Info;
+			const bool bFoundType = IDataTypeRegistry::Get().GetDataTypeInfo(AssetClass, Info);
+			if (bFoundType && Info.IsDataTypeProxyParsable())
+			{
+				if (!CurrentNonMetaSoundClass)
+				{
+					CurrentNonMetaSoundClass = AssetClass;
+					continue;
+				}
+				else if (CurrentNonMetaSoundClass != AssetClass)
+				{
+					bIsValid = false;
+					OutReasonText = MixedTypeErrorText;
+					break;
+				}
+				//else = default for multiple non MetaSound assets
+			}
+			else
+			{
+				bIsValid = false;
+				OutReasonText = LOCTEXT("NotRegisteredType", "Asset type(s) not registered for use in MetaSounds.").ToString();
+			}
 		}
 	}
+
+	if (OutTooltipText)
+	{
+		*OutTooltipText = OutReasonText;
+	}
+	return bIsValid;
 }
 
-void UMetasoundEditorGraphSchema::GetAssetsPinHoverMessage(const TArray<FAssetData>& Assets, const UEdGraphPin* HoverPin, FString& OutTooltipText, bool& OutOkIcon) const
+void UMetasoundEditorGraphSchema::GetAssetsGraphHoverMessage(const TArray<FAssetData>& Assets, const UEdGraph* HoverGraph, FString& OutTooltipText, bool& OutOkIcon) const
 {
+	using namespace Metasound;
 	using namespace Metasound::Editor;
 	using namespace Metasound::Frontend;
+	
+	OutOkIcon = ValidateDroppedAssetsOnGraph(Assets, HoverGraph, &OutTooltipText);
+	return; 
+}
 
-	if (HoverPin && HoverPin->Direction == EGPD_Input)
-	{
-		if (const UEdGraphNode* Node = HoverPin->GetOwningNode())
-		{
-			if (const UMetasoundEditorGraphNode* MetaSoundNode = Cast<UMetasoundEditorGraphNode>(HoverPin->GetOwningNode()))
-			{
-				if (Assets.Num() == 1)
-				{
-					FDataTypeRegistryInfo RegistryInfo = MetaSoundNode->GetPinDataTypeInfo(*HoverPin);
-					const bool bAssetTypesMatch = SchemaPrivate::DataTypeSupportsAssetTypes(RegistryInfo, Assets);
-					if (bAssetTypesMatch)
-					{
-						const TSharedPtr<IAssetReferenceFilter> AssetReferenceFilter = MakeAssetReferenceFilter(MetaSoundNode->GetGraph());
-						FText FailureReason;
-						if (AssetReferenceFilter.IsValid() && !AssetReferenceFilter->PassesFilter(Assets[0], &FailureReason))
-						{
-							OutTooltipText = FailureReason.ToString();
-							OutOkIcon = false;
-							return;
-						}
+void UMetasoundEditorGraphSchema::GetAssetsPinHoverMessage(const TArray<FAssetData>& Assets, const UEdGraphPin* HoverPin, FString& OutTooltipText, bool& OutOkIcon) const
+{
+	OutOkIcon = ValidateDroppedAssetsOnPin(Assets, HoverPin, &OutTooltipText);
+}
 
-						OutTooltipText = FText::FormatOrdered(LOCTEXT("DropAsset", "Set to '{0}'"), FText::FromName(Assets[0].AssetName)).ToString();
-						OutOkIcon = true;
-						return;
-					}
+void UMetasoundEditorGraphSchema::DroppedAssetsOnGraph(const TArray<FAssetData>& Assets, const FVector2f& GraphPosition, UEdGraph* Graph) const
+{
+	using namespace Metasound;
+	using namespace Metasound::Engine;
+	using namespace Metasound::Editor;
 
-					OutTooltipText = FText::FormatOrdered(LOCTEXT("InvalidType", "'{0}': Invalid Type"), FText::FromName(Assets[0].AssetName)).ToString();
-					OutOkIcon = false;
-					return;
-				}
+	const bool bIsValid = ValidateDroppedAssetsOnGraph(Assets, Graph);
+	if (!bIsValid)
+	{
+		return; 
+	}
 
-				OutTooltipText = LOCTEXT("CannotDropMultiple", "Cannot drop multiple assets on single pin.").ToString();
-				OutOkIcon = false;
-				return;
+	// Collect MetaSound assets
+	UMetasoundEditorGraph* MetaSoundGraph = CastChecked<UMetasoundEditorGraph>(Graph);
+	const TSharedPtr<IAssetReferenceFilter> AssetReferenceFilter = MakeAssetReferenceFilter(MetaSoundGraph);
+	TArray<FMetasoundAssetBase*> MetaSoundAssets;
+	Algo::Transform(Assets, MetaSoundAssets, [&AssetReferenceFilter](const FAssetData& DroppedAsset)
+	{
+		if (!AssetReferenceFilter.IsValid() || AssetReferenceFilter->PassesFilter(DroppedAsset))
+		{
+			if (UObject* DroppedObject = DroppedAsset.GetAsset())
+			{
+				return IMetasoundUObjectRegistry::Get().GetObjectAsAssetBase(DroppedObject);
 			}
-
-			OutTooltipText = FText::FormatOrdered(LOCTEXT("DragDropNotSupported", "Node '{0}' does not support drag/drop"), FText::FromString(Node->GetName())).ToString();
-			OutOkIcon = false;
-			return;
 		}
+		return static_cast<FMetasoundAssetBase*>(nullptr);
+	});
+	
+	// Add dropped MetaSound assets as nodes
+	if (MetaSoundAssets.ContainsByPredicate([](FMetasoundAssetBase* Asset) { return Asset != nullptr; }))
+	{
+		DroppedMetaSoundAssetsOnGraph(MetaSoundAssets, GraphPosition, MetaSoundGraph);
+	}
+	// Open action menu for non MetaSound assets
+	else
+	{
+		DroppedNonMetaSoundAssetsOnGraph(Assets, GraphPosition, MetaSoundGraph);
 	}
-
-	OutTooltipText.Reset();
-	OutOkIcon = false;
 }
 
-void UMetasoundEditorGraphSchema::DroppedAssetsOnGraph(const TArray<FAssetData>& Assets, const FVector2f& GraphPosition, UEdGraph* Graph) const
+void UMetasoundEditorGraphSchema::DroppedMetaSoundAssetsOnGraph(const TArray<FMetasoundAssetBase*>& MetaSoundAssets, const FVector2f& GraphPosition, UMetasoundEditorGraph* MetaSoundGraph) const
 {
 	using namespace Metasound;
 	using namespace Metasound::Engine;
 	using namespace Metasound::Frontend;
+	using namespace Metasound::Editor;
 
-	FScopedTransaction Transaction(LOCTEXT("DropMetaSoundOnGraph", "Drop MetaSound On Graph"));
-
-	UMetasoundEditorGraph* MetaSoundGraph = CastChecked<UMetasoundEditorGraph>(Graph);
 	bool bTransactionSucceeded = false;
 	bool bModifiedObjects = false;
+	check(MetaSoundGraph);
 	UObject& MetaSound = MetaSoundGraph->GetMetasoundChecked();
 
 	FMetasoundAssetBase* MetaSoundAsset = IMetasoundUObjectRegistry::Get().GetObjectAsAssetBase(&MetaSound);
@@ -2462,18 +2903,11 @@ void UMetasoundEditorGraphSchema::DroppedAssetsOnGraph(const TArray<FAssetData>&
 	UMetaSoundBuilderBase& Builder = FDocumentBuilderRegistry::GetChecked().FindOrBeginBuilding(MetaSound);
 	EMetaSoundBuilderResult Result = EMetaSoundBuilderResult::Failed;
 
-	const TSharedPtr<IAssetReferenceFilter> AssetReferenceFilter = MakeAssetReferenceFilter(MetaSoundGraph);
-
-	for (const FAssetData& DroppedAsset : Assets)
+	FScopedTransaction Transaction(LOCTEXT("DropMetaSoundOnGraph", "Drop MetaSound On Graph"));
+	for (FMetasoundAssetBase* DroppedMetaSoundAsset : MetaSoundAssets)
 	{
-		if (UObject* DroppedObject = DroppedAsset.GetAsset())
+		if (DroppedMetaSoundAsset)
 		{
-			FMetasoundAssetBase* DroppedMetaSoundAsset = IMetasoundUObjectRegistry::Get().GetObjectAsAssetBase(DroppedObject);
-			if (!DroppedMetaSoundAsset)
-			{
-				continue;
-			}
-
 			const EMetasoundFrontendClassAccessFlags DroppedAccessFlags = DroppedMetaSoundAsset->GetConstDocumentChecked().RootGraph.Metadata.GetAccessFlags();
 			if (!EnumHasAnyFlags(DroppedAccessFlags, EMetasoundFrontendClassAccessFlags::Referenceable))
 			{
@@ -2485,19 +2919,14 @@ void UMetasoundEditorGraphSchema::DroppedAssetsOnGraph(const TArray<FAssetData>&
 				continue;
 			}
 
-			if (AssetReferenceFilter.IsValid() && !AssetReferenceFilter->PassesFilter(DroppedAsset))
-			{
-				continue;
-			}
-
 			if (!bModifiedObjects)
 			{
 				MetaSound.Modify();
-				Graph->Modify();
+				MetaSoundGraph->Modify();
 				bModifiedObjects = true;
 			}
 
-			TScriptInterface<IMetaSoundDocumentInterface> DroppedDocInterface(DroppedObject);
+			TScriptInterface<IMetaSoundDocumentInterface> DroppedDocInterface(DroppedMetaSoundAsset->GetOwningAsset());
 			FMetaSoundNodeHandle NodeHandle = Builder.AddNode(DroppedDocInterface, Result);
 			if (ensure(Result == EMetaSoundBuilderResult::Succeeded))
 			{
@@ -2513,41 +2942,215 @@ void UMetasoundEditorGraphSchema::DroppedAssetsOnGraph(const TArray<FAssetData>&
 	}
 }
 
+void UMetasoundEditorGraphSchema::DroppedNonMetaSoundAssetsOnGraph(const TArray<FAssetData>& Assets, const FVector2f& GraphPosition, UMetasoundEditorGraph* MetaSoundGraph) const
+{
+	using namespace Metasound;
+	using namespace Metasound::Engine;
+	using namespace Metasound::Frontend;
+	using namespace Metasound::Editor;
+
+	bool bModifiedObjects = false;
+	check(MetaSoundGraph);
+	UObject& MetaSound = MetaSoundGraph->GetMetasoundChecked();
+	FScopedTransaction Transaction(LOCTEXT("DropMetaSoundOnGraph", "Drop Asset On Graph"));
+
+	// Collect dropped objects
+	FName DataType; 
+	TArray<UObject*> DroppedObjects;
+	for (const FAssetData& DroppedAsset : Assets)
+	{
+		if (UObject* DroppedObject = DroppedAsset.GetAsset())
+		{
+			if (!bModifiedObjects)
+			{
+				MetaSound.Modify();
+				MetaSoundGraph->Modify();
+				bModifiedObjects = true;
+			}
+			FDataTypeRegistryInfo Info;
+			if (IDataTypeRegistry::Get().GetDataTypeInfo(DroppedObject, Info))
+			{
+				if (Info.IsDataTypeProxyParsable())
+				{
+					if (DataType.IsNone())
+					{
+						DataType = Info.DataTypeName;
+					}
+					// Must be of same type (caller should have filtered out) 
+					if (ensure(DataType == Info.DataTypeName))
+					{
+						DroppedObjects.Add(DroppedObject);
+					}
+				}
+			}
+		}
+	}
+
+	// Set objects and open menu 
+	TSharedPtr<FEditor> MetasoundEditor = FGraphBuilder::GetEditorForGraph(*MetaSoundGraph);
+	if (MetasoundEditor.IsValid() && DroppedObjects.Num() > 0)
+	{
+		TSharedPtr<SGraphEditor> GraphEditor = MetasoundEditor->GetGraphEditor();
+		if (GraphEditor.IsValid())
+		{
+			if (DroppedObjects.Num() > 1)
+			{
+				DataType = CreateArrayTypeNameFromElementTypeName(DataType);
+			}
+			const FModifierKeysState ModifierKeys = FSlateApplication::Get().GetModifierKeys();
+			// Add new input for alt 
+			if (ModifierKeys.IsAltDown())
+			{
+				FMetasoundFrontendLiteral DefaultLiteral;
+				if (DroppedObjects.Num() == 1)
+				{
+					DefaultLiteral.Set(DroppedObjects[0]);
+				}
+				else
+				{
+					DefaultLiteral.Set(DroppedObjects);
+				}
+				SchemaUtils::AddNewGraphInput(DataType, DefaultLiteral, MetaSoundGraph, GraphPosition, /*bSelectNewNode*/true);
+			}
+			// Open node menu for drop or shift + drop
+			else
+			{
+				if (SGraphPanel* GraphPanel = GraphEditor->GetGraphPanel())
+				{
+					// Set dropped objects info
+					FDroppedObjectsInfo DroppedObjectsInfo =
+					{
+						.DroppedObjects = DroppedObjects,
+						.ModifierKeys = ModifierKeys,
+						.TypeName = DataType
+					};
+					MetasoundEditor->SetDroppedObjectsInfo(MoveTemp(DroppedObjectsInfo));
+
+					// Open and focus node menu 
+					// Copied from SGraphPanel::GraphCoordToPanelCoord
+					const FVector2f PanelPosition = (GraphPosition - GraphPanel->GetViewOffset()) * GraphPanel->GetZoomAmount();
+					// Convert panel to screen coordinates
+					const FVector2f ScreenPosition = GraphPanel->GetTickSpaceGeometry().LocalToAbsolute(PanelPosition);
+					const TSharedPtr<SWidget> MenuWidget = GraphPanel->SummonContextMenu(ScreenPosition, GraphPosition, {}, {}, {});
+					FSlateApplication::Get().SetKeyboardFocus(MenuWidget, EFocusCause::SetDirectly);
+
+					MetasoundEditor->ClearDroppedObjectsInfo();
+				}
+			}
+		}
+	}
+}
+
+bool UMetasoundEditorGraphSchema::ValidateDroppedAssetsOnPin(const TArray<FAssetData>& Assets, const UEdGraphPin* Pin, FString* OutTooltipText) const
+{
+	using namespace Metasound::Editor;
+	using namespace Metasound::Frontend;
+
+	const auto SetTooltipText = [&](FString InString)
+	{
+		if (OutTooltipText)
+		{
+			*OutTooltipText = InString;
+		}
+	};
+
+	if (Assets.Num() > 0)
+	{
+		if (Pin && Pin->Direction == EGPD_Input)
+		{
+			if (const UEdGraphNode* Node = Pin->GetOwningNode())
+			{
+				if (const UMetasoundEditorGraphNode* MetaSoundNode = Cast<UMetasoundEditorGraphNode>(Pin->GetOwningNode()))
+				{
+					const FDataTypeRegistryInfo RegistryInfo = MetaSoundNode->GetPinDataTypeInfo(*Pin);
+					const bool bAssetTypesMatch = SchemaPrivate::DataTypeSupportsAssetTypes(RegistryInfo, Assets);
+					if (bAssetTypesMatch)
+					{
+						const TSharedPtr<IAssetReferenceFilter> AssetReferenceFilter = MakeAssetReferenceFilter(MetaSoundNode->GetGraph());
+						FText FailureReason;
+						if (AssetReferenceFilter.IsValid())
+						{
+							for (const FAssetData& Asset : Assets)
+							{
+								if (!AssetReferenceFilter->PassesFilter(Asset, &FailureReason))
+								{
+									SetTooltipText(FailureReason.ToString());
+									return false;
+								}
+							}
+						}
+						if (Assets.Num() == 1)
+						{
+							SetTooltipText(FText::FormatOrdered(LOCTEXT("DropAsset", "Set to '{0}'"), FText::FromName(Assets[0].AssetName)).ToString());
+							return true;
+						}
+						else if (Assets.Num() > 1)
+						{
+							// Multiple assets 
+							SetTooltipText(LOCTEXT("DropMultipleAssetsAsInput", "Add assets as connected input").ToString());
+							return true;
+						}
+					}
+					else
+					{
+						SetTooltipText(FText::FormatOrdered(LOCTEXT("InvalidType", "'{0}': Type mismatch"), FText::FromName(Assets[0].AssetName)).ToString());
+					}
+					return false;
+				}
+
+				SetTooltipText(FText::FormatOrdered(LOCTEXT("DragDropNotSupported", "Node '{0}' does not support drag/drop"), FText::FromString(Node->GetName())).ToString());
+				return false;
+			}
+		}
+	}
+
+	SetTooltipText(LOCTEXT("InvalidAssetsOrPin", "Invalid assets or pin").ToString());
+	return false;
+}
+
 void UMetasoundEditorGraphSchema::DroppedAssetsOnPin(const TArray<FAssetData>& Assets, const FVector2f& GraphPosition, UEdGraphPin* Pin) const
 {
 	using namespace Metasound::Editor;
 	using namespace Metasound::Frontend;
 
-	if (!Pin)
+	if (!ValidateDroppedAssetsOnPin(Assets, Pin))
 	{
 		return;
 	}
 
 	if (UMetasoundEditorGraphNode* Node = Cast<UMetasoundEditorGraphNode>(Pin->GetOwningNode()))
 	{
-		if (Assets.Num() == 1)
+		const FDataTypeRegistryInfo RegistryInfo = Node->GetPinDataTypeInfo(*Pin);
+		const bool bAssetTypesMatch = SchemaPrivate::DataTypeSupportsAssetTypes(RegistryInfo, Assets);
+		if (bAssetTypesMatch)
 		{
-			FDataTypeRegistryInfo RegistryInfo = Node->GetPinDataTypeInfo(*Pin);
-			const bool bAssetTypesMatch = SchemaPrivate::DataTypeSupportsAssetTypes(RegistryInfo, Assets);
-			if (bAssetTypesMatch)
+			if (Assets.Num() == 1)
 			{
-				const TSharedPtr<IAssetReferenceFilter> AssetReferenceFilter = MakeAssetReferenceFilter(Node->GetGraph());
-				if (!AssetReferenceFilter.IsValid() || AssetReferenceFilter->PassesFilter(Assets.Last()))
+				UObject* Object = Assets.Last().GetAsset();
+				if (Object)
 				{
-					UObject* Object = Assets.Last().GetAsset();
-					if (Object)
-					{
-						const FText TransactionText = FText::Format(LOCTEXT("ChangeDefaultObjectTransaction", "Set {0} to '{1}'"),
-							Pin->GetDisplayName(),
-							FText::FromName(Object->GetFName()));
-						const FScopedTransaction Transaction(TransactionText);
-						Node->Modify();
-
-						constexpr bool bMarkAsModified = true;
-						TrySetDefaultObject(*Pin, Object, bMarkAsModified);
-					}
+					const FText TransactionText = FText::Format(LOCTEXT("ChangeDefaultObjectTransaction", "Set {0} to '{1}'"),
+						Pin->GetDisplayName(),
+						FText::FromName(Object->GetFName()));
+					const FScopedTransaction Transaction(TransactionText);
+					Node->Modify();
+
+					constexpr bool bMarkAsModified = true;
+					TrySetDefaultObject(*Pin, Object, bMarkAsModified);
 				}
 			}
+			// Add object array as new connected input
+			else
+			{
+				FMetasoundFrontendLiteral DefaultLiteral;
+				TArray<UObject*> Objects;
+				Algo::Transform(Assets, Objects,
+					[](const FAssetData& AssetData) {
+						return AssetData.GetAsset();
+					});
+				DefaultLiteral.Set(Objects);
+				SchemaUtils::PromoteToInput(Node->GetGraph(), Pin, FDeprecateSlateVector2D(GraphPosition) - DisplayStyle::NodeLayout::DefaultOffsetX, /*bSelectNewNode*/true, TOptional<FMetasoundFrontendLiteral>(MoveTemp(DefaultLiteral)));
+			}
 		}
 	}
 }
@@ -2655,7 +3258,7 @@ void UMetasoundEditorGraphSchema::GetDataTypeOutputNodeActions(FGraphContextMenu
 	SchemaPrivate::GetDataTypeActions<FMetasoundGraphSchemaAction_NewOutput>(ActionQuery);
 }
 
-void UMetasoundEditorGraphSchema::GetFunctionActions(FGraphActionMenuBuilder& ActionMenuBuilder, Metasound::Editor::FActionVertexFilters AccessFilters, bool bShowSelectedActions, Metasound::Frontend::FConstGraphHandle InGraphHandle) const
+void UMetasoundEditorGraphSchema::GetFunctionActions(FGraphActionMenuBuilder& ActionMenuBuilder, Metasound::Editor::FActionVertexFilters AccessFilters, bool bShowSelectedActions, Metasound::Frontend::FConstGraphHandle InGraphHandle, const UEdGraph* CurrentGraph) const
 {
 	using namespace Metasound::Editor;
 	using namespace Metasound::Engine;
@@ -2710,6 +3313,16 @@ void UMetasoundEditorGraphSchema::GetFunctionActions(FGraphActionMenuBuilder& Ac
 		{
 			return;
 		}
+		
+		TOptional<FDroppedObjectsInfo> OptionalDroppedObjectsInfo = TOptional<FDroppedObjectsInfo>();
+		if (CurrentGraph)
+		{
+			const TSharedPtr<FEditor> MetasoundEditor = FGraphBuilder::GetEditorForGraph(*CurrentGraph);
+			if (MetasoundEditor.IsValid())
+			{
+				OptionalDroppedObjectsInfo = MetasoundEditor->GetDroppedObjectsInfo();
+			}
+		}
 
 		const FText Tooltip = Result->GetTooltip();
 		const EPrimaryContextGroup ContextGroup = Result->IsNative() ? EPrimaryContextGroup::Functions : EPrimaryContextGroup::Graphs;
@@ -2722,7 +3335,8 @@ void UMetasoundEditorGraphSchema::GetFunctionActions(FGraphActionMenuBuilder& Ac
 			Result->GetDisplayName(),
 			Tooltip,
 			ContextGroup,
-			FText::Join(SchemaPrivate::KeywordDelim, Result->GetKeywords())
+			FText::Join(SchemaPrivate::KeywordDelim, Result->GetKeywords()), 
+			OptionalDroppedObjectsInfo
 		);
 
 		NewNodeAction->QueryResult = TSharedPtr<ISchemaQueryResult>(Result.Release());
@@ -2738,7 +3352,8 @@ void UMetasoundEditorGraphSchema::GetFunctionActions(FGraphActionMenuBuilder& Ac
 					NewNodeAction->GetMenuDescription(),
 					NewNodeAction->GetTooltipDescription(),
 					EPrimaryContextGroup::Recents,
-					NewNodeAction->GetKeywords());
+					NewNodeAction->GetKeywords(), 
+					OptionalDroppedObjectsInfo);
 				NewRecentAction->Grouping = static_cast<int32>(EPrimaryContextGroup::Recents) + RecentIndex;
 				NewRecentAction->QueryResult = NewNodeAction->QueryResult;
 				ActionMenuBuilder.AddAction(NewRecentAction);
@@ -2751,7 +3366,8 @@ void UMetasoundEditorGraphSchema::GetFunctionActions(FGraphActionMenuBuilder& Ac
 					NewNodeAction->GetMenuDescription(),
 					NewNodeAction->GetTooltipDescription(),
 					EPrimaryContextGroup::Favorites,
-					NewNodeAction->GetKeywords());
+					NewNodeAction->GetKeywords(), 
+					OptionalDroppedObjectsInfo);
 				NewFavAction->QueryResult = NewNodeAction->QueryResult;
 				ActionMenuBuilder.AddAction(NewFavAction);
 			}
diff --git a/Engine/Plugins/Runtime/Metasound/Source/MetasoundEditor/Public/IMetasoundEditor.h b/Engine/Plugins/Runtime/Metasound/Source/MetasoundEditor/Public/IMetasoundEditor.h
--- a/Engine/Plugins/Runtime/Metasound/Source/MetasoundEditor/Public/IMetasoundEditor.h
+++ b/Engine/Plugins/Runtime/Metasound/Source/MetasoundEditor/Public/IMetasoundEditor.h
@@ -29,3 +29,14 @@ public:
 	/** Whether or not the given editor is current associated with a UAudioComponent that is actively playing. */
 	virtual bool IsPlaying() const = 0;
 };
+
+namespace Metasound::Editor
+{
+	// For tracking objects being dropped onto the editor graph
+	struct FDroppedObjectsInfo
+	{
+		TArray<UObject*> DroppedObjects;
+		FModifierKeysState ModifierKeys;
+		FName TypeName;
+	};
+}
diff --git a/Engine/Plugins/Runtime/Metasound/Source/MetasoundEditor/Public/MetasoundEditorGraphSchema.h b/Engine/Plugins/Runtime/Metasound/Source/MetasoundEditor/Public/MetasoundEditorGraphSchema.h
--- a/Engine/Plugins/Runtime/Metasound/Source/MetasoundEditor/Public/MetasoundEditorGraphSchema.h
+++ b/Engine/Plugins/Runtime/Metasound/Source/MetasoundEditor/Public/MetasoundEditorGraphSchema.h
@@ -5,6 +5,7 @@
 #include "ConnectionDrawingPolicy.h"
 #include "EdGraph/EdGraphSchema.h"
 #include "EdGraphUtilities.h"
+#include "IMetasoundEditor.h"
 #include "MetasoundFrontend.h"
 #include "MetasoundFrontendController.h"
 #include "MetasoundFrontendNodesCategories.h"
@@ -22,6 +23,7 @@ class UEdGraph;
 class UEdGraphNode;
 class UMetaSoundPatch;
 class UMetasoundEditorGraphNode;
+class UMetasoundEditorGraph;
 
 
 namespace Metasound::Editor
@@ -87,7 +89,7 @@ namespace Metasound::Editor
 
 	namespace SchemaUtils
 	{
-		UEdGraphNode* PromoteToInput(UEdGraph* ParentGraph, UEdGraphPin* FromPin, const UE::Slate::FDeprecateVector2DParameter& InLocation, bool bSelectNewNode);
+		UEdGraphNode* PromoteToInput(UEdGraph* ParentGraph, UEdGraphPin* FromPin, const UE::Slate::FDeprecateVector2DParameter& InLocation, bool bSelectNewNode, TOptional<FMetasoundFrontendLiteral> Literal = {});
 		UEdGraphNode* PromoteToOutput(UEdGraph* ParentGraph, UEdGraphPin* FromPin, const UE::Slate::FDeprecateVector2DParameter& InLocation, bool bSelectNewNode);
 		UEdGraphNode* PromoteToVariable(UEdGraph* ParentGraph, UEdGraphPin* FromPin, const UE::Slate::FDeprecateVector2DParameter& InLocation, bool bSelectNewNode);
 		UEdGraphNode* PromoteToDeferredVariable(UEdGraph* ParentGraph, UEdGraphPin* FromPin, const UE::Slate::FDeprecateVector2DParameter& InLocation, bool bSelectNewNode);
@@ -165,7 +167,7 @@ struct FMetasoundGraphSchemaAction_NodeWithMultipleOutputs : public FMetasoundGr
 	//~ End FEdGraphSchemaAction Interface
 };
 
-/** Action to add an input reference to the graph */
+/** Action to add an input node of an existing graph input to the graph */
 USTRUCT()
 struct FMetasoundGraphSchemaAction_NewInput : public FMetasoundGraphSchemaAction_NodeWithMultipleOutputs
 {
@@ -199,6 +201,104 @@ struct FMetasoundGraphSchemaAction_NewInput : public FMetasoundGraphSchemaAction
 	//~ End FEdGraphSchemaAction Interface
 };
 
+/** Action to add a new graph member and member node to the graph, with the given default literal */
+USTRUCT()
+struct FMetasoundGraphSchemaAction_NewGraphMember : public FMetasoundGraphSchemaAction_NodeWithMultipleOutputs
+{
+	GENERATED_BODY();
+
+	FMetasoundGraphSchemaAction_NewGraphMember()
+		: FMetasoundGraphSchemaAction_NodeWithMultipleOutputs()
+	{}
+
+	UE_API FMetasoundGraphSchemaAction_NewGraphMember(FName InDataType, FMetasoundFrontendLiteral InLiteral, FText InNodeCategory, FText InMenuDesc, FText InToolTip, Metasound::Editor::EPrimaryContextGroup InGroup)
+		: FMetasoundGraphSchemaAction_NodeWithMultipleOutputs(InNodeCategory, InMenuDesc, InToolTip, InGroup)
+		, DataType(InDataType)
+		, DefaultLiteral(InLiteral)
+	{}
+
+	virtual ~FMetasoundGraphSchemaAction_NewGraphMember() = default;
+
+protected: 
+	FName DataType;
+	FMetasoundFrontendLiteral DefaultLiteral;
+};
+
+/** Action to add a new graph input and input node to the graph, with the given default literal */
+USTRUCT()
+struct FMetasoundGraphSchemaAction_NewGraphInput : public FMetasoundGraphSchemaAction_NewGraphMember
+{
+	GENERATED_BODY();
+
+	FMetasoundGraphSchemaAction_NewGraphInput()
+		: FMetasoundGraphSchemaAction_NewGraphMember()
+	{}
+
+	UE_API FMetasoundGraphSchemaAction_NewGraphInput(FName DataType, FMetasoundFrontendLiteral InLiteral);
+
+	virtual ~FMetasoundGraphSchemaAction_NewGraphInput() = default;
+
+	//~ Begin FMetasoundGraphSchemaAction Interface
+	UE_API virtual const FSlateBrush* GetIconBrush() const override;
+	UE_API virtual const FLinearColor& GetIconColor() const override;
+	//~ End FMetasoundGraphSchemaAction Interface
+
+	//~ Begin FEdGraphSchemaAction Interface
+	using FEdGraphSchemaAction::PerformAction; // Prevent hiding of deprecated base class function with FVector2D
+	UE_API virtual UEdGraphNode* PerformAction(UEdGraph* ParentGraph, UEdGraphPin* FromPin, const FVector2f& Location, bool bSelectNewNode = true) override;
+	//~ End FEdGraphSchemaAction Interface
+};
+
+/** Action to add a new graph output and output node to the graph, with the given default literal */
+USTRUCT()
+struct FMetasoundGraphSchemaAction_NewGraphOutput : public FMetasoundGraphSchemaAction_NewGraphMember
+{
+	GENERATED_BODY();
+
+	FMetasoundGraphSchemaAction_NewGraphOutput()
+		: FMetasoundGraphSchemaAction_NewGraphMember()
+	{}
+
+	UE_API FMetasoundGraphSchemaAction_NewGraphOutput(FName DataType, FMetasoundFrontendLiteral InLiteral);
+
+	virtual ~FMetasoundGraphSchemaAction_NewGraphOutput() = default;
+
+	//~ Begin FMetasoundGraphSchemaAction Interface
+	UE_API virtual const FSlateBrush* GetIconBrush() const override;
+	UE_API virtual const FLinearColor& GetIconColor() const override;
+	//~ End FMetasoundGraphSchemaAction Interface
+
+	//~ Begin FEdGraphSchemaAction Interface
+	using FEdGraphSchemaAction::PerformAction; // Prevent hiding of deprecated base class function with FVector2D
+	UE_API virtual UEdGraphNode* PerformAction(UEdGraph* ParentGraph, UEdGraphPin* FromPin, const FVector2f& Location, bool bSelectNewNode = true) override;
+	//~ End FEdGraphSchemaAction Interface
+};
+
+/** Action to add a new graph variable and getter node to the graph, with the given default literal */
+USTRUCT()
+struct FMetasoundGraphSchemaAction_NewGraphVariable : public FMetasoundGraphSchemaAction_NewGraphMember
+{
+	GENERATED_BODY();
+
+	FMetasoundGraphSchemaAction_NewGraphVariable()
+		: FMetasoundGraphSchemaAction_NewGraphMember()
+	{}
+
+	UE_API FMetasoundGraphSchemaAction_NewGraphVariable(FName DataType, FMetasoundFrontendLiteral InLiteral);
+
+	virtual ~FMetasoundGraphSchemaAction_NewGraphVariable() = default;
+
+	//~ Begin FMetasoundGraphSchemaAction Interface
+	UE_API virtual const FSlateBrush* GetIconBrush() const override;
+	UE_API virtual const FLinearColor& GetIconColor() const override;
+	//~ End FMetasoundGraphSchemaAction Interface
+
+	//~ Begin FEdGraphSchemaAction Interface
+	using FEdGraphSchemaAction::PerformAction; // Prevent hiding of deprecated base class function with FVector2D
+	UE_API virtual UEdGraphNode* PerformAction(UEdGraph* ParentGraph, UEdGraphPin* FromPin, const FVector2f& Location, bool bSelectNewNode = true) override;
+	//~ End FEdGraphSchemaAction Interface
+};
+
 /** Promotes an input to a graph input, using its respective literal value as the default value */
 USTRUCT()
 struct FMetasoundGraphSchemaAction_PromoteToInput : public FMetasoundGraphSchemaAction_NodeWithMultipleOutputs
@@ -427,8 +527,9 @@ struct FMetasoundGraphSchemaAction_NewNode : public FMetasoundGraphSchemaAction_
 		: FMetasoundGraphSchemaAction_NodeWithMultipleOutputs()
 	{}
 
-	FMetasoundGraphSchemaAction_NewNode(const FText& InNodeCategory, const FText& InMenuDesc, const FText& InToolTip, Metasound::Editor::EPrimaryContextGroup InGroup, FText InKeywords = FText::GetEmpty())
+	FMetasoundGraphSchemaAction_NewNode(const FText& InNodeCategory, const FText& InMenuDesc, const FText& InToolTip, Metasound::Editor::EPrimaryContextGroup InGroup, FText InKeywords = FText::GetEmpty(), TOptional<Metasound::Editor::FDroppedObjectsInfo> InDroppedObjectsInfo = {})
 		: FMetasoundGraphSchemaAction_NodeWithMultipleOutputs(InNodeCategory, InMenuDesc, InToolTip, InGroup, InKeywords)
+		, DroppedObjectsInfo(InDroppedObjectsInfo)
 	{}
 
 	virtual ~FMetasoundGraphSchemaAction_NewNode() = default;
@@ -445,6 +546,8 @@ struct FMetasoundGraphSchemaAction_NewNode : public FMetasoundGraphSchemaAction_
 	using FEdGraphSchemaAction::PerformAction; // Prevent hiding of deprecated base class function with FVector2D
 	UE_API virtual UEdGraphNode* PerformAction(UEdGraph* ParentGraph, UEdGraphPin* FromPin, const FVector2f& Location, bool bSelectNewNode = true) override;
 	//~ End FEdGraphSchemaAction Interface
+private: 
+	TOptional<Metasound::Editor::FDroppedObjectsInfo> DroppedObjectsInfo;
 };
 
 /** Action to add nodes to the graph based on selected objects*/
@@ -604,7 +707,7 @@ private:
 
 	/** Adds actions for creating actions associated with graph DataTypes */
 	void GetConversionActions(FGraphActionMenuBuilder& ActionMenuBuilder, Metasound::Editor::FActionVertexFilters InFilters = Metasound::Editor::FActionVertexFilters(), bool bShowSelectedActions = true) const;
-	void GetFunctionActions(FGraphActionMenuBuilder& ActionMenuBuilder, Metasound::Editor::FActionVertexFilters InFilters = Metasound::Editor::FActionVertexFilters(), bool bShowSelectedActions = true, Metasound::Frontend::FConstGraphHandle InGraphHandle = Metasound::Frontend::IGraphController::GetInvalidHandle()) const;
+	void GetFunctionActions(FGraphActionMenuBuilder& ActionMenuBuilder, Metasound::Editor::FActionVertexFilters InFilters = Metasound::Editor::FActionVertexFilters(), bool bShowSelectedActions = true, Metasound::Frontend::FConstGraphHandle InGraphHandle = Metasound::Frontend::IGraphController::GetInvalidHandle(), const UEdGraph* CurrentGraph = nullptr) const;
 	void GetVariableActions(FGraphActionMenuBuilder& ActionMenuBuilder, Metasound::Editor::FActionVertexFilters InFilters = Metasound::Editor::FActionVertexFilters(), bool bShowSelectedActions = true, Metasound::Frontend::FConstGraphHandle InGraphHandle = Metasound::Frontend::IGraphController::GetInvalidHandle()) const;
 
 	void BreakSinglePinLink(UEdGraphPin* SourcePin, UEdGraphPin* TargetPin, bool bShouldTransact) const;
@@ -615,6 +718,14 @@ private:
 	/** Adds action for creating a comment */
 	void GetCommentAction(FGraphActionMenuBuilder& ActionMenuBuilder, const UEdGraph* CurrentGraph = nullptr) const;
 
+	// Process dropped assets
+	void DroppedMetaSoundAssetsOnGraph(const TArray<FMetasoundAssetBase*>& MetaSoundAssets, const FVector2f& GraphPosition, UMetasoundEditorGraph* MetaSoundGraph) const;
+	void DroppedNonMetaSoundAssetsOnGraph(const TArray<FAssetData>& Assets, const FVector2f& GraphPosition, UMetasoundEditorGraph* MetaSoundGraph) const;
+
+	// Validate dropped assets, returning true if valid
+	bool ValidateDroppedAssetsOnGraph(const TArray<FAssetData>& Assets, const UEdGraph* HoverGraph, FString* OutTooltipText = nullptr) const;
+	bool ValidateDroppedAssetsOnPin(const TArray<FAssetData>& Assets, const UEdGraphPin* Pin, FString* OutTooltipText = nullptr) const;
+
 	static TSharedPtr<IAssetReferenceFilter> MakeAssetReferenceFilter(const UEdGraph* Graph);
 };
 
diff --git a/Engine/Plugins/Runtime/Metasound/Source/MetasoundFrontend/Private/MetasoundFrontendDataTypeRegistry.cpp b/Engine/Plugins/Runtime/Metasound/Source/MetasoundFrontend/Private/MetasoundFrontendDataTypeRegistry.cpp
--- a/Engine/Plugins/Runtime/Metasound/Source/MetasoundFrontend/Private/MetasoundFrontendDataTypeRegistry.cpp
+++ b/Engine/Plugins/Runtime/Metasound/Source/MetasoundFrontend/Private/MetasoundFrontendDataTypeRegistry.cpp
@@ -581,6 +581,7 @@ PRAGMA_ENABLE_DEPRECATION_WARNINGS
 
 				virtual const IDataTypeRegistryEntry* FindDataTypeRegistryEntry(const FName& InDataTypeName) const override;
 				virtual bool GetDataTypeInfo(const UObject* InObject, FDataTypeRegistryInfo& OutInfo) const override;
+				virtual bool GetDataTypeInfo(UClass* InClass, FDataTypeRegistryInfo& OutInfo) const override;
 				virtual bool GetDataTypeInfo(const FName& InDataType, FDataTypeRegistryInfo& OutInfo) const override;
 
 				virtual void IterateDataTypeInfo(TFunctionRef<void(const FDataTypeRegistryInfo&)> InFunction) const override;
@@ -785,6 +786,31 @@ PRAGMA_ENABLE_DEPRECATION_WARNINGS
 
 				return false;
 			}
+			
+			bool FDataTypeRegistry::GetDataTypeInfo(UClass* InClass, FDataTypeRegistryInfo& OutInfo) const
+			{
+				if (InClass)
+				{
+					UE::TScopeLock Lock(RegistryObjectMapMutex);
+					for (UClass* Class = InClass; Class != UObject::StaticClass(); Class = Class->GetSuperClass())
+					{
+						if (const FName* DataTypeName = RegisteredObjectClasses.Find(Class))
+						{
+							if (const IDataTypeRegistryEntry* Entry = FindDataTypeEntry(*DataTypeName))
+							{
+								const FDataTypeRegistryInfo& Info = Entry->GetDataTypeInfo();
+								if (Info.bIsExplicit && Class != InClass)
+								{
+									return false;
+								}
+								OutInfo = Info;
+								return true;
+							}
+						}
+					}
+				}
+				return false;
+			}
 
 			bool FDataTypeRegistry::GetDataTypeInfo(const FName& InDataType, FDataTypeRegistryInfo& OutInfo) const
 			{
diff --git a/Engine/Plugins/Runtime/Metasound/Source/MetasoundFrontend/Public/MetasoundFrontendDataTypeRegistry.h b/Engine/Plugins/Runtime/Metasound/Source/MetasoundFrontend/Public/MetasoundFrontendDataTypeRegistry.h
--- a/Engine/Plugins/Runtime/Metasound/Source/MetasoundFrontend/Public/MetasoundFrontendDataTypeRegistry.h
+++ b/Engine/Plugins/Runtime/Metasound/Source/MetasoundFrontend/Public/MetasoundFrontendDataTypeRegistry.h
@@ -325,6 +325,9 @@ namespace Metasound
 			/** Returns DataType info associated with the provided object. */
 			virtual bool GetDataTypeInfo(const UObject* InObject, FDataTypeRegistryInfo& OutInfo) const = 0;
 
+			/** Returns DataType info associated with the provided object class. */
+			virtual bool GetDataTypeInfo(UClass* InClass, FDataTypeRegistryInfo& OutInfo) const = 0;
+
 			/** Returns DataType info associated with the provided DataType name. */
 			virtual bool GetDataTypeInfo(const FName& InDataType, FDataTypeRegistryInfo& OutInfo) const = 0;
 
