Minor fix for MetaSound drag and drop of single object

#rnx
#Synced-CL 48957546

diff --git a/Engine/Plugins/Runtime/Metasound/Source/MetasoundEditor/Private/MetasoundEditorGraphSchema.cpp b/Engine/Plugins/Runtime/Metasound/Source/MetasoundEditor/Private/MetasoundEditorGraphSchema.cpp
--- a/Engine/Plugins/Runtime/Metasound/Source/MetasoundEditor/Private/MetasoundEditorGraphSchema.cpp
+++ b/Engine/Plugins/Runtime/Metasound/Source/MetasoundEditor/Private/MetasoundEditorGraphSchema.cpp
@@ -3081,6 +3081,11 @@ bool UMetasoundEditorGraphSchema::ValidateDroppedAssetsOnPin(const TArray<FAsset
 						}
 						if (Assets.Num() == 1)
 						{
+							if (Pin->HasAnyConnections())
+							{
+								SetTooltipText(LOCTEXT("PinAlreadyConnected", "Pin already has a connection").ToString());
+								return false;
+							}
 							SetTooltipText(FText::FormatOrdered(LOCTEXT("DropAsset", "Set to '{0}'"), FText::FromName(Assets[0].AssetName)).ToString());
 							return true;
 						}
@@ -3104,7 +3109,7 @@ bool UMetasoundEditorGraphSchema::ValidateDroppedAssetsOnPin(const TArray<FAsset
 		}
 	}
 
-	SetTooltipText(LOCTEXT("InvalidAssetsOrPin", "Invalid assets or pin").ToString());
+	SetTooltipText(LOCTEXT("InvalidAssetsOrPin", "Invalid asset(s) or input pin").ToString());
 	return false;
 }
 
@@ -3124,7 +3129,7 @@ void UMetasoundEditorGraphSchema::DroppedAssetsOnPin(const TArray<FAssetData>& A
 		const bool bAssetTypesMatch = SchemaPrivate::DataTypeSupportsAssetTypes(RegistryInfo, Assets);
 		if (bAssetTypesMatch)
 		{
-			if (Assets.Num() == 1)
+			if (Assets.Num() == 1 && !RegistryInfo.bIsArrayType)
 			{
 				UObject* Object = Assets.Last().GetAsset();
 				if (Object)
@@ -3139,7 +3144,7 @@ void UMetasoundEditorGraphSchema::DroppedAssetsOnPin(const TArray<FAssetData>& A
 					TrySetDefaultObject(*Pin, Object, bMarkAsModified);
 				}
 			}
-			// Add object array as new connected input
+			// Add object array as new connected input (including setting single object as array)
 			else
 			{
 				FMetasoundFrontendLiteral DefaultLiteral;
