- Wait until asset is versioned prior to collecting AssetRegistryTags
- Wait until MetaSound references are versioned prior to versioning referencing MetaSound
- Fix-up a couple spots missing task map lock protection

#rnx

#Synced-CL 48741549

diff --git a/Engine/Plugins/Runtime/Metasound/Source/MetasoundEngine/Private/MetasoundEngineAsset.h b/Engine/Plugins/Runtime/Metasound/Source/MetasoundEngine/Private/MetasoundEngineAsset.h
--- a/Engine/Plugins/Runtime/Metasound/Source/MetasoundEngine/Private/MetasoundEngineAsset.h
+++ b/Engine/Plugins/Runtime/Metasound/Source/MetasoundEngine/Private/MetasoundEngineAsset.h
@@ -214,6 +214,11 @@ namespace Metasound::Engine
 
 			const UObject* MetaSound = DocInterface.GetObject();
 			check(MetaSound);
+
+#if WITH_EDITORONLY_DATA
+			FVersioningManager::Get().WaitUntilVersioningComplete(*MetaSound);
+#endif // WITH_EDITORONLY_DATA
+
 			if (MetaSound->GetFlags() & (RF_Transient | RF_ClassDefaultObject))
 			{
 				return;
diff --git a/Engine/Plugins/Runtime/Metasound/Source/MetasoundFrontend/Private/MetasoundFrontendDocumentVersioning.cpp b/Engine/Plugins/Runtime/Metasound/Source/MetasoundFrontend/Private/MetasoundFrontendDocumentVersioning.cpp
--- a/Engine/Plugins/Runtime/Metasound/Source/MetasoundFrontend/Private/MetasoundFrontendDocumentVersioning.cpp
+++ b/Engine/Plugins/Runtime/Metasound/Source/MetasoundFrontend/Private/MetasoundFrontendDocumentVersioning.cpp
@@ -1151,55 +1151,78 @@ PRAGMA_ENABLE_DEPRECATION_WARNINGS
 		}
 
 		const FUniqueId AssetId = InMetaSound.GetUniqueID();
-		if (!ActiveVersionTaskData.Contains(AssetId))
+
 		{
-			TWeakObjectPtr<UObject> AssetPtr(&InMetaSound);
+			FScopeLock Lock(&ActiveVersionCritSection);
+			if (ActiveVersionTaskData.Contains(AssetId))
+			{
+				return;
+			}
+		}
 
-			check(DocInterface.GetObject());
-			const FMetasoundFrontendDocument& Document = DocInterface->GetConstDocument();
-			const FMetasoundFrontendClassName& ClassName = Document.RootGraph.Metadata.GetClassName();
+		// References have to be up-to-date before versioning the referencing asset to protect
+		// against accessing the referenced assets while being mutated asyncronously.
+		WaitUntilVersioningReferencesComplete(InMetaSound);
 
-			// Ensure nothing is currently holding on to a builder and use a local builder for versioning. All other call sites
-			// (such as a request to register and grabbing a new builder) should wait until versioning is complete for safety.
-			IDocumentBuilderRegistry::GetChecked().FinishBuilding(ClassName, DocInterface->GetAssetPathChecked());
+		check(DocInterface.GetObject());
+		const FMetasoundFrontendDocument& Document = DocInterface->GetConstDocument();
+		const FMetasoundFrontendClassName& ClassName = Document.RootGraph.Metadata.GetClassName();
 
-			FScopeLock Lock(&ActiveVersionCritSection);
-			ActiveVersionTaskData.Add(AssetId, FTaskData
+		// Ensure nothing is currently holding on to a builder and use a local builder for versioning. All other callsites
+		// (such as a request to register and grabbing a new builder) should wait until versioning is complete for safety.
+		// In theory, this should not be necessary.
+		IDocumentBuilderRegistry::GetChecked().FinishBuilding(ClassName, DocInterface->GetAssetPathChecked());
+
+		FTaskData NewTaskData
+		{
+			.AssetPtr = TWeakObjectPtr<UObject>(&InMetaSound),
+			.Task = UE::Tasks::Launch(UE_SOURCE_LOCATION, [this, AssetId, bIsDeterministic]()
 			{
-				.AssetPtr = MoveTemp(AssetPtr),
-				.Task = UE::Tasks::Launch(UE_SOURCE_LOCATION, [this, AssetId, bIsDeterministic]()
+				FTaskData TaskData;
 				{
 					FScopeLock VersionLock(&ActiveVersionCritSection);
-					FTaskData& TaskData = ActiveVersionTaskData.FindChecked(AssetId);
-
+					if (const FTaskData* TaskEntry = ActiveVersionTaskData.Find(AssetId))
 					{
-						FGCScopeGuard ScopeGuard;
-						FDocumentIDGenerator::FScopeDeterminism DeterminismScope(bIsDeterministic);
-						if (UObject* MetaSoundObject = TaskData.AssetPtr.Get())
-						{
-							FMetaSoundFrontendDocumentBuilder DocBuilder(MetaSoundObject);
-							const bool bVersionedDoc = VersionDocument(DocBuilder);
-							const bool bVersionedInterfaces = VersionInterfacesInternal(DocBuilder);
+						TaskData = *TaskEntry;
+					}
+					else
+					{
+						return;
+					}
+				}
 
-							if (bVersionedDoc || bVersionedInterfaces)
-							{
-								DocBuilder.GetMetasoundAsset().SetVersionedOnLoad();
-							}
+				if (TStrongObjectPtr<UObject> MetaSoundObject = TaskData.AssetPtr.Pin())
+				{
+					FDocumentIDGenerator::FScopeDeterminism DeterminismScope(bIsDeterministic);
+					FMetaSoundFrontendDocumentBuilder DocBuilder(MetaSoundObject.Get());
+					const bool bVersionedDoc = VersionDocument(DocBuilder);
+					const bool bVersionedInterfaces = VersionInterfacesInternal(DocBuilder);
 
-							ForEachObjectWithOuter(MetaSoundObject, [](UObject* Object)
-							{
-								if (Object->HasAnyInternalFlags(EInternalObjectFlags::Async))
-								{
-									METASOUND_VERSIONING_LOG(VeryVerbose, TEXT("Clearing object '%s' async flag (object likely generated while MetaSound versioning)"), *Object->GetName());
-									Object->ClearInternalFlags(EInternalObjectFlags::Async);
-								}
-							});
-						}
+					if (bVersionedDoc || bVersionedInterfaces)
+					{
+						DocBuilder.GetMetasoundAsset().SetVersionedOnLoad();
 					}
 
+					ForEachObjectWithOuter(MetaSoundObject.Get(), [](UObject* Object)
+					{
+						if (Object->HasAnyInternalFlags(EInternalObjectFlags::Async))
+						{
+							METASOUND_VERSIONING_LOG(VeryVerbose, TEXT("Clearing object '%s' async flag (object likely generated while MetaSound versioning)"), *Object->GetName());
+							Object->ClearInternalFlags(EInternalObjectFlags::Async);
+						}
+					});
+				}
+
+				{
+					FScopeLock VersionLock(&ActiveVersionCritSection);
 					ActiveVersionTaskData.Remove(AssetId);
-				})
-			});
+				}
+			})
+		};
+
+		{
+			FScopeLock Lock(&ActiveVersionCritSection);
+			ActiveVersionTaskData.Add(AssetId, MoveTemp(NewTaskData));
 		}
 	}
 
@@ -1303,8 +1326,47 @@ PRAGMA_ENABLE_DEPRECATION_WARNINGS
 		return bWasUpdated;
 	}
 
+	void FVersioningManager::WaitUntilVersioningReferencesComplete(const UObject& InMetaSound) const
+	{
+		using namespace UE::Tasks;
+
+		// Async reference load must be waited for prior to versioning wait as this is a
+		// required step to kick off versioning for older assets with soft object references.
+		if (IMetaSoundAssetManager* AssetManager = IMetaSoundAssetManager::Get())
+		{
+			if (const FMetasoundAssetBase* MetaSoundAsset = AssetManager->GetAsAsset(InMetaSound))
+			{
+				TArray<FMetasoundAssetBase*> ReferencedAssets = MetaSoundAsset->GetReferencedAssets();
+				for (FMetasoundAssetBase* Reference : ReferencedAssets)
+				{
+					AssetManager->WaitUntilAsyncLoadReferencedAssetsComplete(*Reference);
+
+					const UObject* RefOwner = Reference->GetOwningAsset();
+					check(RefOwner);
+
+					FTask TaskToWaitFor;
+					{
+						FScopeLock VersionLock(&ActiveVersionCritSection);
+						if (const FTaskData* TaskData = ActiveVersionTaskData.Find(RefOwner->GetUniqueID()))
+						{
+							TaskToWaitFor = TaskData->Task;
+						}
+						else
+						{
+							continue;
+						}
+					}
+
+					TaskToWaitFor.Wait();
+				}
+			}
+		}
+	}
+
 	void FVersioningManager::WaitUntilVersioningComplete(const UObject& InMetaSound) const
 	{
+		using namespace UE::Tasks;
+
 		// Async reference load must be waited for prior to versioning wait as this is a
 		// required step to kick off versioning for older assets with soft object references.
 		if (IMetaSoundAssetManager* AssetManager = IMetaSoundAssetManager::Get())
@@ -1315,10 +1377,19 @@ PRAGMA_ENABLE_DEPRECATION_WARNINGS
 			}
 		}
 
-		if (const FTaskData* TaskData = ActiveVersionTaskData.Find(InMetaSound.GetUniqueID()))
+		FTask TaskToWaitFor;
 		{
-			TaskData->Task.Wait();
+			FScopeLock VersionLock(&ActiveVersionCritSection);
+			if (const FTaskData* TaskData = ActiveVersionTaskData.Find(InMetaSound.GetUniqueID()))
+			{
+				TaskToWaitFor = TaskData->Task;
+			}
+			else
+			{
+				return;
+			}
 		}
+		TaskToWaitFor.Wait();
 	}
 } // namespace Metasound::Frontend
 
diff --git a/Engine/Plugins/Runtime/Metasound/Source/MetasoundFrontend/Public/MetasoundFrontendDocumentVersioning.h b/Engine/Plugins/Runtime/Metasound/Source/MetasoundFrontend/Public/MetasoundFrontendDocumentVersioning.h
--- a/Engine/Plugins/Runtime/Metasound/Source/MetasoundFrontend/Public/MetasoundFrontendDocumentVersioning.h
+++ b/Engine/Plugins/Runtime/Metasound/Source/MetasoundFrontend/Public/MetasoundFrontendDocumentVersioning.h
@@ -73,6 +73,9 @@ namespace Metasound::Frontend
 		// or impose any timing restrictions as the public API call does).
 		bool VersionInterfacesInternal(FMetaSoundFrontendDocumentBuilder& DocBuilder) const;
 
+		// Waits until the given MetaSound's references finish asynchronously versioning.
+		void WaitUntilVersioningReferencesComplete(const UObject& InMetaSound) const;
+
 		using FUniqueId = uint32;
 
 		mutable FCriticalSection ActiveVersionCritSection;
