Use the "new" Wave Player Operator by default for the MetasoundWavePlayerNode

The "new" Wave Player Operator swaps out the FSoundWaveProxy*Reader* for the FSoundWaveProxy*Player* among other things

The new opertaor implementation improves CPU and Memory performance and fixes some issues with Cue Point event triggers misbehaving when adjusting the speed

#rnx
#Synced-CL 48881229

diff --git a/Engine/Plugins/Runtime/Metasound/Source/MetasoundEngine/Private/MetasoundWavePlayerNode.cpp b/Engine/Plugins/Runtime/Metasound/Source/MetasoundEngine/Private/MetasoundWavePlayerNode.cpp
--- a/Engine/Plugins/Runtime/Metasound/Source/MetasoundEngine/Private/MetasoundWavePlayerNode.cpp
+++ b/Engine/Plugins/Runtime/Metasound/Source/MetasoundEngine/Private/MetasoundWavePlayerNode.cpp
@@ -332,7 +332,7 @@ namespace Metasound
 		TEXT("Default: 512"),
 		ECVF_Default);
 
-	static bool UseNewWavePlayerOperator = false;
+	static bool UseNewWavePlayerOperator = true;
 	FAutoConsoleVariableRef CVarMetaSoundWavePlayerUseNewWavePlayer(
 		TEXT("au.MetaSound.WavePlayer.UseNewOperator"),
 		UseNewWavePlayerOperator,
