Change assertions to warnings to unblock cook

#Synced-CL 49042160

diff --git a/Engine/Plugins/Runtime/Metasound/Source/MetasoundEngine/Private/MetasoundDocumentBuilderRegistry.cpp b/Engine/Plugins/Runtime/Metasound/Source/MetasoundEngine/Private/MetasoundDocumentBuilderRegistry.cpp
--- a/Engine/Plugins/Runtime/Metasound/Source/MetasoundEngine/Private/MetasoundDocumentBuilderRegistry.cpp
+++ b/Engine/Plugins/Runtime/Metasound/Source/MetasoundEngine/Private/MetasoundDocumentBuilderRegistry.cpp
@@ -3,6 +3,7 @@
 #include "MetasoundDocumentBuilderRegistry.h"
 
 #include "HAL/PlatformProperties.h"
+#include "Logging/StructuredLog.h"
 #include "MetasoundAssetManager.h"
 #include "MetasoundEngineAsset.h"
 #include "MetasoundFrontendDocumentVersioning.h"
@@ -172,11 +173,11 @@ namespace Metasound::Engine
 
 			const int32 NumRemainingGraphs = Document.RootGraph.GetConstGraphPages().Num();
 
-			checkf(NumRemainingGraphs > 0,
-				TEXT("Document in MetaSound asset '%s' had all default values "
+			UE_CLOGFMT(NumRemainingGraphs <= 0, LogMetaSound, Warning,
+				"Document in MetaSound asset '{AssetPath}' had all default values "
 					"cooked away leaving it in an invalid state. "
-					"Graph must always have at least one implementation."),
-				*Builder.GetDebugName());
+					"Graph must always have at least one implementation.",
+				Builder.GetDebugName());
 
 			if (NumInitGraphs > NumRemainingGraphs)
 			{
diff --git a/Engine/Plugins/Runtime/Metasound/Source/MetasoundFrontend/Private/MetasoundFrontendPages.cpp b/Engine/Plugins/Runtime/Metasound/Source/MetasoundFrontend/Private/MetasoundFrontendPages.cpp
--- a/Engine/Plugins/Runtime/Metasound/Source/MetasoundFrontend/Private/MetasoundFrontendPages.cpp
+++ b/Engine/Plugins/Runtime/Metasound/Source/MetasoundFrontend/Private/MetasoundFrontendPages.cpp
@@ -6,6 +6,7 @@
 #include "Algo/Transform.h"
 #include "Containers/Array.h"
 #include "Containers/ArrayView.h"
+#include "Logging/StructuredLog.h"
 #include "Misc/Guid.h"
 #include "MetasoundFrontendDocument.h"
 #include "MetasoundFrontendDocumentBuilder.h"
@@ -70,11 +71,11 @@ namespace Metasound::Frontend
 		// Log results and perform sanity checking
 		const int32 NumRemainingGraphs = Document.RootGraph.GetConstGraphPages().Num();
 
-		checkf(NumRemainingGraphs > 0,
-			TEXT("Document in MetaSound asset '%s' had all default values "
+		UE_CLOGFMT(NumRemainingGraphs <= 0, LogMetaSound, Warning,
+			"Document in MetaSound asset '{AssetPath}' had all default values "
 				"cooked away leaving it in an invalid state. "
-				"Graph must always have at least one implementation."),
-			*InBuilder.GetDebugName());
+				"Graph must always have at least one implementation.",
+			InBuilder.GetDebugName());
 
 		if (NumInitGraphs > NumRemainingGraphs)
 		{
@@ -130,12 +131,12 @@ namespace Metasound::Frontend
 			// Log results and perform sanity checking
 			const int32 NumRemainingDefaults = GraphInput.GetDefaults().Num();
 
-			checkf(NumRemainingDefaults > 0,
-				TEXT("Document in MetaSound asset '%s' had all default input values "
-					"cooked away for input %s leaving it in an invalid state. "
-					"Graph must always have at least one implementation."),
-				*InBuilder.GetDebugName(), 
-				*GraphInput.Name.ToString());
+			UE_CLOGFMT(NumRemainingDefaults <= 0, LogMetaSound, Warning,
+				"Document in MetaSound asset '{AssetPath}' had all default input values "
+					"cooked away for input {InputName} leaving it in an invalid state. "
+					"Graph must always have at least one implementation.",
+				InBuilder.GetDebugName(), 
+				GraphInput.Name);
 
 			if (NumInitDefaults > NumRemainingDefaults)
 			{
