- Add better logging when calling finish building prior to versioning MetaSound
- Avoid editor-only data manipulation in PreSave while cooking to avoid builder creationg/query

#rnx
#Synced-CL 49061405

diff --git a/Engine/Plugins/Runtime/Metasound/Source/MetasoundEditor/Private/MetasoundEditorGraph.cpp b/Engine/Plugins/Runtime/Metasound/Source/MetasoundEditor/Private/MetasoundEditorGraph.cpp
--- a/Engine/Plugins/Runtime/Metasound/Source/MetasoundEditor/Private/MetasoundEditorGraph.cpp
+++ b/Engine/Plugins/Runtime/Metasound/Source/MetasoundEditor/Private/MetasoundEditorGraph.cpp
@@ -2137,15 +2137,20 @@ void UMetasoundEditorGraph::PreSave(FObjectPreSaveContext InSaveContext)
 {
 	using namespace Metasound::Frontend;
 
-	TArray<UMetasoundEditorGraphNode*> MetaSoundNodes;
-	GetNodesOfClass<UMetasoundEditorGraphNode>(MetaSoundNodes);
-	for (UMetasoundEditorGraphNode* Node : MetaSoundNodes)
+#if WITH_EDITOR
+	if (!InSaveContext.IsCooking())
 	{
-		if (const FMetasoundFrontendNode* FrontendNode = Node->GetFrontendNode())
+		TArray<UMetasoundEditorGraphNode*> MetaSoundNodes;
+		GetNodesOfClass<UMetasoundEditorGraphNode>(MetaSoundNodes);
+		for (UMetasoundEditorGraphNode* Node : MetaSoundNodes)
 		{
-			FrontendNode->Style.bMessageNodeUpdated = false;
+			if (const FMetasoundFrontendNode* FrontendNode = Node->GetFrontendNode())
+			{
+				FrontendNode->Style.bMessageNodeUpdated = false;
+			}
 		}
 	}
+#endif // WITH_EDITOR
 
 	Super::PreSave(InSaveContext);
 }
diff --git a/Engine/Plugins/Runtime/Metasound/Source/MetasoundEditor/Private/MetasoundEditorGraphNode.cpp b/Engine/Plugins/Runtime/Metasound/Source/MetasoundEditor/Private/MetasoundEditorGraphNode.cpp
--- a/Engine/Plugins/Runtime/Metasound/Source/MetasoundEditor/Private/MetasoundEditorGraphNode.cpp
+++ b/Engine/Plugins/Runtime/Metasound/Source/MetasoundEditor/Private/MetasoundEditorGraphNode.cpp
@@ -713,14 +713,19 @@ void UMetasoundEditorGraphNode::PreSave(FObjectPreSaveContext InSaveContext)
 
 	Super::PreSave(InSaveContext);
 
-	// Required to refresh upgrade nodes that are stale when saving.
-	if (TSharedPtr<FEditor> MetaSoundEditor = FGraphBuilder::GetEditorForMetasound(GetMetasoundChecked()))
+#if WITH_EDITOR
+	if (!InSaveContext.IsCooking())
 	{
-		if (TSharedPtr<SGraphEditor> GraphEditor = MetaSoundEditor->GetGraphEditor())
+		// Required to refresh upgrade nodes that are stale when saving.
+		if (TSharedPtr<FEditor> MetaSoundEditor = FGraphBuilder::GetEditorForMetasound(GetMetasoundChecked()))
 		{
-			GraphEditor->RefreshNode(*this);
+			if (TSharedPtr<SGraphEditor> GraphEditor = MetaSoundEditor->GetGraphEditor())
+			{
+				GraphEditor->RefreshNode(*this);
+			}
 		}
 	}
+#endif // WITH_EDITOR
 }
 
 void UMetasoundEditorGraphNode::PostEditImport()
diff --git a/Engine/Plugins/Runtime/Metasound/Source/MetasoundFrontend/Private/MetasoundFrontendDocumentVersioning.cpp b/Engine/Plugins/Runtime/Metasound/Source/MetasoundFrontend/Private/MetasoundFrontendDocumentVersioning.cpp
--- a/Engine/Plugins/Runtime/Metasound/Source/MetasoundFrontend/Private/MetasoundFrontendDocumentVersioning.cpp
+++ b/Engine/Plugins/Runtime/Metasound/Source/MetasoundFrontend/Private/MetasoundFrontendDocumentVersioning.cpp
@@ -1167,11 +1167,19 @@ PRAGMA_ENABLE_DEPRECATION_WARNINGS
 		check(DocInterface.GetObject());
 		const FMetasoundFrontendDocument& Document = DocInterface->GetConstDocument();
 		const FMetasoundFrontendClassName& ClassName = Document.RootGraph.Metadata.GetClassName();
+		const FTopLevelAssetPath AssetPath = DocInterface->GetAssetPathChecked();
 
 		// Ensure nothing is currently holding on to a builder and use a local builder for versioning. All other callsites
 		// (such as a request to register and grabbing a new builder) should wait until versioning is complete for safety.
-		// In theory, this should not be necessary.
-		IDocumentBuilderRegistry::GetChecked().FinishBuilding(ClassName, DocInterface->GetAssetPathChecked());
+		// In theory, this should never happen. This log line helps determine if something somehow registered a builder
+		// prior to the asset being loaded and registered properly.
+		if (const FMetaSoundFrontendDocumentBuilder* ExistingBuilder = IDocumentBuilderRegistry::GetChecked().FindBuilder(ClassName, AssetPath))
+		{
+			METASOUND_VERSIONING_LOG(Warning,
+				TEXT("Existing builder for asset '%s' currently active prior to launching MetaSound versioning. Forcing builder to finish. "
+				"This indicates a system is erroneously initializing a builder prior to completing versioning."), *ExistingBuilder->GetDebugName());
+			IDocumentBuilderRegistry::GetChecked().FinishBuilding(ClassName, AssetPath);
+		}
 
 		FTaskData NewTaskData
 		{
