Versioning Clean-up Pass
- Make sure versioning code uses scoped deterministic new guid calls
- Consolidate migration code into single pass prior to all document version upgrade transforms
- Add Migration code for verifying no nodes have zero guid (some were found in old documents serialized in Lyra)

#rnx

#Synced-CL 49084301

diff --git a/Engine/Plugins/Runtime/Metasound/Source/MetasoundFrontend/Private/MetasoundFrontendDocumentVersioning.cpp b/Engine/Plugins/Runtime/Metasound/Source/MetasoundFrontend/Private/MetasoundFrontendDocumentVersioning.cpp
--- a/Engine/Plugins/Runtime/Metasound/Source/MetasoundFrontend/Private/MetasoundFrontendDocumentVersioning.cpp
+++ b/Engine/Plugins/Runtime/Metasound/Source/MetasoundFrontend/Private/MetasoundFrontendDocumentVersioning.cpp
@@ -167,7 +167,10 @@ namespace Metasound::Frontend
 PRAGMA_DISABLE_DEPRECATION_WARNINGS
 						if (Input.DefaultLiteral.IsValid())
 						{
-							Input.AddDefault(Frontend::DefaultPageID) = MoveTemp(Input.DefaultLiteral);
+							FMetasoundFrontendLiteral& Default = Input.ContainsDefault(Frontend::DefaultPageID)
+								? *Input.FindDefault(Frontend::DefaultPageID)
+								: Input.AddDefault(Frontend::DefaultPageID);
+							Default = MoveTemp(Input.DefaultLiteral);
 							Input.DefaultLiteral = FMetasoundFrontendLiteral::GetInvalid();
 							bUpdated = true;
 						}
@@ -176,11 +179,6 @@ PRAGMA_ENABLE_DEPRECATION_WARNINGS
 				};
 
 				FMetasoundFrontendDocument& Document = GetDocumentUnsafe(OutBuilder);
-				// For all class definitions we are going to access the default interface instead of inspecting the
-				// interface override. This is safe here because the class interface override did not exist in this
-				// version of the document. 
-				checkf(Document.Metadata.Version.Number <= FMetasoundFrontendVersionNumber(1, 14), TEXT("Migration of page properties needs to happen before the introduction of node configuration to the document"));
-
 				FMetasoundFrontendGraphClass& GraphClass = Document.RootGraph;
 				MigrateInterfaceInputDefaults(GraphClass.GetDefaultInterface());
 				for (FMetasoundFrontendClass& Dependency : Document.Dependencies)
@@ -188,9 +186,20 @@ PRAGMA_ENABLE_DEPRECATION_WARNINGS
 					MigrateInterfaceInputDefaults(Dependency.GetDefaultInterface());
 				}
 
+				if (bUpdated)
+				{
+					METASOUND_VERSIONING_LOG(Display, TEXT("Migrated Class Interface paged input defaults for MetaSound '%s'"), *OutBuilder.GetDebugName());
+				}
+
 				struct FMigratePageGraphs : public FMetasoundFrontendGraphClass::IPropertyVersionTransform
 				{
+					const UObject& PageGraphMS;
 				public:
+					FMigratePageGraphs(const UObject& InMetaSound)
+						: PageGraphMS(InMetaSound)
+					{
+					}
+
 					virtual ~FMigratePageGraphs() = default;
 
 					virtual bool Transform(FMetasoundFrontendGraphClass& OutClass) const override
@@ -200,6 +209,8 @@ PRAGMA_DISABLE_DEPRECATION_WARNINGS
 						if (Pages.IsEmpty())
 						{
 							Pages.Add(MoveTemp(OutClass.Graph));
+							OutClass.Graph = { };
+							METASOUND_VERSIONING_LOG(Display, TEXT("Migrated Class Interface paged graph for MetaSound '%s'"), *PageGraphMS.GetPathName());
 							return true;
 						}
 PRAGMA_ENABLE_DEPRECATION_WARNINGS
@@ -208,7 +219,68 @@ PRAGMA_ENABLE_DEPRECATION_WARNINGS
 					}
 				};
 
-				bUpdated |= FMigratePageGraphs().Transform(GraphClass);
+				bUpdated |= FMigratePageGraphs(OutBuilder.CastDocumentObjectChecked<UObject>()).Transform(GraphClass);
+
+				struct FMigrateNodeIds : public FMetasoundFrontendGraphClass::IPropertyVersionTransform
+				{
+					const UObject& NodeIDMS;
+					const FMetasoundFrontendDocument& NodeIDDoc;
+
+				public:
+					FMigrateNodeIds(const UObject& InMetaSound, const FMetasoundFrontendDocument& InDocument)
+						: NodeIDMS(InMetaSound)
+						, NodeIDDoc(InDocument)
+					{
+					}
+
+					virtual ~FMigrateNodeIds() = default;
+
+					virtual bool Transform(FMetasoundFrontendGraphClass& OutClass) const override
+					{
+						TArray<FMetasoundFrontendGraph>& Pages = GetPagesUnsafe(OutClass);
+
+						bool bNodeIDModified = false;
+
+						for (FMetasoundFrontendGraph& Graph : Pages)
+						{
+							FGuid NewNodeID;
+							for (FMetasoundFrontendNode& Node : Graph.Nodes)
+							{
+								if (Node.GetID() == NewNodeID)
+								{
+									checkf(!NewNodeID.IsValid(), TEXT("More than one node found with zero guid in %s: page %s.  Cannot recover from document corrupted with nodes of the same invalid id"),
+										*NodeIDMS.GetPathName(),
+										*Graph.PageID.ToString());
+									NewNodeID = FDocumentIDGenerator::Get().CreateNodeID(NodeIDDoc);
+									Node.UpdateID(NewNodeID);
+									METASOUND_VERSIONING_LOG(Display, TEXT("Found node (in '%s': page %s) with zero ID. Updating to '%s'"),
+										*NodeIDMS.GetPathName(),
+										*Graph.PageID.ToString(),
+										*NewNodeID.ToString())
+
+									bNodeIDModified = true;
+									for (FMetasoundFrontendEdge& Edge : Graph.Edges)
+									{
+										if (!Edge.ToNodeID.IsValid())
+										{
+											Edge.ToNodeID = NewNodeID;
+										}
+
+										if (!Edge.FromNodeID.IsValid())
+										{
+											Edge.FromNodeID = NewNodeID;
+										}
+									}
+									break;
+								}
+							}
+						}
+
+						return bNodeIDModified;
+					}
+				};
+				bUpdated |= FMigrateNodeIds(OutBuilder.CastDocumentObjectChecked<UObject>(), Document).Transform(GraphClass);
+
 				return bUpdated;
 			}
 		};
@@ -220,8 +292,13 @@ PRAGMA_ENABLE_DEPRECATION_WARNINGS
 
 			bool Transform(FMetaSoundFrontendDocumentBuilder& OutBuilder) const override
 			{
-				const bool bMigrated = GetDocumentUnsafe(OutBuilder).MigrateProperties();
-				return bMigrated;
+				const bool bMigratedProperties = GetDocumentUnsafe(OutBuilder).MigrateProperties();
+
+				// Page Graph migration must be completed for graph accessor back
+				// compat prior to all controller versioning, so just do it here.
+				const bool bMigratedPages = FMigratePagePropertiesTransform().Transform(OutBuilder);
+
+				return bMigratedProperties || bMigratedPages;
 			}
 		};
 
@@ -407,7 +484,10 @@ PRAGMA_ENABLE_DEPRECATION_WARNINGS
 				const FMetasoundFrontendGraphClass& GraphClass = InDocument->GetRootGraphClass();
 				FMetasoundFrontendClassMetadata Metadata = GraphClass.Metadata;
 
-				Metadata.SetClassName(FMetasoundFrontendClassName { FName(), *FGuid::NewGuid().ToString(), FName() });
+				FMetasoundFrontendDocument* Document = InDocument->GetDocumentPtr().Get();
+				check(Document);
+
+				Metadata.SetClassName(FMetasoundFrontendClassName { FName(), *FDocumentIDGenerator::Get().CreateClassID(*Document).ToString(), FName() });
 				InDocument->GetRootGraph()->SetGraphMetadata(Metadata);
 			}
 		};
@@ -529,7 +609,9 @@ PRAGMA_ENABLE_DEPRECATION_WARNINGS
 
 			void TransformInternal(FDocumentHandle InDocument) const override
 			{
-				const FGuid NewAssetClassID = FGuid::NewGuid();
+				FMetasoundFrontendDocument* Document = InDocument->GetDocumentPtr().Get();
+				check(Document);
+				const FGuid NewAssetClassID = FDocumentIDGenerator::Get().CreateClassID(*Document);
 				const FMetasoundFrontendGraphClass& Class = InDocument->GetRootGraphClass();
 				const_cast<FMetasoundFrontendGraphClass&>(Class).Metadata.SetClassName(FMetasoundFrontendClassName({ }, FName(*NewAssetClassID.ToString()), { }));
 			}
@@ -796,7 +878,10 @@ PRAGMA_ENABLE_DEPRECATION_WARNINGS
 				}
 				else
 				{
-					FMigratePagePropertiesTransform().Transform(OutBuilder);
+					// This used to migrate page properties, but this has since just been moved
+					// to property migration at the head of versioning call to keep API simple and
+					// consistent throughout process.
+					// FMigratePagePropertiesTransform().Transform(OutBuilder);
 					OutBuilder.GetMetasoundAsset().MigrateEditorGraph(OutBuilder);
 					METASOUND_VERSIONING_LOG(Display, TEXT("Resave recommended: Asset '%s' at '%s' successfully migrated editor data in target document version '%s'."), *Name.ToString(), *Path->ToString(), *GetTargetVersion().ToString());
 				}
@@ -816,7 +901,10 @@ PRAGMA_ENABLE_DEPRECATION_WARNINGS
 
 			void TransformInternal(FMetaSoundFrontendDocumentBuilder& OutBuilder) const override
 			{
-				FMigratePagePropertiesTransform().Transform(OutBuilder);
+				// This used to migrate page properties, but this has since just been moved
+				// to property migration at the head of versioning call to keep API simple and
+				// consistent throughout process.
+				// FMigratePagePropertiesTransform().Transform(OutBuilder);
 			}
 		};
 
@@ -1283,8 +1371,8 @@ PRAGMA_ENABLE_DEPRECATION_WARNINGS
 
 		METASOUND_TRACE_CPUPROFILER_EVENT_SCOPE(FVersioningManager::VersionDocument);
 
-		// Applies all property migration for data prior to versioning schema.
-		// Only runs if data has yet to be migrated.
+		// Applies all property migration for data prior to document
+		// versioning. Only returns true if data has yet to be migrated.
 		bool bWasUpdated = FMigratePropertiesTransform().Transform(DocBuilder);
 
 		UObject& MetaSoundObject = DocBuilder.CastDocumentObjectChecked<UObject>();
@@ -1300,14 +1388,11 @@ PRAGMA_ENABLE_DEPRECATION_WARNINGS
 			// Controller (Soft Deprecated) Transforms
 			if (InitVersionNumber.Major == 1 && InitVersionNumber.Minor < 12)
 			{
-				// Page Graph migration must be completed for graph accessor back
-				// compat prior to all controller versioning, so just do it here.
-				FMigratePagePropertiesTransform().Transform(DocBuilder);
-
 				// Builder is not registered during versioning, so don't attempt to reload builder
-				// (which can cause delegates to be fired in invalid scenarios like reloading assets
+				// through build (which can cause delegates to be fired in invalid scenarios like reloading assets
 				// from disk in editor).
-				FDocumentHandle DocHandle = DocBuilder.GetMetasoundAsset().GetDocumentHandle(false);
+				constexpr bool bReloadRegistryBuilder = false;
+				FDocumentHandle DocHandle = DocBuilder.GetMetasoundAsset().GetDocumentHandle(bReloadRegistryBuilder);
 
 				bWasUpdated |= FVersionDocument_1_1(Name, Path).Transform(DocHandle);
 				bWasUpdated |= FVersionDocument_1_2(Name, Path).Transform(DocHandle);
@@ -1323,6 +1408,8 @@ PRAGMA_ENABLE_DEPRECATION_WARNINGS
 				// No longer supported, new versions should go in VersioningPrivate::VersionBuilderDocument
 			}
 
+			// Reload builder directly here in case cache is left in invalid state from deprecated handle versioning above.
+			DocBuilder.Reload();
 			bWasUpdated |= VersionBuilderDocument(DocBuilder);
 			if (bWasUpdated)
 			{
