Initialize deprecated input default field to invalid value which can cause migrating to not version properly (due to later code initializing the struct with "valid" data in the deprecated field

#Synced-CL 49102037

diff --git a/Engine/Plugins/Runtime/Metasound/Source/MetasoundFrontend/Public/MetasoundFrontendDocument.h b/Engine/Plugins/Runtime/Metasound/Source/MetasoundFrontend/Public/MetasoundFrontendDocument.h
--- a/Engine/Plugins/Runtime/Metasound/Source/MetasoundFrontend/Public/MetasoundFrontendDocument.h
+++ b/Engine/Plugins/Runtime/Metasound/Source/MetasoundFrontend/Public/MetasoundFrontendDocument.h
@@ -1024,7 +1024,7 @@ struct FMetasoundFrontendClassInput : public FMetasoundFrontendClassVertex
 
 #if WITH_EDITORONLY_DATA
 	UPROPERTY(meta = (DeprecationMessage = "5.5 - Direct access will be revoked and page manipulation limited to public API in future builds. Field has been rolled into DefaultLiterals Array."))
-	FMetasoundFrontendLiteral DefaultLiteral;
+	FMetasoundFrontendLiteral DefaultLiteral = FMetasoundFrontendLiteral::GetInvalid();
 #endif // WITH_EDITORONLY_DATA
 
 	static UE_API bool IsFunctionalEquivalent(const FMetasoundFrontendClassInput& InLHS, const FMetasoundFrontendClassInput& InRHS);
