Feedback from prior CL: Fix-up comment.

#rnx
#Synced-CL 49103885

diff --git a/Engine/Plugins/Runtime/Metasound/Source/MetasoundFrontend/Private/MetasoundFrontendDocumentVersioning.cpp b/Engine/Plugins/Runtime/Metasound/Source/MetasoundFrontend/Private/MetasoundFrontendDocumentVersioning.cpp
--- a/Engine/Plugins/Runtime/Metasound/Source/MetasoundFrontend/Private/MetasoundFrontendDocumentVersioning.cpp
+++ b/Engine/Plugins/Runtime/Metasound/Source/MetasoundFrontend/Private/MetasoundFrontendDocumentVersioning.cpp
@@ -1388,9 +1388,9 @@ PRAGMA_ENABLE_DEPRECATION_WARNINGS
 			// Controller (Soft Deprecated) Transforms
 			if (InitVersionNumber.Major == 1 && InitVersionNumber.Minor < 12)
 			{
-				// Builder is not registered during versioning, so don't attempt to reload builder
-				// through build (which can cause delegates to be fired in invalid scenarios like reloading assets
-				// from disk in editor).
+				// Builder is not registered during versioning, so don't attempt to reload builder through builder
+				// registry indirectly here (which can cause delegates to be fired in invalid scenarios like
+				// reloading assets from disk in editor).
 				constexpr bool bReloadRegistryBuilder = false;
 				FDocumentHandle DocHandle = DocBuilder.GetMetasoundAsset().GetDocumentHandle(bReloadRegistryBuilder);
 
