Move log to just display when dropping connections (intent cannot be determined if underlying asset class changes API and should not rely on current application state to deteremine whether or not its a warning)

#rnx
#Synced-CL 49103885

diff --git a/Engine/Plugins/Runtime/Metasound/Source/MetasoundFrontend/Private/MetasoundFrontendDocumentVersioning.cpp b/Engine/Plugins/Runtime/Metasound/Source/MetasoundFrontend/Private/MetasoundFrontendDocumentVersioning.cpp
--- a/Engine/Plugins/Runtime/Metasound/Source/MetasoundFrontend/Private/MetasoundFrontendDocumentVersioning.cpp
+++ b/Engine/Plugins/Runtime/Metasound/Source/MetasoundFrontend/Private/MetasoundFrontendDocumentVersioning.cpp
@@ -246,7 +246,7 @@ PRAGMA_ENABLE_DEPRECATION_WARNINGS
 							FGuid NewNodeID;
 							for (FMetasoundFrontendNode& Node : Graph.Nodes)
 							{
-								if (Node.GetID() == NewNodeID)
+								if (!Node.GetID().IsValid())
 								{
 									checkf(!NewNodeID.IsValid(), TEXT("More than one node found with zero guid in %s: page %s.  Cannot recover from document corrupted with nodes of the same invalid id"),
 										*NodeIDMS.GetPathName(),
@@ -271,7 +271,6 @@ PRAGMA_ENABLE_DEPRECATION_WARNINGS
 											Edge.FromNodeID = NewNodeID;
 										}
 									}
-									break;
 								}
 							}
 						}
@@ -1452,6 +1451,7 @@ PRAGMA_ENABLE_DEPRECATION_WARNINGS
 						}
 					}
 
+					METASOUND_VERSIONING_LOG(VeryVerbose, TEXT("Waiting for %s to finish versioning on async task..."), *InMetaSound.GetFullName());
 					TaskToWaitFor.Wait();
 				}
 			}
@@ -1484,6 +1484,8 @@ PRAGMA_ENABLE_DEPRECATION_WARNINGS
 				return;
 			}
 		}
+
+		METASOUND_VERSIONING_LOG(VeryVerbose, TEXT("Waiting for %s to finish versioning on async task..."), *InMetaSound.GetFullName());
 		TaskToWaitFor.Wait();
 	}
 } // namespace Metasound::Frontend
diff --git a/Engine/Plugins/Runtime/Metasound/Source/MetasoundFrontend/Private/MetasoundFrontendTransform.cpp b/Engine/Plugins/Runtime/Metasound/Source/MetasoundFrontend/Private/MetasoundFrontendTransform.cpp
--- a/Engine/Plugins/Runtime/Metasound/Source/MetasoundFrontend/Private/MetasoundFrontendTransform.cpp
+++ b/Engine/Plugins/Runtime/Metasound/Source/MetasoundFrontend/Private/MetasoundFrontendTransform.cpp
@@ -29,22 +29,6 @@ namespace Metasound
 		{
 			bool bVersioningLoggingEnabled = true;
 
-			void LogAutoUpdateWarning(const FString& LogMessage)
-			{
-				// These should eventually move back to warning on cook 
-				// but are temporarily downgraded to prevent 
-				// warnings on things like unused test content from 
-				// blocking code checkins 
-				if (IsRunningCookCommandlet())
-				{
-					METASOUND_VERSIONING_LOG(Display, TEXT("%s"), *LogMessage);
-				}
-				else
-				{
-					METASOUND_VERSIONING_LOG(Warning, TEXT("%s"), *LogMessage);
-				}
-			}
-
 #if WITH_EDITOR
 			FGetNodeDisplayNameProjection NodeDisplayNameProjection = [] (const FNodeHandle&) { return FText(); };
 
@@ -621,12 +605,14 @@ namespace Metasound
 
 						for (const FVertexNameAndType& InputPin : DisconnectedInputs)
 						{
-							DocumentTransform::LogAutoUpdateWarning(FString::Printf(TEXT("Auto-Updating '%s' node class '%s (%s)': Previously connected input '%s' with data type '%s' no longer exists."), *DebugAssetPath, *NodeClassName, *NewClassVersion, *InputPin.Get<0>().ToString(), *InputPin.Get<1>().ToString()));
+							METASOUND_VERSIONING_LOG(Display, TEXT("Auto-Updating '%s' node class '%s (%s)': Previously connected input '%s' with data type '%s' no longer exists."),
+								*DebugAssetPath, *NodeClassName, *NewClassVersion, *InputPin.Get<0>().ToString(), *InputPin.Get<1>().ToString());
 						}
 
 						for (const FVertexNameAndType& OutputPin : DisconnectedOutputs)
 						{
-							DocumentTransform::LogAutoUpdateWarning(FString::Printf(TEXT("Auto-Updating '%s' node class '%s (%s)': Previously connected output '%s' with data type '%s' no longer exists."), *DebugAssetPath, *NodeClassName, *NewClassVersion, *OutputPin.Get<0>().ToString(), *OutputPin.Get<1>().ToString()));
+							METASOUND_VERSIONING_LOG(Display, TEXT("Auto-Updating '%s' node class '%s (%s)': Previously connected output '%s' with data type '%s' no longer exists."),
+								*DebugAssetPath, *NodeClassName, *NewClassVersion, *OutputPin.Get<0>().ToString(), *OutputPin.Get<1>().ToString());
 						}
 					}
 				}
