Add more protections against GC when versioning MetaSound assets
#rnx
#Synced-CL 49103885

diff --git a/Engine/Plugins/Runtime/Metasound/Source/MetasoundFrontend/Private/MetasoundFrontendDocumentVersioning.cpp b/Engine/Plugins/Runtime/Metasound/Source/MetasoundFrontend/Private/MetasoundFrontendDocumentVersioning.cpp
--- a/Engine/Plugins/Runtime/Metasound/Source/MetasoundFrontend/Private/MetasoundFrontendDocumentVersioning.cpp
+++ b/Engine/Plugins/Runtime/Metasound/Source/MetasoundFrontend/Private/MetasoundFrontendDocumentVersioning.cpp
@@ -1268,9 +1268,10 @@ PRAGMA_ENABLE_DEPRECATION_WARNINGS
 			IDocumentBuilderRegistry::GetChecked().FinishBuilding(ClassName, AssetPath);
 		}
 
+		InMetaSound.SetInternalFlags(EInternalObjectFlags::Async);
 		FTaskData NewTaskData
 		{
-			.AssetPtr = TWeakObjectPtr<UObject>(&InMetaSound),
+			.AssetPtr = TStrongObjectPtr<UObject>(&InMetaSound),
 			.Task = UE::Tasks::Launch(UE_SOURCE_LOCATION, [this, AssetId, bIsDeterministic]()
 			{
 				FTaskData TaskData;
@@ -1286,26 +1287,31 @@ PRAGMA_ENABLE_DEPRECATION_WARNINGS
 					}
 				}
 
-				if (TStrongObjectPtr<UObject> MetaSoundObject = TaskData.AssetPtr.Pin())
 				{
-					FDocumentIDGenerator::FScopeDeterminism DeterminismScope(bIsDeterministic);
-					FMetaSoundFrontendDocumentBuilder DocBuilder(MetaSoundObject.Get());
-					const bool bVersionedDoc = VersionDocument(DocBuilder);
-					const bool bVersionedInterfaces = VersionInterfacesInternal(DocBuilder);
+					FGCScopeGuard ScopeGuard;
 
-					if (bVersionedDoc || bVersionedInterfaces)
+					if (UObject* VersioningMetaSound = TaskData.AssetPtr.Get())
 					{
-						DocBuilder.GetMetasoundAsset().SetVersionedOnLoad();
-					}
+						FDocumentIDGenerator::FScopeDeterminism DeterminismScope(bIsDeterministic);
+						FMetaSoundFrontendDocumentBuilder DocBuilder(VersioningMetaSound);
+						const bool bVersionedDoc = VersionDocument(DocBuilder);
+						const bool bVersionedInterfaces = VersionInterfacesInternal(DocBuilder);
 
-					ForEachObjectWithOuter(MetaSoundObject.Get(), [](UObject* Object)
-					{
-						if (Object->HasAnyInternalFlags(EInternalObjectFlags::Async))
+						if (bVersionedDoc || bVersionedInterfaces)
 						{
-							METASOUND_VERSIONING_LOG(VeryVerbose, TEXT("Clearing object '%s' async flag (object likely generated while MetaSound versioning)"), *Object->GetName());
-							Object->ClearInternalFlags(EInternalObjectFlags::Async);
+							DocBuilder.GetMetasoundAsset().SetVersionedOnLoad();
 						}
-					});
+
+						ForEachObjectWithOuter(VersioningMetaSound, [](UObject* Object)
+						{
+							if (Object->HasAnyInternalFlags(EInternalObjectFlags::Async))
+							{
+								METASOUND_VERSIONING_LOG(VeryVerbose, TEXT("Clearing object '%s' async flag (object likely generated while MetaSound versioning)"), *Object->GetName());
+								Object->ClearInternalFlags(EInternalObjectFlags::Async);
+							}
+						});
+						VersioningMetaSound->ClearInternalFlags(EInternalObjectFlags::Async);
+					}
 				}
 
 				{
diff --git a/Engine/Plugins/Runtime/Metasound/Source/MetasoundFrontend/Public/MetasoundFrontendDocumentVersioning.h b/Engine/Plugins/Runtime/Metasound/Source/MetasoundFrontend/Public/MetasoundFrontendDocumentVersioning.h
--- a/Engine/Plugins/Runtime/Metasound/Source/MetasoundFrontend/Public/MetasoundFrontendDocumentVersioning.h
+++ b/Engine/Plugins/Runtime/Metasound/Source/MetasoundFrontend/Public/MetasoundFrontendDocumentVersioning.h
@@ -82,7 +82,7 @@ namespace Metasound::Frontend
 
 		struct FTaskData
 		{
-			TWeakObjectPtr<UObject> AssetPtr;
+			TStrongObjectPtr<UObject> AssetPtr;
 			UE::Tasks::FTask Task;
 		};
 		TMap<FUniqueId, FTaskData> ActiveVersionTaskData;
