Fix bad logic inadvertently applied when executing page input migration versioning fix-up

#rnx
#Synced-CL 49193787

diff --git a/Engine/Plugins/Runtime/Metasound/Source/MetasoundFrontend/Private/MetasoundFrontendDocumentVersioning.cpp b/Engine/Plugins/Runtime/Metasound/Source/MetasoundFrontend/Private/MetasoundFrontendDocumentVersioning.cpp
--- a/Engine/Plugins/Runtime/Metasound/Source/MetasoundFrontend/Private/MetasoundFrontendDocumentVersioning.cpp
+++ b/Engine/Plugins/Runtime/Metasound/Source/MetasoundFrontend/Private/MetasoundFrontendDocumentVersioning.cpp
@@ -167,10 +167,18 @@ namespace Metasound::Frontend
 PRAGMA_DISABLE_DEPRECATION_WARNINGS
 						if (Input.DefaultLiteral.IsValid())
 						{
-							FMetasoundFrontendLiteral& Default = Input.ContainsDefault(Frontend::DefaultPageID)
-								? *Input.FindDefault(Frontend::DefaultPageID)
-								: Input.AddDefault(Frontend::DefaultPageID);
-							Default = MoveTemp(Input.DefaultLiteral);
+							// There were cases where the default literal may have
+							// been set to 'None' incorrectly on struct initialization
+							// instead of invalid after versioning.  In this case, don't
+							// port the input default as a paged value if valid data
+							// already exists.
+							if (Input.GetDefaults().IsEmpty())
+							{
+								Input.AddDefault(Frontend::DefaultPageID) = MoveTemp(Input.DefaultLiteral);
+							}
+
+							// Always set default literal to invalid due to a minor versioning issue
+							// where the input default used to initialize its type set to 'None' incorrectly.
 							Input.DefaultLiteral = FMetasoundFrontendLiteral::GetInvalid();
 							bUpdated = true;
 						}
