Clarify versioning message when loading MetaSounds to include possible interface or property migration as reason for updates.

#rnx
#Synced-CL 49193787

diff --git a/Engine/Plugins/Runtime/Metasound/Source/MetasoundEditor/Private/MetasoundEditor.cpp b/Engine/Plugins/Runtime/Metasound/Source/MetasoundEditor/Private/MetasoundEditor.cpp
--- a/Engine/Plugins/Runtime/Metasound/Source/MetasoundEditor/Private/MetasoundEditor.cpp
+++ b/Engine/Plugins/Runtime/Metasound/Source/MetasoundEditor/Private/MetasoundEditor.cpp
@@ -1653,7 +1653,9 @@ namespace Metasound
 						check(DocInterface.GetObject());
 
 						const FString VersionString = DocInterface->GetConstDocument().Metadata.Version.Number.ToString();
-						FText Msg = FText::Format(LOCTEXT("MetaSoundDocumentVersioned", "Document versioned to '{0}' on load."), FText::FromString(VersionString));
+						FText Msg = FText::Format(LOCTEXT("MetaSoundDocumentVersioned",
+							"Document data migrated, interfaces updated, and/or properties versioned on "
+							"load. (Current document version: '{0}'). See log for more details."), FText::FromString(VersionString));
 						FNotificationInfo Info(Msg);
 						Info.bFireAndForget = true;
 						Info.bUseSuccessFailIcons = false;
