- Don't reset inherit override for presets when OnDefaultChanged called
- Initialize StartTime correctly as a float rather and default construct so override capability works as intended for Presets
- Fix missing delegate fired for back compat of OnInputInheritsDefaultChanged tools support

#rnx

#Synced-CL 49103885

diff --git a/Engine/Plugins/Runtime/Metasound/Source/MetasoundFrontend/Private/DocumentConfigurations/MetasoundFrontendPresetDocumentConfiguration.cpp b/Engine/Plugins/Runtime/Metasound/Source/MetasoundFrontend/Private/DocumentConfigurations/MetasoundFrontendPresetDocumentConfiguration.cpp
--- a/Engine/Plugins/Runtime/Metasound/Source/MetasoundFrontend/Private/DocumentConfigurations/MetasoundFrontendPresetDocumentConfiguration.cpp
+++ b/Engine/Plugins/Runtime/Metasound/Source/MetasoundFrontend/Private/DocumentConfigurations/MetasoundFrontendPresetDocumentConfiguration.cpp
@@ -18,6 +18,15 @@
 
 namespace Metasound::Frontend::PresetDocConfigPrivate
 {
+	// This likely should be deprecated in favor of using config data directly
+	void BroadcastInputInheritanceUpdated(const FMetaSoundFrontendDocumentBuilder& Builder, const FGuid& InNodeID)
+	{
+		auto IsNode = [&InNodeID](const FMetasoundFrontendClassInput& Input) { return Input.NodeID == InNodeID; };
+		const int32 Index = Builder.GetConstDocumentChecked().RootGraph.GetDefaultInterface().Inputs.IndexOfByPredicate(IsNode);
+		const_cast<FMetaSoundFrontendDocumentBuilder&>(Builder).GetDocumentDelegates()
+			.InterfaceDelegates.OnInputInheritsDefaultChanged.Broadcast(Index);
+	}
+
 	bool IsInheritableInputNode(const FMetaSoundFrontendDocumentBuilder& Builder, const FGuid& NodeID)
 	{
 		const FMetasoundFrontendNode* Node = Builder.FindNode(NodeID);
@@ -148,11 +157,22 @@ bool FMetaSoundFrontendPresetConfigurationVertexMetadata::IsDefaultEditable() co
 
 void FMetaSoundFrontendPresetConfigurationVertexMetadata::OnDefaultReset(const FMetaSoundFrontendDocumentBuilder& Builder)
 {
+	using namespace Metasound::Frontend;
+
 	bOverrideInheritedDefault = false;
+	PresetDocConfigPrivate::BroadcastInputInheritanceUpdated(const_cast<FMetaSoundFrontendDocumentBuilder&>(Builder), GetNodeID());
 }
 
 void FMetaSoundFrontendPresetConfigurationVertexMetadata::OnDefaultUpdated(const FMetaSoundFrontendDocumentBuilder& Builder)
 {
-	bOverrideInheritedDefault = true;
+	using namespace Metasound::Frontend;
+
+	// When actively configuring the document, ignore calls to defaults being
+	// updated as this means the request isn't being applied by a user/client code.
+	if (!Builder.IsConfiguringDocument())
+	{
+		bOverrideInheritedDefault = true;
+		PresetDocConfigPrivate::BroadcastInputInheritanceUpdated(const_cast<FMetaSoundFrontendDocumentBuilder&>(Builder), GetNodeID());
+	}
 }
 #undef LOCTEXT_NAMESPACE
diff --git a/Engine/Plugins/Runtime/Metasound/Source/MetasoundFrontend/Private/Interfaces/MetasoundFrontendSourceInterface.cpp b/Engine/Plugins/Runtime/Metasound/Source/MetasoundFrontend/Private/Interfaces/MetasoundFrontendSourceInterface.cpp
--- a/Engine/Plugins/Runtime/Metasound/Source/MetasoundFrontend/Private/Interfaces/MetasoundFrontendSourceInterface.cpp
+++ b/Engine/Plugins/Runtime/Metasound/Source/MetasoundFrontend/Private/Interfaces/MetasoundFrontendSourceInterface.cpp
@@ -102,10 +102,10 @@ namespace Metasound::Frontend
 					Inputs =
 					{
 						{
-							LOCTEXT("StartTime", "Start Time"),
-							LOCTEXT("StartTime Description", "[0, Duration] The StartTime of the source passed into the MetaSound from the Gameplay API."),
-							GetMetasoundDataTypeName<FTime>(),
-							Inputs::StartTime
+							.DisplayName = LOCTEXT("StartTime", "Start Time"),
+							.Description = LOCTEXT("StartTime Description", "[0, Duration] The StartTime of the source passed into the MetaSound from the Gameplay API."),
+							.DataType = GetMetasoundDataTypeName<FTime>(),
+							.InitValue = FAudioParameter(Inputs::StartTime, 0.f)
 						}
 					};
 #if WITH_EDITOR
diff --git a/Engine/Plugins/Runtime/Metasound/Source/MetasoundFrontend/Private/MetasoundFrontendDocumentBuilder.cpp b/Engine/Plugins/Runtime/Metasound/Source/MetasoundFrontend/Private/MetasoundFrontendDocumentBuilder.cpp
--- a/Engine/Plugins/Runtime/Metasound/Source/MetasoundFrontend/Private/MetasoundFrontendDocumentBuilder.cpp
+++ b/Engine/Plugins/Runtime/Metasound/Source/MetasoundFrontend/Private/MetasoundFrontendDocumentBuilder.cpp
@@ -4317,6 +4317,11 @@ void FMetaSoundFrontendDocumentBuilder::InitNodeLocations()
 #endif // WITH_EDITORONLY_DATA
 }
 
+bool FMetaSoundFrontendDocumentBuilder::IsConfiguringDocument() const
+{
+	return bInDocConfigScope;
+}
+
 bool FMetaSoundFrontendDocumentBuilder::IsDependencyReferenced(const FGuid& InClassID) const
 {
 	bool bIsReferenced = false;
diff --git a/Engine/Plugins/Runtime/Metasound/Source/MetasoundFrontend/Public/MetasoundFrontendDocumentBuilder.h b/Engine/Plugins/Runtime/Metasound/Source/MetasoundFrontend/Public/MetasoundFrontendDocumentBuilder.h
--- a/Engine/Plugins/Runtime/Metasound/Source/MetasoundFrontend/Public/MetasoundFrontendDocumentBuilder.h
+++ b/Engine/Plugins/Runtime/Metasound/Source/MetasoundFrontend/Public/MetasoundFrontendDocumentBuilder.h
@@ -545,6 +545,10 @@ PRAGMA_ENABLE_EXPERIMENTAL_WARNINGS
 
 	UE_API void InitNodeLocations();
 
+	// Whether or not the builder is actively configuring or re-configuring its document.
+	UE_EXPERIMENTAL(5.8, "Configurator is in active development and subject to change")
+	UE_API bool IsConfiguringDocument() const;
+
 	UE_API bool IsDependencyReferenced(const FGuid& InClassID) const;
 
 	// Determines if the builder's document configuration is of the given type
