Unregister MetaSound Data Types
#rnx
#Synced-CL 48963506

diff --git a/Engine/Plugins/Runtime/Metasound/Source/MetasoundFrontend/Private/MetasoundArrayNodes.cpp b/Engine/Plugins/Runtime/Metasound/Source/MetasoundFrontend/Private/MetasoundArrayNodes.cpp
--- a/Engine/Plugins/Runtime/Metasound/Source/MetasoundFrontend/Private/MetasoundArrayNodes.cpp
+++ b/Engine/Plugins/Runtime/Metasound/Source/MetasoundFrontend/Private/MetasoundArrayNodes.cpp
@@ -9,7 +9,7 @@ namespace Metasound
 {
 	namespace ArrayNodeVertexNames
 	{
-		const FLazyName InputInitialArrayName = TEXT("Array");
+		const FLazyName InputInitialArrayName{"Array"};
 #if WITH_EDITOR
 		const FText InputInitialArrayTooltip = LOCTEXT("InitialArrayTooltip", "Initial Array");
 		const FText InputInitialArrayDisplayName = LOCTEXT("InitialArrayDisplayName", "Init Array"); 
@@ -38,11 +38,42 @@ namespace Metasound
 
 	namespace MetasoundArrayNodesPrivate
 	{
+		const FLazyName ArrayNamespace{"Array"};
+
+		const FLazyName ArrayNumOperatorName{"Num"};
+		constexpr int32 ArrayNumNodeMajorVersion = 1;
+		constexpr int32 ArrayNumNodeMinorVersion = 0;
+
+		const FLazyName ArrayGetOperatorName{"Get"};
+		constexpr int32 ArrayGetNodeMajorVersion = 1;
+		constexpr int32 ArrayGetNodeMinorVersion = 0;
+
+		const FLazyName ArraySetOperatorName{"Set"};
+		constexpr int32 ArraySetNodeMajorVersion = 1;
+		constexpr int32 ArraySetNodeMinorVersion = 0;
+
+		const FLazyName ArraySubsetOperatorName{"Subset"};
+		constexpr int32 ArraySubsetNodeMajorVersion = 1;
+		constexpr int32 ArraySubsetNodeMinorVersion = 0;
+
+		const FLazyName ArrayConcatOperatorName{"Concat"}; 
+		constexpr int32 ArrayConcatNodeMajorVersion = 1;
+		constexpr int32 ArrayConcatNodeMinorVersion = 0;
+
+		const FLazyName ArrayLastIndexOperatorName{"GetLastIndex"};
+		constexpr int32 ArrayLastIndexNodeMajorVersion = 1;
+		constexpr int32 ArrayLastIndexNodeMinorVersion = 0;
+
+		FNodeClassName CreateArrayNodeClassName(const FName& InOperatorName, const FName& InDataTypeName)
+		{
+			return FNodeClassName{ ArrayNamespace, InOperatorName, InDataTypeName };
+		}
+
 		FNodeClassMetadata CreateArrayNodeClassMetadata(const FName& InDataTypeName, const FName& InOperatorName, const FText& InDisplayName, const FText& InDescription, const FVertexInterface& InDefaultInterface, int32 MajorVersion, int32 MinorVersion, bool bIsDeprecated)
 		{
 			FNodeClassMetadata Metadata
 			{
-				FNodeClassName { FName("Array"), InOperatorName, InDataTypeName },
+				CreateArrayNodeClassName(InOperatorName, InDataTypeName),
 				MajorVersion, 
 				MinorVersion,
 				InDisplayName, 
@@ -157,6 +188,110 @@ namespace Metasound
 				)
 			};
 		}
+
+		Frontend::FNodeClassRegistryKey CreateArrayNumNodeClassRegistryKey(const FName& InArrayDataTypeName)
+		{
+			return Frontend::FNodeClassRegistryKey(EMetasoundFrontendClassType::External, CreateArrayNodeClassName(ArrayNumOperatorName, InArrayDataTypeName), ArrayNumNodeMajorVersion, ArrayNumNodeMinorVersion);
+		}
+
+		Frontend::FNodeClassRegistryKey CreateArrayGetNodeClassRegistryKey(const FName& InArrayDataTypeName)
+		{
+			return Frontend::FNodeClassRegistryKey(EMetasoundFrontendClassType::External, CreateArrayNodeClassName(ArrayGetOperatorName, InArrayDataTypeName), ArrayGetNodeMajorVersion, ArrayGetNodeMinorVersion);
+		}
+
+		Frontend::FNodeClassRegistryKey CreateArraySetNodeClassRegistryKey(const FName& InArrayDataTypeName)
+		{
+			return Frontend::FNodeClassRegistryKey(EMetasoundFrontendClassType::External, CreateArrayNodeClassName(ArraySetOperatorName, InArrayDataTypeName), ArraySetNodeMajorVersion, ArraySetNodeMinorVersion);
+		}
+
+		Frontend::FNodeClassRegistryKey CreateArraySubsetNodeClassRegistryKey(const FName& InArrayDataTypeName)
+		{
+			return Frontend::FNodeClassRegistryKey(EMetasoundFrontendClassType::External, CreateArrayNodeClassName(ArraySubsetOperatorName, InArrayDataTypeName), ArraySubsetNodeMajorVersion, ArraySubsetNodeMinorVersion);
+		}
+
+		Frontend::FNodeClassRegistryKey CreateArrayConcatNodeClassRegistryKey(const FName& InArrayDataTypeName)
+		{
+			return Frontend::FNodeClassRegistryKey(EMetasoundFrontendClassType::External, CreateArrayNodeClassName(ArrayConcatOperatorName, InArrayDataTypeName), ArrayConcatNodeMajorVersion, ArrayConcatNodeMinorVersion);
+		}
+
+		Frontend::FNodeClassRegistryKey CreateArrayLastIndexNodeClassRegistryKey(const FName& InArrayDataTypeName)
+		{
+			return Frontend::FNodeClassRegistryKey(EMetasoundFrontendClassType::External, CreateArrayNodeClassName(ArrayLastIndexOperatorName, InArrayDataTypeName), ArrayLastIndexNodeMajorVersion, ArrayLastIndexNodeMinorVersion);
+		}
+
+
+		FNodeClassMetadata CreateArrayNumNodeClassMetadata(const FName& InArrayDataTypeName, const FText& InArrayTypeDisplayText)
+		{
+			const FText NodeDisplayName = METASOUND_LOCTEXT_FORMAT("ArrayOpArrayNumDisplayNamePattern", "Num ({0})", InArrayTypeDisplayText);
+			const FText NodeDescription = METASOUND_LOCTEXT("ArrayOpArrayNumDescription", "Number of elements in the array");
+
+			FNodeClassMetadata NodeClassMetadata = MetasoundArrayNodesPrivate::CreateArrayNodeClassMetadata(InArrayDataTypeName, ArrayNumOperatorName, NodeDisplayName, NodeDescription, CreateArrayNumInterface(InArrayDataTypeName), ArrayNumNodeMajorVersion, ArrayNumNodeMinorVersion);
+
+			NodeClassMetadata.Keywords.Append({
+				METASOUND_LOCTEXT("ArrayNumKeyword_Length", "length"),
+				METASOUND_LOCTEXT("ArrayNumKeyword_Size", "size"),
+			});
+			return NodeClassMetadata;
+		}
+
+
+		FNodeClassMetadata CreateArrayGetNodeClassMetadata(const FName& InArrayDataTypeName, const FName& InElementDataTypeName, const FText& InArrayTypeDisplayText)
+		{
+			const FText NodeDisplayName = METASOUND_LOCTEXT_FORMAT("ArrayOpArrayGetDisplayNamePattern", "Get ({0})", InArrayTypeDisplayText);
+			const FText NodeDescription = METASOUND_LOCTEXT("ArrayOpArrayGetDescription", "Get element at index in array.");
+			const FVertexInterface NodeInterface = MetasoundArrayNodesPrivate::CreateArrayGetInterface(InArrayDataTypeName, InElementDataTypeName);
+
+			return MetasoundArrayNodesPrivate::CreateArrayNodeClassMetadata(InArrayDataTypeName, ArrayGetOperatorName, NodeDisplayName, NodeDescription, NodeInterface, ArrayGetNodeMajorVersion, ArrayGetNodeMinorVersion);
+		}
+
+
+		FNodeClassMetadata CreateArraySetNodeClassMetadata(const FName& InArrayDataTypeName, const FName& InElementDataTypeName, const FText& InArrayTypeDisplayText)
+		{
+			const FText NodeDisplayName = METASOUND_LOCTEXT_FORMAT("ArrayOpArraySetDisplayNamePattern", "Set ({0})", InArrayTypeDisplayText);
+			const FText NodeDescription = METASOUND_LOCTEXT("ArrayOpArraySetDescription", "Set element at index in array.");
+			const FVertexInterface NodeInterface = MetasoundArrayNodesPrivate::CreateArraySetInterface(InArrayDataTypeName, InElementDataTypeName);
+		
+			return MetasoundArrayNodesPrivate::CreateArrayNodeClassMetadata(InArrayDataTypeName, ArraySetOperatorName, NodeDisplayName, NodeDescription, NodeInterface, ArraySetNodeMajorVersion, ArraySetNodeMinorVersion);
+		}
+		
+
+		FNodeClassMetadata CreateArraySubsetNodeClassMetadata(const FName& InArrayDataTypeName, const FText& InArrayDataTypeDisplayText)
+		{
+			const FText NodeDisplayName = METASOUND_LOCTEXT_FORMAT("ArrayOpArraySubsetDisplayNamePattern", "Subset ({0})", InArrayDataTypeDisplayText);
+			const FText NodeDescription = METASOUND_LOCTEXT("ArrayOpArraySubsetDescription", "Subset array on trigger.");
+			const FVertexInterface NodeInterface = MetasoundArrayNodesPrivate::CreateArraySubsetInterface(InArrayDataTypeName);
+		
+			FNodeClassMetadata NodeClassMetadata = MetasoundArrayNodesPrivate::CreateArrayNodeClassMetadata(InArrayDataTypeName, ArraySubsetOperatorName, NodeDisplayName, NodeDescription, NodeInterface, ArraySubsetNodeMajorVersion, ArraySubsetNodeMinorVersion);
+
+			NodeClassMetadata.Keywords.Append({
+				METASOUND_LOCTEXT("ArraySubsetKeyword_Slice", "slice"),
+			});
+
+			return NodeClassMetadata;
+		}
+
+		FNodeClassMetadata CreateArrayConcatNodeClassMetadata(const FName& InArrayDataTypeName, const FText& InArrayDataTypeDisplayText)
+		{
+			const FText NodeDisplayName = METASOUND_LOCTEXT_FORMAT("ArrayOpArrayConcatDisplayNamePattern", "Concatenate ({0})", InArrayDataTypeDisplayText);
+			const FText NodeDescription = METASOUND_LOCTEXT("ArrayOpArrayConcatDescription", "Concatenates two arrays on trigger.");
+			const FVertexInterface NodeInterface = MetasoundArrayNodesPrivate::CreateArrayConcatInterface(InArrayDataTypeName);
+		
+			FNodeClassMetadata NodeClassMetadata = MetasoundArrayNodesPrivate::CreateArrayNodeClassMetadata(InArrayDataTypeName, ArrayConcatOperatorName, NodeDisplayName, NodeDescription, NodeInterface, ArrayConcatNodeMajorVersion, ArrayConcatNodeMinorVersion);
+
+			NodeClassMetadata.Keywords.Append({
+				METASOUND_LOCTEXT("ArrayConcatKeyword_Append", "append"),
+			});
+
+			return NodeClassMetadata;
+		}
+
+		FNodeClassMetadata CreateArrayLastIndexNodeClassMetadata(const FName& InArrayDataTypeName, const FText& InArrayDataTypeDisplayText)
+		{
+			const FText NodeDisplayName = METASOUND_LOCTEXT_FORMAT("ArrayOpArrayLastIndexDisplayNamePattern", "Get Last Index ({0})", InArrayDataTypeDisplayText);
+			const FText NodeDescription = METASOUND_LOCTEXT("ArrayOpArrayLastIndexDescription", "Last index of the array");
+			const FVertexInterface NodeInterface = MetasoundArrayNodesPrivate::CreateArrayLastIndexInterface(InArrayDataTypeName);
+			return MetasoundArrayNodesPrivate::CreateArrayNodeClassMetadata(InArrayDataTypeName, ArrayLastIndexOperatorName, NodeDisplayName, NodeDescription, NodeInterface, ArrayLastIndexNodeMajorVersion, ArrayLastIndexNodeMinorVersion);
+		}
 	}
 }
 
diff --git a/Engine/Plugins/Runtime/Metasound/Source/MetasoundFrontend/Private/MetasoundArrayNodesRegistration.cpp b/Engine/Plugins/Runtime/Metasound/Source/MetasoundFrontend/Private/MetasoundArrayNodesRegistration.cpp
new file mode 100644
--- /dev/null
+++ b/Engine/Plugins/Runtime/Metasound/Source/MetasoundFrontend/Private/MetasoundArrayNodesRegistration.cpp
@@ -0,0 +1,39 @@
+// Copyright Epic Games, Inc. All Rights Reserved.
+
+#include "MetasoundArrayNodesRegistration.h"
+
+#include "MetasoundFrontendModuleRegistrationMacros.h"
+#include "MetasoundFrontendNodeClassRegistry.h"
+#include "MetasoundFrontendRegistryKey.h"
+
+namespace Metasound::Frontend
+{
+	bool UnregisterNodeIfRegistered(INodeClassRegistry& InRegistry, const FNodeClassRegistryKey& InKey)
+	{
+		if (InRegistry.IsNodeRegistered(InKey))
+		{
+			return InRegistry.UnregisterNode(InKey);
+		}
+		return true;
+	}
+
+	bool UnregisterArrayNodes(const FName& InArrayDataTypeName, const FModuleInfo& InModuleInfo)
+	{
+		using namespace MetasoundArrayNodesPrivate;
+
+		bool bSuccess = true;
+
+		INodeClassRegistry& NodeRegistry = INodeClassRegistry::GetChecked();
+
+		bSuccess = bSuccess && UnregisterNodeIfRegistered(NodeRegistry, CreateArrayNumNodeClassRegistryKey(InArrayDataTypeName));
+		bSuccess = bSuccess && UnregisterNodeIfRegistered(NodeRegistry, CreateArrayGetNodeClassRegistryKey(InArrayDataTypeName));
+		bSuccess = bSuccess && UnregisterNodeIfRegistered(NodeRegistry, CreateArraySetNodeClassRegistryKey(InArrayDataTypeName));
+		bSuccess = bSuccess && UnregisterNodeIfRegistered(NodeRegistry, CreateArraySubsetNodeClassRegistryKey(InArrayDataTypeName));
+		bSuccess = bSuccess && UnregisterNodeIfRegistered(NodeRegistry, CreateArrayConcatNodeClassRegistryKey(InArrayDataTypeName));
+		bSuccess = bSuccess && UnregisterNodeIfRegistered(NodeRegistry, ArrayShuffleNodePrivate::CreateArrayShuffleNodeClassRegistryKey(InArrayDataTypeName));
+		bSuccess = bSuccess && UnregisterNodeIfRegistered(NodeRegistry, ArrayRandomNodePrivate::CreateArrayRandomNodeClassRegistryKey(InArrayDataTypeName));
+		bSuccess = bSuccess && UnregisterNodeIfRegistered(NodeRegistry, CreateArrayLastIndexNodeClassRegistryKey(InArrayDataTypeName));
+
+		return bSuccess;
+	}
+}
diff --git a/Engine/Plugins/Runtime/Metasound/Source/MetasoundFrontend/Private/MetasoundArrayRandomNode.cpp b/Engine/Plugins/Runtime/Metasound/Source/MetasoundFrontend/Private/MetasoundArrayRandomNode.cpp
--- a/Engine/Plugins/Runtime/Metasound/Source/MetasoundFrontend/Private/MetasoundArrayRandomNode.cpp
+++ b/Engine/Plugins/Runtime/Metasound/Source/MetasoundFrontend/Private/MetasoundArrayRandomNode.cpp
@@ -34,6 +34,51 @@ namespace Metasound
 			"State is shared with other nodes with the same input array data (by value) regardless of where the node is located. Useful for sharing state regardless of graph composition or between multiple nodes within a single MetaSound using the same input data. Input array type must implement a hash function.")
 	DEFINE_METASOUND_ENUM_END()
 
+	namespace ArrayRandomNodePrivate
+	{
+		const FLazyName ArrayRandomOperatorName{"Random Get"};
+		constexpr int32 ArrayRandomNodeMajorVersion = 1;
+		constexpr int32 ArrayRandomNodeMinorVersion = 1;
+
+		Frontend::FNodeClassRegistryKey CreateArrayRandomNodeClassRegistryKey(const FName& InArrayDataTypeName)
+		{
+			return Frontend::FNodeClassRegistryKey(EMetasoundFrontendClassType::External, MetasoundArrayNodesPrivate::CreateArrayNodeClassName(ArrayRandomOperatorName, InArrayDataTypeName), ArrayRandomNodeMajorVersion, ArrayRandomNodeMinorVersion);
+		}
+
+		FVertexInterface CreateArrayRandomNodeInterface(const FName& InArrayDataTypeName, const FName& InElementDataTypeName)
+		{
+			using namespace ArrayNodeRandomGetVertexNames;
+
+			return FVertexInterface {
+				FInputVertexInterface{
+					TInputDataVertex<FTrigger>(METASOUND_GET_PARAM_NAME_AND_METADATA(InputTriggerNextValue)),
+					TInputDataVertex<FTrigger>(METASOUND_GET_PARAM_NAME_AND_METADATA(InputTriggerResetSeed)),
+					FInputDataVertex(METASOUND_GET_PARAM_NAME(InputRandomArray), InArrayDataTypeName, METASOUND_GET_PARAM_METADATA(InputRandomArray)),
+					TInputDataVertex<TArray<float>>(METASOUND_GET_PARAM_NAME_AND_METADATA(InputWeights)),
+					TInputDataVertex<int32>(METASOUND_GET_PARAM_NAME_AND_METADATA(InputSeed), -1),
+					TInputDataVertex<int32>(METASOUND_GET_PARAM_NAME_AND_METADATA(InputNoRepeatOrder), 1),
+					TInputDataVertex<bool>(METASOUND_GET_PARAM_NAME_AND_METADATA(InputEnableSharedState), false),
+					TInputConstructorVertex<FEnumSharedStateBehaviorType>(METASOUND_GET_PARAM_NAME_AND_METADATA(InputSharedStateBehavior), (int32)ESharedStateBehaviorType::SameNodeInComposition)
+				},
+				FOutputVertexInterface{
+					TOutputDataVertex<FTrigger>(METASOUND_GET_PARAM_NAME_AND_METADATA(OutputTriggerOnNext)),
+					TOutputDataVertex<FTrigger>(METASOUND_GET_PARAM_NAME_AND_METADATA(OutputTriggerOnReset)),
+					FOutputDataVertex(METASOUND_GET_PARAM_NAME(ShuffleOutputValue), InElementDataTypeName, METASOUND_GET_PARAM_METADATA(ShuffleOutputValue), EVertexAccessType::Reference),
+					TOutputDataVertex<int32>(METASOUND_GET_PARAM_NAME_AND_METADATA(OutputIndex))
+				}
+			};
+		}
+
+		FNodeClassMetadata CreateArrayRandomNodeClassMetadata(const FName& InArrayDataTypeName, const FName& InElementDataTypeName, const FText& InArrayDataTypeDisplayText)
+		{
+			const FText NodeDisplayName = METASOUND_LOCTEXT_FORMAT("RandomArrayGetNode_OpDisplayNamePattern", "Random Get ({0})", InArrayDataTypeDisplayText);
+
+			const FText NodeDescription = METASOUND_LOCTEXT("RandomArrayGetNode_Description", "Randomly retrieve data from input array using the supplied weights.");
+
+			return MetasoundArrayNodesPrivate::CreateArrayNodeClassMetadata(InArrayDataTypeName, ArrayRandomOperatorName, NodeDisplayName, NodeDescription, CreateArrayRandomNodeInterface(InArrayDataTypeName, InElementDataTypeName), ArrayRandomNodeMajorVersion, ArrayRandomNodeMinorVersion);
+		}
+	}
+
 	FArrayRandomGet::FArrayRandomGet(int32 InSeed, int32 InMaxIndex, const TArray<float>& InWeights, int32 InNoRepeatOrder)
 	{
 		UpdateState(InSeed, InMaxIndex, InWeights, InNoRepeatOrder);
diff --git a/Engine/Plugins/Runtime/Metasound/Source/MetasoundFrontend/Private/MetasoundArrayShuffleNode.cpp b/Engine/Plugins/Runtime/Metasound/Source/MetasoundFrontend/Private/MetasoundArrayShuffleNode.cpp
--- a/Engine/Plugins/Runtime/Metasound/Source/MetasoundFrontend/Private/MetasoundArrayShuffleNode.cpp
+++ b/Engine/Plugins/Runtime/Metasound/Source/MetasoundFrontend/Private/MetasoundArrayShuffleNode.cpp
@@ -7,6 +7,49 @@
 
 namespace Metasound
 {
+	namespace ArrayShuffleNodePrivate
+	{
+		const FLazyName ArrayShuffleOperatorName = "Shuffle";
+		constexpr int32 ArrayShuffleNodeMajorVersion = 1;
+		constexpr int32 ArrayShuffleNodeMinorVersion = 0;
+
+		Frontend::FNodeClassRegistryKey CreateArrayShuffleNodeClassRegistryKey(const FName& InArrayDataTypeName)
+		{
+			return Frontend::FNodeClassRegistryKey(EMetasoundFrontendClassType::External, MetasoundArrayNodesPrivate::CreateArrayNodeClassName(ArrayShuffleOperatorName, InArrayDataTypeName), ArrayShuffleNodeMajorVersion, ArrayShuffleNodeMinorVersion);
+		}
+
+		FVertexInterface CreateArrayShuffleNodeInterface(const FName& InArrayDataTypeName, const FName& InElementDataTypeName)
+		{
+			using namespace ArrayNodeShuffleVertexNames;
+
+			return FVertexInterface {
+				FInputVertexInterface{
+					TInputDataVertex<FTrigger>(METASOUND_GET_PARAM_NAME_AND_METADATA(InputTriggerNext)),
+					TInputDataVertex<FTrigger>(METASOUND_GET_PARAM_NAME_AND_METADATA(InputTriggerShuffle)),
+					TInputDataVertex<FTrigger>(METASOUND_GET_PARAM_NAME_AND_METADATA(InputTriggerReset)),
+					FInputDataVertex(METASOUND_GET_PARAM_NAME(InputShuffleArray), InArrayDataTypeName, METASOUND_GET_PARAM_METADATA(InputShuffleArray)),
+					TInputDataVertex<int32>(METASOUND_GET_PARAM_NAME_AND_METADATA(InputShuffleSeed), -1),
+					TInputDataVertex<bool>(METASOUND_GET_PARAM_NAME_AND_METADATA(InputAutoShuffle), true),
+					TInputDataVertex<bool>(METASOUND_GET_PARAM_NAME_AND_METADATA(InputShuffleEnableSharedState), false)
+					},
+				FOutputVertexInterface{
+					TOutputDataVertex<FTrigger>(METASOUND_GET_PARAM_NAME_AND_METADATA(OutputTriggerOnNext)),
+					TOutputDataVertex<FTrigger>(METASOUND_GET_PARAM_NAME_AND_METADATA(OutputTriggerOnShuffle)),
+					TOutputDataVertex<FTrigger>(METASOUND_GET_PARAM_NAME_AND_METADATA(OutputTriggerOnResetSeed)),
+					FOutputDataVertex(METASOUND_GET_PARAM_NAME(OutputShuffledValue), InElementDataTypeName, METASOUND_GET_PARAM_METADATA(OutputShuffledValue), EVertexAccessType::Reference)
+				}
+			};
+		}
+
+		FNodeClassMetadata CreateArrayShuffleNodeClassMetadata(const FName& InArrayDataTypeName, const FName& InElementDataTypeName, const FText& InArrayDataTypeDisplayText)
+		{
+			const FText NodeDisplayName = METASOUND_LOCTEXT_FORMAT("ArrayOpArrayShuffleDisplayNamePattern", "Shuffle ({0})", InArrayDataTypeDisplayText);
+			const FText NodeDescription = METASOUND_LOCTEXT("ArrayOpArrayShuffleDescription", "Output next element of a shuffled array on trigger.");
+
+			return MetasoundArrayNodesPrivate::CreateArrayNodeClassMetadata(InArrayDataTypeName, ArrayShuffleOperatorName, NodeDisplayName, NodeDescription, CreateArrayShuffleNodeInterface(InArrayDataTypeName, InElementDataTypeName), ArrayShuffleNodeMajorVersion, ArrayShuffleNodeMinorVersion);
+		}
+	}
+
 	/**
 	* FArrayIndexShuffler
 	*/
diff --git a/Engine/Plugins/Runtime/Metasound/Source/MetasoundFrontend/Private/MetasoundAutoConverterNode.cpp b/Engine/Plugins/Runtime/Metasound/Source/MetasoundFrontend/Private/MetasoundAutoConverterNode.cpp
--- a/Engine/Plugins/Runtime/Metasound/Source/MetasoundFrontend/Private/MetasoundAutoConverterNode.cpp
+++ b/Engine/Plugins/Runtime/Metasound/Source/MetasoundFrontend/Private/MetasoundAutoConverterNode.cpp
@@ -8,6 +8,9 @@ namespace Metasound
 {
 	namespace AutoConverterNodePrivate
 	{
+		constexpr int32 AutoConverterNodeMajorVersion = 1;
+		constexpr int32 AutoConverterNodeMinorVersion = 0;
+
 		FVertexName GetInputName(const FConvertDataTypeInfo& InInfo)
 		{
 			return InInfo.FromDataTypeName;
@@ -33,6 +36,16 @@ namespace Metasound
 			);
 		}
 
+		FNodeClassName CreateAutoConverterNodeClassName(const FName& InFromDataTypeName, const FName& InToDataTypeName)
+		{
+			return FNodeClassName{ TEXT("Convert"), InToDataTypeName, InFromDataTypeName };
+		}
+
+		Frontend::FNodeClassRegistryKey CreateAutoConverterNodeClassRegistryKey(const FName& InFromDataType, const FName& InToDataType)
+		{
+			return Frontend::FNodeClassRegistryKey{EMetasoundFrontendClassType::External, CreateAutoConverterNodeClassName(InFromDataType, InToDataType), AutoConverterNodeMajorVersion, AutoConverterNodeMinorVersion};
+		}
+
 		FNodeClassMetadata CreateAutoConverterNodeMetadata(const FConvertDataTypeInfo& InInfo)
 		{
 			FNodeDisplayStyle DisplayStyle;
@@ -45,9 +58,9 @@ namespace Metasound
 			const FText ToTypeText = InInfo.ToDataTypeText;
 
 			FNodeClassMetadata Info;
-			Info.ClassName = { TEXT("Convert"), InInfo.ToDataTypeName, InInfo.FromDataTypeName };
-			Info.MajorVersion = 1;
-			Info.MinorVersion = 0;
+			Info.ClassName = CreateAutoConverterNodeClassName(InInfo.FromDataTypeName, InInfo.ToDataTypeName);
+			Info.MajorVersion = AutoConverterNodeMajorVersion;
+			Info.MinorVersion = AutoConverterNodeMinorVersion;
 			Info.DisplayName = METASOUND_LOCTEXT_FORMAT("Metasound_AutoConverterNodeDisplayNameFormat", "{0} to {1}", InInfo.FromDataTypeText, InInfo.ToDataTypeText);
 			Info.Description = METASOUND_LOCTEXT_FORMAT("Metasound_AutoConverterNodeDescriptionNameFormat", "Converts from {0} to {1}.", InInfo.FromDataTypeText, InInfo.ToDataTypeText);
 			Info.Author = PluginAuthor;
diff --git a/Engine/Plugins/Runtime/Metasound/Source/MetasoundFrontend/Private/MetasoundConverterNodeRegistrationMacro.cpp b/Engine/Plugins/Runtime/Metasound/Source/MetasoundFrontend/Private/MetasoundConverterNodeRegistrationMacro.cpp
new file mode 100644
--- /dev/null
+++ b/Engine/Plugins/Runtime/Metasound/Source/MetasoundFrontend/Private/MetasoundConverterNodeRegistrationMacro.cpp
@@ -0,0 +1,26 @@
+// Copyright Epic Games, Inc. All Rights Reserved.
+
+#include "MetasoundConverterNodeRegistrationMacro.h"
+
+#include "MetasoundAutoConverterNode.h"
+#include "MetasoundFrontendNodeClassRegistry.h"
+
+namespace Metasound
+{
+	bool RegisterAutoConverterNode(const AutoConverterNodePrivate::FConvertDataTypeInfo& InConvertInfo)
+	{
+		using namespace AutoConverterNodePrivate;
+		using namespace Frontend;
+
+		FConverterNodeRegistryKey RegistryKey = { InConvertInfo.FromDataTypeName, InConvertInfo.ToDataTypeName };
+
+		FConverterNodeInfo ConverterNodeInfo =
+		{
+			GetInputName(InConvertInfo),
+			GetOutputName(InConvertInfo),
+			CreateAutoConverterNodeClassRegistryKey(InConvertInfo.FromDataTypeName, InConvertInfo.ToDataTypeName)
+		};
+
+		return INodeClassRegistry::GetChecked().RegisterConversionNode(RegistryKey, ConverterNodeInfo);
+	}
+}
diff --git a/Engine/Plugins/Runtime/Metasound/Source/MetasoundFrontend/Private/MetasoundDataTypeRegistrationMacro.cpp b/Engine/Plugins/Runtime/Metasound/Source/MetasoundFrontend/Private/MetasoundDataTypeRegistrationMacro.cpp
--- a/Engine/Plugins/Runtime/Metasound/Source/MetasoundFrontend/Private/MetasoundDataTypeRegistrationMacro.cpp
+++ b/Engine/Plugins/Runtime/Metasound/Source/MetasoundFrontend/Private/MetasoundDataTypeRegistrationMacro.cpp
@@ -6,6 +6,38 @@ namespace Metasound::Frontend
 {
 	namespace MetasoundDataTypeRegistrationPrivate
 	{
+		void UnregisterSendAndReceiveNodes(const FName& InDataTypeName, const FModuleInfo& InModuleInfo)
+		{
+			UnregisterNode(SendNodePrivate::CreateNodeClassRegistryKey(InDataTypeName), InModuleInfo);
+			UnregisterNode(ReceiveNodePrivate::CreateNodeClassRegistryKey(InDataTypeName), InModuleInfo);
+		}
+
+		void UnregisterAutoConverterNode(INodeClassRegistry& InRegistry, const FName& InFromDataTypeName, const FName& InToDataTypeName, const FModuleInfo& InModuleInfo)
+		{
+			const Metasound::Frontend::FNodeClassRegistryKey Key = AutoConverterNodePrivate::CreateAutoConverterNodeClassRegistryKey(InFromDataTypeName, InToDataTypeName);
+
+			if (InRegistry.IsNodeRegistered(Key))
+			{
+				UnregisterNode(Key, InModuleInfo);
+			}
+			InRegistry.UnregisterConversionNode(FConverterNodeClassRegistryKey{InFromDataTypeName, InToDataTypeName}, Key);
+		}
+
+		void UnregisterAutoConverterNodes(const FName& InDataTypeName, const FModuleInfo& InModuleInfo)
+		{
+			INodeClassRegistry& Registry = INodeClassRegistry::GetChecked();
+			// Conversions to this data type:
+			UnregisterAutoConverterNode(Registry, GetMetasoundDataTypeName<bool>(), InDataTypeName, InModuleInfo);
+			UnregisterAutoConverterNode(Registry, GetMetasoundDataTypeName<int32>(), InDataTypeName,  InModuleInfo);
+			UnregisterAutoConverterNode(Registry, GetMetasoundDataTypeName<float>(), InDataTypeName, InModuleInfo);
+			UnregisterAutoConverterNode(Registry, GetMetasoundDataTypeName<FString>(), InDataTypeName, InModuleInfo);
+
+			// Conversions from this data type:
+			UnregisterAutoConverterNode(Registry, InDataTypeName, GetMetasoundDataTypeName<bool>(), InModuleInfo);
+			UnregisterAutoConverterNode(Registry, InDataTypeName, GetMetasoundDataTypeName<int32>(), InModuleInfo);
+			UnregisterAutoConverterNode(Registry, InDataTypeName, GetMetasoundDataTypeName<float>(), InModuleInfo);
+			UnregisterAutoConverterNode(Registry, InDataTypeName, GetMetasoundDataTypeName<FString>(), InModuleInfo);
+		}
 
 		// Base registry entry for any data type
 		FDataTypeRegistryEntryBase::FDataTypeRegistryEntryBase(const Frontend::FDataTypeRegistryInfo& Info, const TSharedPtr<Frontend::IEnumDataTypeInterface>& EnumInterface)
@@ -195,5 +227,13 @@ namespace Metasound::Frontend
 			this->ConstructorOutputClassMetadata = MakeShared<FNodeClassMetadata>(InConstructorOutputClassMetadata);
 			this->ConstructorOutputClass = Metasound::Frontend::GenerateClass(*this->ConstructorOutputClassMetadata, EMetasoundFrontendClassType::Output);
 		}
+
+		bool UnregisterDataTypeWithFrontendInternal(const FName& InDataTypeName, const FModuleInfo& InModuleInfo)
+		{
+			UnregisterSendAndReceiveNodes(InDataTypeName, InModuleInfo);
+			UnregisterAutoConverterNodes(InDataTypeName, InModuleInfo);
+			return Frontend::IDataTypeRegistry::Get().UnregisterDataType(InDataTypeName);
+		}
+
 	}
 }
diff --git a/Engine/Plugins/Runtime/Metasound/Source/MetasoundFrontend/Private/MetasoundFrontendDataTypeRegistry.cpp b/Engine/Plugins/Runtime/Metasound/Source/MetasoundFrontend/Private/MetasoundFrontendDataTypeRegistry.cpp
--- a/Engine/Plugins/Runtime/Metasound/Source/MetasoundFrontend/Private/MetasoundFrontendDataTypeRegistry.cpp
+++ b/Engine/Plugins/Runtime/Metasound/Source/MetasoundFrontend/Private/MetasoundFrontendDataTypeRegistry.cpp
@@ -569,13 +569,9 @@ PRAGMA_ENABLE_DEPRECATION_WARNINGS
 			public:
 				virtual ~FDataTypeRegistry() = default;
 
-				/** Register a data type
-				 * @param InName - Name of data type.
-				 * @param InEntry - TUniquePtr to data type registry entry.
-				 *
-				 * @return True on success, false on failure.
-				 */
+	
 				virtual bool RegisterDataType(TUniquePtr<IDataTypeRegistryEntry>&& InEntry) override;
+				virtual bool UnregisterDataType(const FName& InDataTypeName) override;
 
 				virtual void GetRegisteredDataTypeNames(TArray<FName>& OutNames) const override;
 
@@ -715,7 +711,7 @@ PRAGMA_ENABLE_DEPRECATION_WARNINGS
 					}
 
 					// Register nodes associated with data type.
-					FMetasoundFrontendRegistryContainer* NodeRegistry = FMetasoundFrontendRegistryContainer::Get();
+					INodeClassRegistry* NodeRegistry = INodeClassRegistry::Get();
 					if (ensure(nullptr != NodeRegistry))
 					{
 						if (Entry->GetDataTypeInfo().bIsParsable)
@@ -750,6 +746,56 @@ PRAGMA_ENABLE_DEPRECATION_WARNINGS
 				return false;
 			}
 
+			bool FDataTypeRegistry::UnregisterDataType(const FName& InDataTypeName)
+			{
+				METASOUND_LLM_SCOPE;
+
+				bool bSuccess = false;
+				if (const IDataTypeRegistryEntry* Entry = FindDataTypeRegistryEntry(InDataTypeName))
+				{
+					bSuccess = true;
+					const FDataTypeRegistryInfo& DataTypeInfo = Entry->GetDataTypeInfo();
+
+					// Unregister nodes associated with data type.
+					INodeClassRegistry& NodeRegistry = INodeClassRegistry::GetChecked();
+					if (DataTypeInfo.bIsParsable)
+					{
+						NodeRegistry.UnregisterNode(FNodeClassRegistryKey{Entry->GetFrontendInputClass().Metadata});
+						NodeRegistry.UnregisterNode(FNodeClassRegistryKey{Entry->GetFrontendOutputClass().Metadata});
+						NodeRegistry.UnregisterNode(FNodeClassRegistryKey{Entry->GetFrontendLiteralClass().Metadata});
+						NodeRegistry.UnregisterNode(FNodeClassRegistryKey{Entry->GetFrontendVariableClass().Metadata});
+						NodeRegistry.UnregisterNode(FNodeClassRegistryKey{Entry->GetFrontendVariableMutatorClass().Metadata});
+						NodeRegistry.UnregisterNode(FNodeClassRegistryKey{Entry->GetFrontendVariableAccessorClass().Metadata});
+						NodeRegistry.UnregisterNode(FNodeClassRegistryKey{Entry->GetFrontendVariableDeferredAccessorClass().Metadata});
+
+						if (DataTypeInfo.bIsConstructorType)
+						{
+							NodeRegistry.UnregisterNode(FNodeClassRegistryKey{Entry->GetFrontendConstructorInputClass().Metadata});
+							NodeRegistry.UnregisterNode(FNodeClassRegistryKey{Entry->GetFrontendConstructorOutputClass().Metadata});
+						}
+					}
+					
+					{
+						UE::TScopeLock ObjectRegistryLock(RegistryObjectMapMutex);
+						if (const UClass* Class = DataTypeInfo.ProxyGeneratorClass)
+						{
+							bSuccess = bSuccess && RegisteredObjectClasses.Remove(Class);
+						}
+					}
+					
+					{
+						UE::TScopeLock RegistryLock(RegistryMapMutex);
+						bSuccess = bSuccess && RegisteredDataTypes.Remove(InDataTypeName);
+					}
+
+				}
+
+				UE_CLOG(bSuccess, LogMetaSound, Verbose, TEXT("Unregistered Metasound Datatype [Name:%s]."), *InDataTypeName.ToString());
+				UE_CLOG(!bSuccess, LogMetaSound, Error, TEXT("Failed to unregistered Metasound Datatype [Name:%s]."), *InDataTypeName.ToString());
+
+				return bSuccess;
+			}
+
 			void FDataTypeRegistry::GetRegisteredDataTypeNames(TArray<FName>& OutNames) const
 			{
 				UE::TScopeLock Lock(RegistryMapMutex);
diff --git a/Engine/Plugins/Runtime/Metasound/Source/MetasoundFrontend/Private/MetasoundFrontendNodeClassRegistry.cpp b/Engine/Plugins/Runtime/Metasound/Source/MetasoundFrontend/Private/MetasoundFrontendNodeClassRegistry.cpp
--- a/Engine/Plugins/Runtime/Metasound/Source/MetasoundFrontend/Private/MetasoundFrontendNodeClassRegistry.cpp
+++ b/Engine/Plugins/Runtime/Metasound/Source/MetasoundFrontend/Private/MetasoundFrontendNodeClassRegistry.cpp
@@ -1023,6 +1023,30 @@ PRAGMA_ENABLE_DEPRECATION_WARNINGS
 		return bSuccess;
 	}
 
+	FName FNodeClassRegistry::GetOwningPluginName(const FNodeClassRegistryKey& InKey) const
+	{
+		FName PluginName;
+		auto GetPluginName = [&PluginName](const INodeClassRegistryEntry& Entry)
+		{
+			PluginName = Entry.GetPluginName();
+		};
+
+		AccessNodeEntryThreadSafe(InKey, GetPluginName);
+		return PluginName;
+	}
+
+	FName FNodeClassRegistry::GetOwningModuleName(const FNodeClassRegistryKey& InKey) const
+	{
+		FName ModuleName;
+		auto GetModuleName = [&ModuleName](const INodeClassRegistryEntry& Entry)
+		{
+			ModuleName = Entry.GetModuleName();
+		};
+
+		AccessNodeEntryThreadSafe(InKey, GetModuleName);
+		return ModuleName;
+	}
+
 	bool FNodeClassRegistry::RegisterNodeUpdateTransform(const FNodeClassRegistryKey& InNodeClassKey, const TSharedRef<const INodeUpdateTransform>& InNodeUpdateTransform)
 	{
 		if (IsNodeRegistered(InNodeClassKey))
@@ -1180,14 +1204,14 @@ PRAGMA_ENABLE_DEPRECATION_WARNINGS
 		return false;
 	}
 
-	bool FNodeClassRegistry::RegisterConversionNode(const FConverterNodeClassRegistryKey& InNodeKey, const FConverterNodeInfo& InNodeInfo)
+	bool FNodeClassRegistry::RegisterConversionNode(const FConverterNodeClassRegistryKey& InConversionKey, const FConverterNodeInfo& InNodeInfo)
 	{
-		if (!ConverterNodeClassRegistry.Contains(InNodeKey))
+		if (!ConverterNodeClassRegistry.Contains(InConversionKey))
 		{
-			ConverterNodeClassRegistry.Add(InNodeKey);
+			ConverterNodeClassRegistry.Add(InConversionKey);
 		}
 
-		FConverterNodeClassRegistryValue& ConverterNodeList = ConverterNodeClassRegistry[InNodeKey];
+		FConverterNodeClassRegistryValue& ConverterNodeList = ConverterNodeClassRegistry[InConversionKey];
 
 		if (ensureAlways(!ConverterNodeList.PotentialConverterNodes.Contains(InNodeInfo)))
 		{
@@ -1201,6 +1225,21 @@ PRAGMA_ENABLE_DEPRECATION_WARNINGS
 		}
 	}
 
+	bool FNodeClassRegistry::UnregisterConversionNode(const FConverterNodeClassRegistryKey& InConversionKey, const FNodeClassRegistryKey& InNodeKey)
+	{
+		if (FConverterNodeClassRegistryValue* ConverterNodeList = ConverterNodeClassRegistry.Find(InConversionKey))
+		{
+			auto HasSameNodeKey = [&InNodeKey](const FConverterNodeClassInfo& InInfo) -> bool
+			{
+				return InInfo.NodeKey == InNodeKey;
+			};
+			int32 NumRemoved = ConverterNodeList->PotentialConverterNodes.RemoveAll(HasSameNodeKey);
+
+			return NumRemoved > 0;
+		}
+		return false;
+	}
+
 	bool FNodeClassRegistry::IsNodeRegistered(const FNodeClassRegistryKey& InKey) const
 	{
 		auto IsNodeRegisteredInternal = [this, &InKey]() -> bool
@@ -1389,7 +1428,9 @@ PRAGMA_ENABLE_DEPRECATION_WARNINGS
 
 	void FNodeClassRegistry::IterateRegistry(FIterateMetasoundFrontendClassFunction InIterFunc, EMetasoundFrontendClassType InClassType) const
 	{
-		UE_LOG(LogMetaSound, Warning, TEXT("Calling FMetasoundRegistryContainer::IterateRegistry(...) is not threadsafe. Please use Metasound::Frontend::ISearchEngine instead"));
+		UE_LOG(LogMetaSound, Verbose, TEXT("Calling FMetasoundRegistryContainer::IterateRegistry(...) is can be slow and prone to deadlocks because it locks the registry with a critical section. Please use Metasound::Frontend::ISearchEngine instead"));
+
+		UE::TScopeLock Lock(RegistryMapsCriticalSection);
 		auto WrappedFunc = [&](const TPair<FNodeClassRegistryKey, TSharedPtr<INodeClassRegistryEntry, ESPMode::ThreadSafe>>& Pair)
 		{
 			InIterFunc(Pair.Value->GetFrontendClass());
diff --git a/Engine/Plugins/Runtime/Metasound/Source/MetasoundFrontend/Private/MetasoundFrontendNodeClassRegistryPrivate.h b/Engine/Plugins/Runtime/Metasound/Source/MetasoundFrontend/Private/MetasoundFrontendNodeClassRegistryPrivate.h
--- a/Engine/Plugins/Runtime/Metasound/Source/MetasoundFrontend/Private/MetasoundFrontendNodeClassRegistryPrivate.h
+++ b/Engine/Plugins/Runtime/Metasound/Source/MetasoundFrontend/Private/MetasoundFrontendNodeClassRegistryPrivate.h
@@ -84,6 +84,8 @@ namespace Metasound::Frontend
 #if WITH_EDITORONLY_DATA
 		virtual bool RegisterNodeMigration(const FNodeMigrationInfo& InMigrationInfo) override;
 		virtual bool UnregisterNodeMigration(const FNodeMigrationInfo& InMigrationInfo) override;
+		virtual FName GetOwningPluginName(const FNodeClassRegistryKey& InKey) const override;
+		virtual FName GetOwningModuleName(const FNodeClassRegistryKey& InKey) const override;
 
 		virtual bool RegisterNodeUpdateTransform(const FNodeClassRegistryKey& InNodeUpdateTransformKey, const TSharedRef<const INodeUpdateTransform>& InNodeUpdateTransform) override;
 		virtual bool RegisterNodeUpdateTransform(const TArray<const FNodeClassRegistryKey>& InNodeUpdateTransformKeys, const TSharedRef<const INodeUpdateTransform>& InNodeUpdateTransform) override;
@@ -99,8 +101,10 @@ namespace Metasound::Frontend
 		virtual bool IsNodeRegistered(const FNodeClassRegistryKey& InKey) const override;
 		virtual bool IsGraphRegistered(const FGraphRegistryKey& InKey) const override;
 
-		virtual bool RegisterConversionNode(const FConverterNodeClassRegistryKey& InNodeKey, const FConverterNodeInfo& InNodeInfo) override;
+		virtual bool RegisterConversionNode(const FConverterNodeClassRegistryKey& InConversionKey, const FConverterNodeInfo& InNodeInfo) override;
+		virtual bool UnregisterConversionNode(const FConverterNodeClassRegistryKey& InConversionKey, const FNodeClassRegistryKey& InNodeKey) override;
 
+		UE_DEPRECATED(5.8, "Use Metasound::Frontend::ISearchEngine instead.")
 		virtual void IterateRegistry(FIterateMetasoundFrontendClassFunction InIterFunc, EMetasoundFrontendClassType InClassType = EMetasoundFrontendClassType::Invalid) const override;
 
 		// Find Frontend class data.
diff --git a/Engine/Plugins/Runtime/Metasound/Source/MetasoundFrontend/Private/MetasoundNodeRegistrationMacro.cpp b/Engine/Plugins/Runtime/Metasound/Source/MetasoundFrontend/Private/MetasoundNodeRegistrationMacro.cpp
--- a/Engine/Plugins/Runtime/Metasound/Source/MetasoundFrontend/Private/MetasoundNodeRegistrationMacro.cpp
+++ b/Engine/Plugins/Runtime/Metasound/Source/MetasoundFrontend/Private/MetasoundNodeRegistrationMacro.cpp
@@ -120,9 +120,30 @@ namespace Metasound::Frontend
 			return bSuccessfullyRegisteredNode;
 		}
 
-		bool UnregisterNodeInternal(const FNodeClassMetadata& InMetadata)
+	}
+
+	bool UnregisterNode(const FNodeClassRegistryKey& InRegistryKey, const FModuleInfo& InModuleInfo)
+	{
+		INodeClassRegistry& Registry = INodeClassRegistry::GetChecked();
+
+#if WITH_EDITORONLY_DATA && !UE_BUILD_SHIPPING
+		if (Registry.IsNodeRegistered(InRegistryKey))
 		{
-			return INodeClassRegistry::Get()->UnregisterNode(FNodeClassRegistryKey(InMetadata));
+			/** Check that the module info used to register the node matches the module
+			* info that is used to unregister the node. */
+			const FName ExpectedPlugin = Registry.GetOwningPluginName(InRegistryKey);
+			const FName ExpectedModule = Registry.GetOwningModuleName(InRegistryKey);
+			const bool bIsMatchingModuleInfo = (InModuleInfo.PluginName == ExpectedPlugin) && (InModuleInfo.ModuleName == ExpectedModule);
+
+			UE_CLOG(!bIsMatchingModuleInfo, LogMetaSound, Error, TEXT("Owning Module Info does not match between registration and unregistration of node %s. Expected unregistration to come from plugin %s and module %s, but encountered it from plugin %s and module %s"), *InRegistryKey.ToString(), *ExpectedPlugin.ToString(), *ExpectedModule.ToString(), *InModuleInfo.PluginName.ToString(), *InModuleInfo.ModuleName.ToString());
 		}
+#endif
+
+		return Registry.UnregisterNode(InRegistryKey);
+	}
+
+	bool UnregisterNode(const FNodeClassMetadata& InMetadata, const FModuleInfo& InModuleInfo)
+	{
+		return UnregisterNode(FNodeClassRegistryKey(InMetadata), InModuleInfo);
 	}
 }
diff --git a/Engine/Plugins/Runtime/Metasound/Source/MetasoundFrontend/Private/MetasoundReceiveNode.cpp b/Engine/Plugins/Runtime/Metasound/Source/MetasoundFrontend/Private/MetasoundReceiveNode.cpp
--- a/Engine/Plugins/Runtime/Metasound/Source/MetasoundFrontend/Private/MetasoundReceiveNode.cpp
+++ b/Engine/Plugins/Runtime/Metasound/Source/MetasoundFrontend/Private/MetasoundReceiveNode.cpp
@@ -37,6 +37,15 @@ namespace Metasound
 
 		int32 GetCurrentMinorVersion() { return 0; }
 	}
+
+	namespace ReceiveNodePrivate
+	{
+		Frontend::FNodeClassRegistryKey CreateNodeClassRegistryKey(const FName& InDataTypeName)
+		{
+			using namespace ReceiveNodeInfo;
+			return Frontend::FNodeClassRegistryKey{EMetasoundFrontendClassType::External, GetClassNameForDataType(InDataTypeName), GetCurrentMajorVersion(), GetCurrentMinorVersion()};
+		}
+	}
 }
 
 #undef LOCTEXT_NAMESPACE
diff --git a/Engine/Plugins/Runtime/Metasound/Source/MetasoundFrontend/Private/MetasoundSendNode.cpp b/Engine/Plugins/Runtime/Metasound/Source/MetasoundFrontend/Private/MetasoundSendNode.cpp
new file mode 100644
--- /dev/null
+++ b/Engine/Plugins/Runtime/Metasound/Source/MetasoundFrontend/Private/MetasoundSendNode.cpp
@@ -0,0 +1,71 @@
+// Copyright Epic Games, Inc. All Rights Reserved.
+
+#include "MetasoundSendNode.h"
+
+#define LOCTEXT_NAMESPACE "MetasoundFrontend"
+namespace Metasound::SendNodePrivate
+{
+	constexpr int32 SendNodeMajorVersion = 1;
+	constexpr int32 SendNodeMinorVersion = 0;
+
+	
+	FVertexInterface CreateVertexInterface(const FName& InDataTypeName)
+	{
+		using namespace SendVertexNames; 
+		static const FDataVertexMetadata AddressInputMetadata
+		{
+			  FText::GetEmpty() // description
+			, METASOUND_GET_PARAM_DISPLAYNAME(AddressInput) // display name
+		};
+
+		return FVertexInterface(
+			FInputVertexInterface(
+				TInputDataVertex<FSendAddress>(METASOUND_GET_PARAM_NAME(AddressInput), AddressInputMetadata),
+				FInputDataVertex(GetSendInputName(InDataTypeName), InDataTypeName, FDataVertexMetadata{ FText::GetEmpty() })
+			),
+			FOutputVertexInterface(
+			)
+		);
+	}
+
+	const FName& GetSendInputName(const FName& InDataTypeName)
+	{
+		return InDataTypeName;
+	}
+
+	FNodeClassName CreateNodeClassName(const FName& InDataTypeName)
+	{
+		return FNodeClassName{ "Send", InDataTypeName, FName() };
+	}
+
+	Frontend::FNodeClassRegistryKey CreateNodeClassRegistryKey(const FName& InDataTypeName)
+	{
+		return Frontend::FNodeClassRegistryKey{EMetasoundFrontendClassType::External, CreateNodeClassName(InDataTypeName), SendNodeMajorVersion, SendNodeMinorVersion};
+	}
+
+	FNodeClassMetadata CreateNodeClassMetadata(const FName& InDataTypeName, const FText& InDataTypeDisplayText)
+	{
+		const FVertexName& InputName = GetSendInputName(InDataTypeName);
+		FNodeClassMetadata Info;
+
+		Info.ClassName = CreateNodeClassName(InDataTypeName);
+		Info.MajorVersion = SendNodeMajorVersion;
+		Info.MinorVersion = SendNodeMinorVersion;
+		Info.DisplayName = METASOUND_LOCTEXT_FORMAT("Metasound_SendNodeDisplayNameFormat", "Send {0}", InDataTypeDisplayText);
+		Info.Description = METASOUND_LOCTEXT("Metasound_SendNodeDescription", "Sends data from a send node with the same name.");
+		Info.Author = PluginAuthor;
+		Info.PromptIfMissing = PluginNodeMissingPrompt;
+		Info.DefaultInterface = CreateVertexInterface(InDataTypeName);
+		Info.CategoryHierarchy = { METASOUND_LOCTEXT("Metasound_TransmissionNodeCategory", "Transmission") };
+		Info.Keywords = { };
+
+		// Then send & receive nodes do not work as expected, particularly 
+		// around multiple-consumer scenarios. Deprecate them to avoid
+		// metasound assets from relying on send & receive nodes. 
+		EnumAddFlags(Info.AccessFlags, ENodeClassAccessFlags::Deprecated);
+
+		return Info;
+	}
+} // namespace Metasound::SendNodePrivate
+
+#undef LOCTEXT_NAMESPACE
diff --git a/Engine/Plugins/Runtime/Metasound/Source/MetasoundFrontend/Public/MetasoundArrayNodes.h b/Engine/Plugins/Runtime/Metasound/Source/MetasoundFrontend/Public/MetasoundArrayNodes.h
--- a/Engine/Plugins/Runtime/Metasound/Source/MetasoundFrontend/Public/MetasoundArrayNodes.h
+++ b/Engine/Plugins/Runtime/Metasound/Source/MetasoundFrontend/Public/MetasoundArrayNodes.h
@@ -26,6 +26,7 @@ namespace Metasound
 	namespace MetasoundArrayNodesPrivate
 	{
 		// Convenience function for make FNodeClassMetadata of array nodes.
+		METASOUNDFRONTEND_API FNodeClassName CreateArrayNodeClassName(const FName& InOperatorName, const FName& InDataTypeName);
 		METASOUNDFRONTEND_API FNodeClassMetadata CreateArrayNodeClassMetadata(const FName& InDataTypeName, const FName& InOperatorName, const FText& InDisplayName, const FText& InDescription, const FVertexInterface& InDefaultInterface, int32 MajorVersion=1, int32 MinorVersion=0, bool bIsDeprecated=false);
 
 		METASOUNDFRONTEND_API FVertexInterface CreateArrayNumInterface(const FName& InArrayDataTypeName);
@@ -35,6 +36,20 @@ namespace Metasound
 		METASOUNDFRONTEND_API FVertexInterface CreateArrayConcatInterface(const FName& InArrayDataTypeName);
 		METASOUNDFRONTEND_API FVertexInterface CreateArrayLastIndexInterface(const FName& InArrayDataTypeName);
 
+		METASOUNDFRONTEND_API Frontend::FNodeClassRegistryKey CreateArrayNumNodeClassRegistryKey(const FName& InArrayDataTypeName);
+		METASOUNDFRONTEND_API Frontend::FNodeClassRegistryKey CreateArrayGetNodeClassRegistryKey(const FName& InArrayDataTypeName);
+		METASOUNDFRONTEND_API Frontend::FNodeClassRegistryKey CreateArraySetNodeClassRegistryKey(const FName& InArrayDataTypeName);
+		METASOUNDFRONTEND_API Frontend::FNodeClassRegistryKey CreateArraySubsetNodeClassRegistryKey(const FName& InArrayDataTypeName);
+		METASOUNDFRONTEND_API Frontend::FNodeClassRegistryKey CreateArrayConcatNodeClassRegistryKey(const FName& InArrayDataTypeName);
+		METASOUNDFRONTEND_API Frontend::FNodeClassRegistryKey CreateArrayLastIndexNodeClassRegistryKey(const FName& InArrayDataTypeName);
+
+		METASOUNDFRONTEND_API FNodeClassMetadata CreateArrayNumNodeClassMetadata(const FName& InArrayDataTypeName, const FText& InArrayDataTypeDisplayText);
+		METASOUNDFRONTEND_API FNodeClassMetadata CreateArrayGetNodeClassMetadata(const FName& InArrayDataTypeName, const FName& InElementDataTypeName, const FText& InArrayDataTypeDisplayText);
+		METASOUNDFRONTEND_API FNodeClassMetadata CreateArraySetNodeClassMetadata(const FName& InArrayDataTypeName, const FName& InElementDataTypeName, const FText& InArrayDataTypeDisplayText);
+		METASOUNDFRONTEND_API FNodeClassMetadata CreateArraySubsetNodeClassMetadata(const FName& InArrayDataTypeName, const FText& InArrayDataTypeDisplayText);
+		METASOUNDFRONTEND_API FNodeClassMetadata CreateArrayConcatNodeClassMetadata(const FName& InArrayDataTypeName, const FText& InArrayDataTypeDisplayText);
+		METASOUNDFRONTEND_API FNodeClassMetadata CreateArrayLastIndexNodeClassMetadata(const FName& InArrayDataTypeName, const FText& InArrayDataTypeDisplayText);
+
 		// Retrieve the ElementType from an ArrayType
 		template<typename ArrayType>
 		struct TArrayElementType
@@ -86,23 +101,9 @@ namespace Metasound
 
 		static const FNodeClassMetadata& GetNodeInfo()
 		{
-			auto CreateNodeClassMetadata = []() -> FNodeClassMetadata
-			{
-				const FName DataTypeName = GetMetasoundDataTypeName<ArrayType>();
-				const FName OperatorName = TEXT("Num");
-				const FText NodeDisplayName = METASOUND_LOCTEXT_FORMAT("ArrayOpArrayNumDisplayNamePattern", "Num ({0})", GetMetasoundDataTypeDisplayText<ArrayType>());
-				const FText NodeDescription = METASOUND_LOCTEXT("ArrayOpArrayNumDescription", "Number of elements in the array");
-				const FVertexInterface NodeInterface = MetasoundArrayNodesPrivate::CreateArrayNumInterface(DataTypeName);
-
-				FNodeClassMetadata NodeClassMetadata = MetasoundArrayNodesPrivate::CreateArrayNodeClassMetadata(DataTypeName, OperatorName, NodeDisplayName, NodeDescription, NodeInterface);
-				NodeClassMetadata.Keywords.Append({
-					METASOUND_LOCTEXT("ArrayNumKeyword_Length", "length"),
-					METASOUND_LOCTEXT("ArrayNumKeyword_Size", "size"),
-				});
-				return NodeClassMetadata;
-			};
-			
-			static const FNodeClassMetadata Metadata = CreateNodeClassMetadata();
+			using namespace MetasoundArrayNodesPrivate;
+
+			static const FNodeClassMetadata Metadata = MetasoundArrayNodesPrivate::CreateArrayNumNodeClassMetadata(GetMetasoundDataTypeName<ArrayType>(), GetMetasoundDataTypeDisplayText<ArrayType>());
 
 			return Metadata;
 		}
@@ -176,17 +177,9 @@ namespace Metasound
 
 		static const FNodeClassMetadata& GetNodeInfo()
 		{
-			auto CreateNodeClassMetadata = []() -> FNodeClassMetadata
-			{
-				const FName DataTypeName = GetMetasoundDataTypeName<ArrayType>();
-				const FName OperatorName = TEXT("Get"); 
-				const FText NodeDisplayName = METASOUND_LOCTEXT_FORMAT("ArrayOpArrayGetDisplayNamePattern", "Get ({0})", GetMetasoundDataTypeDisplayText<ArrayType>());
-				const FText NodeDescription = METASOUND_LOCTEXT("ArrayOpArrayGetDescription", "Get element at index in array.");
-				const FVertexInterface NodeInterface = MetasoundArrayNodesPrivate::CreateArrayGetInterface(DataTypeName, GetMetasoundDataTypeName<ElementType>());
-				return MetasoundArrayNodesPrivate::CreateArrayNodeClassMetadata(DataTypeName, OperatorName, NodeDisplayName, NodeDescription, NodeInterface);
-			};
-			
-			static const FNodeClassMetadata Metadata = CreateNodeClassMetadata();
+			using namespace MetasoundArrayNodesPrivate;
+
+			static FNodeClassMetadata Metadata = CreateArrayGetNodeClassMetadata(GetMetasoundDataTypeName<ArrayType>(), GetMetasoundDataTypeName<ElementType>(), GetMetasoundDataTypeDisplayText<ArrayType>());
 
 			return Metadata;
 		}
@@ -338,18 +331,9 @@ namespace Metasound
 
 		static const FNodeClassMetadata& GetNodeInfo()
 		{
-			auto CreateNodeClassMetadata = []() -> FNodeClassMetadata
-			{
-				const FName DataTypeName = GetMetasoundDataTypeName<ArrayType>();
-				const FName OperatorName = TEXT("Set"); 
-				const FText NodeDisplayName = METASOUND_LOCTEXT_FORMAT("ArrayOpArraySetDisplayNamePattern", "Set ({0})", GetMetasoundDataTypeDisplayText<ArrayType>());
-				const FText NodeDescription = METASOUND_LOCTEXT("ArrayOpArraySetDescription", "Set element at index in array.");
-				const FVertexInterface NodeInterface = MetasoundArrayNodesPrivate::CreateArraySetInterface(DataTypeName, GetMetasoundDataTypeName<ElementType>());
-			
-				return MetasoundArrayNodesPrivate::CreateArrayNodeClassMetadata(DataTypeName, OperatorName, NodeDisplayName, NodeDescription, NodeInterface);
-			};
-			
-			static const FNodeClassMetadata Metadata = CreateNodeClassMetadata();
+			using namespace MetasoundArrayNodesPrivate;
+
+			static FNodeClassMetadata Metadata = CreateArraySetNodeClassMetadata(GetMetasoundDataTypeName<ArrayType>(), GetMetasoundDataTypeName<ElementType>(), GetMetasoundDataTypeDisplayText<ArrayType>());
 
 			return Metadata;
 		}
@@ -490,22 +474,9 @@ namespace Metasound
 
 		static const FNodeClassMetadata& GetNodeInfo()
 		{
-			auto CreateNodeClassMetadata = []() -> FNodeClassMetadata
-			{
-				const FName DataTypeName = GetMetasoundDataTypeName<ArrayType>();
-				const FName OperatorName = TEXT("Concat"); 
-				const FText NodeDisplayName = METASOUND_LOCTEXT_FORMAT("ArrayOpArrayConcatDisplayNamePattern", "Concatenate ({0})", GetMetasoundDataTypeDisplayText<ArrayType>());
-				const FText NodeDescription = METASOUND_LOCTEXT("ArrayOpArrayConcatDescription", "Concatenates two arrays on trigger.");
-				const FVertexInterface NodeInterface = MetasoundArrayNodesPrivate::CreateArrayConcatInterface(DataTypeName);
-			
-				FNodeClassMetadata NodeClassMetadata = MetasoundArrayNodesPrivate::CreateArrayNodeClassMetadata(DataTypeName, OperatorName, NodeDisplayName, NodeDescription, NodeInterface);
-				NodeClassMetadata.Keywords.Append({
-					METASOUND_LOCTEXT("ArrayConcatKeyword_Append", "append"),
-				});
-				return NodeClassMetadata;
-			};
-			
-			static const FNodeClassMetadata Metadata = CreateNodeClassMetadata();
+			using namespace MetasoundArrayNodesPrivate;
+
+			static const FNodeClassMetadata Metadata = CreateArrayConcatNodeClassMetadata(GetMetasoundDataTypeName<ArrayType>(), GetMetasoundDataTypeDisplayText<ArrayType>());
 
 			return Metadata;
 		}
@@ -588,22 +559,9 @@ namespace Metasound
 
 		static const FNodeClassMetadata& GetNodeInfo()
 		{
-			auto CreateNodeClassMetadata = []() -> FNodeClassMetadata
-			{
-				const FName DataTypeName = GetMetasoundDataTypeName<ArrayType>();
-				const FName OperatorName = TEXT("Subset"); 
-				const FText NodeDisplayName = METASOUND_LOCTEXT_FORMAT("ArrayOpArraySubsetDisplayNamePattern", "Subset ({0})", GetMetasoundDataTypeDisplayText<ArrayType>());
-				const FText NodeDescription = METASOUND_LOCTEXT("ArrayOpArraySubsetDescription", "Subset array on trigger.");
-				const FVertexInterface NodeInterface = MetasoundArrayNodesPrivate::CreateArraySubsetInterface(DataTypeName);
-			
-				FNodeClassMetadata NodeClassMetadata = MetasoundArrayNodesPrivate::CreateArrayNodeClassMetadata(DataTypeName, OperatorName, NodeDisplayName, NodeDescription, NodeInterface);
-				NodeClassMetadata.Keywords.Append({
-					METASOUND_LOCTEXT("ArraySubsetKeyword_Slice", "slice"),
-				});
-				return NodeClassMetadata;
-			};
-			
-			static const FNodeClassMetadata Metadata = CreateNodeClassMetadata();
+			using namespace MetasoundArrayNodesPrivate;
+
+			static const FNodeClassMetadata Metadata = CreateArraySubsetNodeClassMetadata(GetMetasoundDataTypeName<ArrayType>(), GetMetasoundDataTypeDisplayText<ArrayType>());
 
 			return Metadata;
 		}
@@ -700,18 +658,9 @@ namespace Metasound
 
 		static const FNodeClassMetadata& GetNodeInfo()
 		{
-			auto CreateNodeClassMetadata = []() -> FNodeClassMetadata
-				{
-					const FName DataTypeName = GetMetasoundDataTypeName<ArrayType>();
-					const FName OperatorName = TEXT("GetLastIndex");
-					const FText NodeDisplayName = METASOUND_LOCTEXT_FORMAT("ArrayOpArrayLastIndexDisplayNamePattern", "Get Last Index ({0})", GetMetasoundDataTypeDisplayText<ArrayType>());
-					const FText NodeDescription = METASOUND_LOCTEXT("ArrayOpArrayLastIndexDescription", "Last index of the array");
-					const FVertexInterface NodeInterface = MetasoundArrayNodesPrivate::CreateArrayLastIndexInterface(DataTypeName);
-
-					return MetasoundArrayNodesPrivate::CreateArrayNodeClassMetadata(DataTypeName, OperatorName, NodeDisplayName, NodeDescription, NodeInterface);
-				};
+			using namespace MetasoundArrayNodesPrivate;
 
-			static const FNodeClassMetadata Metadata = CreateNodeClassMetadata();
+			static const FNodeClassMetadata Metadata = CreateArrayLastIndexNodeClassMetadata(GetMetasoundDataTypeName<ArrayType>(), GetMetasoundDataTypeDisplayText<ArrayType>());
 
 			return Metadata;
 		}
diff --git a/Engine/Plugins/Runtime/Metasound/Source/MetasoundFrontend/Public/MetasoundArrayNodesRegistration.h b/Engine/Plugins/Runtime/Metasound/Source/MetasoundFrontend/Public/MetasoundArrayNodesRegistration.h
--- a/Engine/Plugins/Runtime/Metasound/Source/MetasoundFrontend/Public/MetasoundArrayNodesRegistration.h
+++ b/Engine/Plugins/Runtime/Metasound/Source/MetasoundFrontend/Public/MetasoundArrayNodesRegistration.h
@@ -5,6 +5,7 @@
 #include "MetasoundArrayNodes.h"
 #include "MetasoundArrayShuffleNode.h"
 #include "MetasoundArrayRandomNode.h"
+#include "MetasoundFrontendNodeClassRegistry.h"
 #include "MetasoundNodeRegistrationMacro.h"
 
 #include <type_traits>
@@ -39,6 +40,9 @@ namespace Metasound
 			// Array num is supported for all array types.
 			static constexpr bool bIsArrayNumSupported = bEnabled;
 
+			// Array last index is supported for all array types.
+			static constexpr bool bIsArrayLastIndexSupported = bEnabled;
+
 			// Element must be default parsable to create get operator because a
 			// value must be returned even if the index is invalid. Also values are
 			// assigned by copy.
@@ -59,125 +63,6 @@ namespace Metasound
 			// Random get is supported for all types that get is supported for.
 			static constexpr bool bIsArrayRandomGetSupported = bEnabled && bIsElementParsableAndAssignable;
 		};
-
-		template<typename ArrayType, typename std::enable_if<TArrayNodeSupport<ArrayType>::bIsArrayGetSupported, bool>::type = true>
-		bool RegisterArrayGetNode(const Frontend::FModuleInfo& InModuleInfo)
-		{
-			using FNodeType = typename Metasound::TArrayGetNode<ArrayType>;
-			return Frontend::RegisterNode<FNodeType>(InModuleInfo);
-		}
-
-		template<typename ArrayType, typename std::enable_if<!TArrayNodeSupport<ArrayType>::bIsArrayGetSupported, bool>::type = true>
-		bool RegisterArrayGetNode(const Frontend::FModuleInfo& InModuleInfo)
-		{
-			// No op if not supported
-			return true;
-		}
-
-		template<typename ArrayType, typename std::enable_if<TArrayNodeSupport<ArrayType>::bIsArraySetSupported, bool>::type = true>
-		bool RegisterArraySetNode(const Frontend::FModuleInfo& InModuleInfo)
-		{
-			using FNodeType = typename Metasound::TArraySetNode<ArrayType>;
-
-			static_assert(TArrayNodeSupport<ArrayType>::bIsArraySetSupported, "TArraySetNode<> is not supported by array type");
-
-			return Frontend::RegisterNode<FNodeType>(InModuleInfo);
-		}
-
-		template<typename ArrayType, typename std::enable_if<!TArrayNodeSupport<ArrayType>::bIsArraySetSupported, bool>::type = true>
-		bool RegisterArraySetNode(const Frontend::FModuleInfo& InModuleInfo)
-		{
-			// No op if not supported
-			return true;
-		}
-
-		template<typename ArrayType, typename std::enable_if<TArrayNodeSupport<ArrayType>::bIsArraySubsetSupported, bool>::type = true>
-		bool RegisterArraySubsetNode(const Frontend::FModuleInfo& InModuleInfo)
-		{
-			using FNodeType = typename Metasound::TArraySubsetNode<ArrayType>;
-
-			static_assert(TArrayNodeSupport<ArrayType>::bIsArraySubsetSupported, "TArraySubsetNode<> is not supported by array type");
-
-			return Frontend::RegisterNode<FNodeType>(InModuleInfo);
-		}
-
-		template<typename ArrayType, typename std::enable_if<!TArrayNodeSupport<ArrayType>::bIsArraySubsetSupported, bool>::type = true>
-		bool RegisterArraySubsetNode(const Frontend::FModuleInfo& InModuleInfo)
-		{
-			// No op if not supported
-			return true;
-		}
-
-		template<typename ArrayType, typename std::enable_if<TArrayNodeSupport<ArrayType>::bIsArrayConcatSupported, bool>::type = true>
-		bool RegisterArrayConcatNode(const Frontend::FModuleInfo& InModuleInfo)
-		{
-			using FNodeType = typename Metasound::TArrayConcatNode<ArrayType>;
-
-			static_assert(TArrayNodeSupport<ArrayType>::bIsArrayConcatSupported, "TArrayConcatNode<> is not supported by array type");
-
-			return Frontend::RegisterNode<FNodeType>(InModuleInfo);
-		}
-
-		template<typename ArrayType, typename std::enable_if<!TArrayNodeSupport<ArrayType>::bIsArrayConcatSupported, bool>::type = true>
-		bool RegisterArrayConcatNode(const Frontend::FModuleInfo& InModuleInfo)
-		{
-			// No op if not supported
-			return true;
-		}
-
-		template<typename ArrayType, typename std::enable_if<TArrayNodeSupport<ArrayType>::bIsArrayNumSupported, bool>::type = true>
-		bool RegisterArrayNumNode(const Frontend::FModuleInfo& InModuleInfo)
-		{
-			return Frontend::RegisterNode<Metasound::TArrayNumNode<ArrayType>>(InModuleInfo);
-		}
-
-		template<typename ArrayType, typename std::enable_if<!TArrayNodeSupport<ArrayType>::bIsArrayNumSupported, bool>::type = true>
-		bool RegisterArrayNumNode(const Frontend::FModuleInfo& InModuleInfo)
-		{
-			// No op if not supported
-			return true;
-		}
-
-		template<typename ArrayType, typename std::enable_if<TArrayNodeSupport<ArrayType>::bIsArrayShuffleSupported, bool>::type = true>
-		bool RegisterArrayShuffleNode(const Frontend::FModuleInfo& InModuleInfo)
-		{
-			using FNodeType = typename Metasound::TArrayShuffleNode<ArrayType>;
-			return Frontend::RegisterNode<FNodeType>(InModuleInfo);
-		}
-
-		template<typename ArrayType, typename std::enable_if<!TArrayNodeSupport<ArrayType>::bIsArrayShuffleSupported, bool>::type = true>
-		bool RegisterArrayShuffleNode(const Frontend::FModuleInfo& InModuleInfo)
-		{
-			// No op if not supported
-			return true;
-		}
-
-		template<typename ArrayType, typename std::enable_if<TArrayNodeSupport<ArrayType>::bIsArrayRandomGetSupported, bool>::type = true>
-		bool RegisterArrayRandomGetNode(const Frontend::FModuleInfo& InModuleInfo)
-		{
-			using FNodeType = typename Metasound::TArrayRandomGetNode<ArrayType>;
-			return Frontend::RegisterNode<FNodeType>(InModuleInfo);
-		}
-
-		template<typename ArrayType, typename std::enable_if<!TArrayNodeSupport<ArrayType>::bIsArrayRandomGetSupported, bool>::type = true>
-		bool RegisterArrayRandomGetNode(const Frontend::FModuleInfo& InModuleInfo)
-		{
-			// No op if not supported
-			return true;
-		}
-
-		template<typename ArrayType, typename std::enable_if<TArrayNodeSupport<ArrayType>::bIsArrayNumSupported, bool>::type = true>
-		bool RegisterArrayLastIndexNode(const Frontend::FModuleInfo& InModuleInfo)
-		{
-			return Frontend::RegisterNode<Metasound::TArrayLastIndexNode<ArrayType>>(InModuleInfo);
-		}
-
-		template<typename ArrayType, typename std::enable_if<!TArrayNodeSupport<ArrayType>::bIsArrayNumSupported, bool>::type = true>
-		bool RegisterArrayLastIndexNode(const Frontend::FModuleInfo& InModuleInfo)
-		{
-			// No op if not supported
-			return true;
-		}
 	}
 
 	namespace Frontend
@@ -191,23 +76,59 @@ namespace Metasound
 		{
 			using namespace MetasoundArrayNodesPrivate;
 
-			bool bSuccess = RegisterArrayNumNode<ArrayType>(InModuleInfo);
-			bSuccess = bSuccess && RegisterArrayGetNode<ArrayType>(InModuleInfo);
-			bSuccess = bSuccess && RegisterArraySetNode<ArrayType>(InModuleInfo);
-			bSuccess = bSuccess && RegisterArraySubsetNode<ArrayType>(InModuleInfo);
-			bSuccess = bSuccess && RegisterArrayConcatNode<ArrayType>(InModuleInfo);
-			bSuccess = bSuccess && RegisterArrayShuffleNode<ArrayType>(InModuleInfo);
-			bSuccess = bSuccess && RegisterArrayRandomGetNode<ArrayType>(InModuleInfo);
-			bSuccess = bSuccess && RegisterArrayLastIndexNode<ArrayType>(InModuleInfo);
+			bool bSuccess = true;
+
+			if constexpr (TArrayNodeSupport<ArrayType>::bIsArrayNumSupported)
+			{
+				bSuccess = bSuccess && Frontend::RegisterNode<Metasound::TArrayNumNode<ArrayType>>(InModuleInfo);
+			}
+
+			if constexpr (TArrayNodeSupport<ArrayType>::bIsArrayGetSupported)
+			{
+				using FGetNodeType = typename Metasound::TArrayGetNode<ArrayType>;
+				bSuccess = bSuccess && Frontend::RegisterNode<FGetNodeType>(InModuleInfo);
+			}
+			
+			if constexpr (TArrayNodeSupport<ArrayType>::bIsArraySetSupported)
+			{
+				using FSetNodeType = typename Metasound::TArraySetNode<ArrayType>;
+				bSuccess = bSuccess && Frontend::RegisterNode<FSetNodeType>(InModuleInfo);
+			}
+
+			if constexpr (TArrayNodeSupport<ArrayType>::bIsArraySubsetSupported)
+			{
+				using FSubsetNodeType = typename Metasound::TArraySubsetNode<ArrayType>;
+				bSuccess = bSuccess && Frontend::RegisterNode<FSubsetNodeType>(InModuleInfo);
+			}
+
+			if constexpr (TArrayNodeSupport<ArrayType>::bIsArrayConcatSupported)
+			{
+				using FConcatNodeType = typename Metasound::TArrayConcatNode<ArrayType>;
+				bSuccess = bSuccess && Frontend::RegisterNode<FConcatNodeType>(InModuleInfo);
+			}
+
+			if constexpr (TArrayNodeSupport<ArrayType>::bIsArrayShuffleSupported)
+			{
+				using FShuffleNodeType = typename Metasound::TArrayShuffleNode<ArrayType>;
+				bSuccess = bSuccess && Frontend::RegisterNode<FShuffleNodeType>(InModuleInfo);
+			}
+
+			if constexpr (TArrayNodeSupport<ArrayType>::bIsArrayRandomGetSupported)
+			{
+				using FRandomGetNodeType = typename Metasound::TArrayRandomGetNode<ArrayType>;
+				bSuccess = bSuccess && Frontend::RegisterNode<FRandomGetNodeType>(InModuleInfo);
+			}
+
+			if constexpr (TArrayNodeSupport<ArrayType>::bIsArrayLastIndexSupported)
+			{
+				bSuccess = bSuccess && Frontend::RegisterNode<Metasound::TArrayLastIndexNode<ArrayType>>(InModuleInfo);
+			}
+
 			return bSuccess;
 		}
-	}
 
-	template<typename ArrayType>
-	UE_DEPRECATED(5.6, "Use Frontend::RegisterArrayNodes() instead")
-	bool RegisterArrayNodes()
-	{
-		return Frontend::RegisterArrayNodes<ArrayType>();
+		UE_INTERNAL
+		METASOUNDFRONTEND_API bool UnregisterArrayNodes(const FName& InArrayDataTypeName, const FModuleInfo& InModuleInfo);
 	}
 }
 
diff --git a/Engine/Plugins/Runtime/Metasound/Source/MetasoundFrontend/Public/MetasoundArrayRandomNode.h b/Engine/Plugins/Runtime/Metasound/Source/MetasoundFrontend/Public/MetasoundArrayRandomNode.h
--- a/Engine/Plugins/Runtime/Metasound/Source/MetasoundFrontend/Public/MetasoundArrayRandomNode.h
+++ b/Engine/Plugins/Runtime/Metasound/Source/MetasoundFrontend/Public/MetasoundArrayRandomNode.h
@@ -74,6 +74,13 @@ namespace Metasound
 		METASOUND_PARAM(OutputIndex, "Index", "Array index of the current random element.")
 	}
 
+	namespace ArrayRandomNodePrivate
+	{
+		METASOUNDFRONTEND_API Frontend::FNodeClassRegistryKey CreateArrayRandomNodeClassRegistryKey(const FName& InArrayDataTypeName);
+		METASOUNDFRONTEND_API FVertexInterface CreateArrayRandomNodeInterface(const FName& InArrayDataTypeName, const FName& InElementDataTypeName);
+		METASOUNDFRONTEND_API FNodeClassMetadata CreateArrayRandomNodeClassMetadata(const FName& InArrayDataTypeName, const FName& InElementDataTypeName, const FText& InArrayDataTypeDisplayText);
+	}
+
 	enum class ESharedStateBehaviorType : int32
 	{
 		SameNode,
@@ -176,43 +183,17 @@ namespace Metasound
 		static const FVertexInterface& GetDefaultInterface()
 		{
 			using namespace ArrayNodeRandomGetVertexNames;
+			using namespace ArrayRandomNodePrivate;
 
-			static const FVertexInterface DefaultInterface(
-				FInputVertexInterface(
-					TInputDataVertex<FTrigger>(METASOUND_GET_PARAM_NAME_AND_METADATA(InputTriggerNextValue)),
-					TInputDataVertex<FTrigger>(METASOUND_GET_PARAM_NAME_AND_METADATA(InputTriggerResetSeed)),
-					TInputDataVertex<ArrayType>(METASOUND_GET_PARAM_NAME_AND_METADATA(InputRandomArray)),
-					TInputDataVertex<WeightsArrayType>(METASOUND_GET_PARAM_NAME_AND_METADATA(InputWeights)),
-					TInputDataVertex<int32>(METASOUND_GET_PARAM_NAME_AND_METADATA(InputSeed), -1),
-					TInputDataVertex<int32>(METASOUND_GET_PARAM_NAME_AND_METADATA(InputNoRepeatOrder), 1),
-					TInputDataVertex<bool>(METASOUND_GET_PARAM_NAME_AND_METADATA(InputEnableSharedState), false),
-					TInputConstructorVertex<FEnumSharedStateBehaviorType>(METASOUND_GET_PARAM_NAME_AND_METADATA(InputSharedStateBehavior), (int32)ESharedStateBehaviorType::SameNodeInComposition)
-				),
-				FOutputVertexInterface(
-					TOutputDataVertex<FTrigger>(METASOUND_GET_PARAM_NAME_AND_METADATA(OutputTriggerOnNext)),
-					TOutputDataVertex<FTrigger>(METASOUND_GET_PARAM_NAME_AND_METADATA(OutputTriggerOnReset)),
-					TOutputDataVertex<ElementType>(METASOUND_GET_PARAM_NAME_AND_METADATA(ShuffleOutputValue)),
-					TOutputDataVertex<int32>(METASOUND_GET_PARAM_NAME_AND_METADATA(OutputIndex))
-				)
-			);
+			static const FVertexInterface DefaultInterface = CreateArrayRandomNodeInterface(GetMetasoundDataTypeName<ArrayType>,GetMetasoundDataTypeName<ElementType>());
 
 			return DefaultInterface;
 		}
 
 		static const FNodeClassMetadata& GetNodeInfo()
 		{
-			auto CreateNodeClassMetadata = []() -> FNodeClassMetadata
-			{
-				FName DataTypeName = GetMetasoundDataTypeName<ArrayType>();
-				FName OperatorName = "Random Get";
-				FText NodeDisplayName = METASOUND_LOCTEXT_FORMAT("RandomArrayGetNode_OpDisplayNamePattern", "Random Get ({0})", GetMetasoundDataTypeDisplayText<ArrayType>());
-				FText NodeDescription = METASOUND_LOCTEXT("RandomArrayGetNode_Description", "Randomly retrieve data from input array using the supplied weights.");
-				FVertexInterface NodeInterface = GetDefaultInterface();
-
-				return MetasoundArrayNodesPrivate::CreateArrayNodeClassMetadata(DataTypeName, OperatorName, NodeDisplayName, NodeDescription, NodeInterface, /*MajorVersion=*/1, /*MinorVersion=*/1);
-			};
-
-			static const FNodeClassMetadata Metadata = CreateNodeClassMetadata();
+			using namespace ArrayRandomNodePrivate;
+			static const FNodeClassMetadata Metadata = CreateArrayRandomNodeClassMetadata(GetMetasoundDataTypeName<ArrayType>(), GetMetasoundDataTypeName<ElementType>(), GetMetasoundDataTypeDisplayText<ArrayType>());
 
 			return Metadata;
 		}
@@ -287,22 +268,6 @@ namespace Metasound
 			InOutVertexData.BindReadVertex(METASOUND_GET_PARAM_NAME(OutputIndex), OutIndex);
 		}
 
-		virtual FDataReferenceCollection GetInputs() const override
-		{
-			// This should never be called. Bind(...) is called instead. This method
-			// exists as a stop-gap until the API can be deprecated and removed.
-			checkNoEntry();
-			return {};
-		}
-
-		virtual FDataReferenceCollection GetOutputs() const override
-		{
-			// This should never be called. Bind(...) is called instead. This method
-			// exists as a stop-gap until the API can be deprecated and removed.
-			checkNoEntry();
-			return {};
-		}
-
 		void Reset(const IOperator::FResetParams& InParams)
 		{
 			using namespace Frontend;
diff --git a/Engine/Plugins/Runtime/Metasound/Source/MetasoundFrontend/Public/MetasoundArrayShuffleNode.h b/Engine/Plugins/Runtime/Metasound/Source/MetasoundFrontend/Public/MetasoundArrayShuffleNode.h
--- a/Engine/Plugins/Runtime/Metasound/Source/MetasoundFrontend/Public/MetasoundArrayShuffleNode.h
+++ b/Engine/Plugins/Runtime/Metasound/Source/MetasoundFrontend/Public/MetasoundArrayShuffleNode.h
@@ -44,6 +44,13 @@ namespace Metasound
 		METASOUND_PARAM(OutputShuffledValue, "Value", "Value of the current shuffled element.")
 	}
 
+	namespace ArrayShuffleNodePrivate
+	{
+		UE_API Frontend::FNodeClassRegistryKey CreateArrayShuffleNodeClassRegistryKey(const FName& InArrayDataTypeName);
+		UE_API FVertexInterface CreateArrayShuffleNodeInterface(const FName& InArrayDataTypeName, const FName& InElementDataTypeName);
+		UE_API FNodeClassMetadata CreateArrayShuffleNodeClassMetadata(const FName& InArrayDataTypeName, const FName& InElementDataTypeName, const FText& InArrayDataTypeDisplayText);
+	}
+
 	class FArrayIndexShuffler
 	{
 	public:
@@ -145,42 +152,17 @@ namespace Metasound
 		static const FVertexInterface& GetDefaultInterface()
 		{
 			using namespace ArrayNodeShuffleVertexNames;
+			using namespace ArrayShuffleNodePrivate;
 
-			static const FVertexInterface DefaultInterface(
-				FInputVertexInterface(
-					TInputDataVertex<FTrigger>(METASOUND_GET_PARAM_NAME_AND_METADATA(InputTriggerNext)),
-					TInputDataVertex<FTrigger>(METASOUND_GET_PARAM_NAME_AND_METADATA(InputTriggerShuffle)),
-					TInputDataVertex<FTrigger>(METASOUND_GET_PARAM_NAME_AND_METADATA(InputTriggerReset)),
-					TInputDataVertex<ArrayType>(METASOUND_GET_PARAM_NAME_AND_METADATA(InputShuffleArray)),
-					TInputDataVertex<int32>(METASOUND_GET_PARAM_NAME_AND_METADATA(InputShuffleSeed), -1),
-					TInputDataVertex<bool>(METASOUND_GET_PARAM_NAME_AND_METADATA(InputAutoShuffle), true),
-					TInputDataVertex<bool>(METASOUND_GET_PARAM_NAME_AND_METADATA(InputShuffleEnableSharedState), false)
-					),
-				FOutputVertexInterface(
-					TOutputDataVertex<FTrigger>(METASOUND_GET_PARAM_NAME_AND_METADATA(OutputTriggerOnNext)),
-					TOutputDataVertex<FTrigger>(METASOUND_GET_PARAM_NAME_AND_METADATA(OutputTriggerOnShuffle)),
-					TOutputDataVertex<FTrigger>(METASOUND_GET_PARAM_NAME_AND_METADATA(OutputTriggerOnResetSeed)),
-					TOutputDataVertex<ElementType>(METASOUND_GET_PARAM_NAME_AND_METADATA(OutputShuffledValue))
-				)
-			);
+			static const FVertexInterface DefaultInterface = CreateArrayShuffleNodeInterface(GetMetasoundDataTypeName<ArrayType>,GetMetasoundDataTypeName<ElementType>());
 
 			return DefaultInterface;
 		}
 
 		static const FNodeClassMetadata& GetNodeInfo()
 		{
-			auto CreateNodeClassMetadata = []() -> FNodeClassMetadata
-			{
-				const FName DataTypeName = GetMetasoundDataTypeName<ArrayType>();
-				const FName OperatorName = "Shuffle";
-				const FText NodeDisplayName = METASOUND_LOCTEXT_FORMAT("ArrayOpArrayShuffleDisplayNamePattern", "Shuffle ({0})", GetMetasoundDataTypeDisplayText<ArrayType>());
-				const FText NodeDescription = METASOUND_LOCTEXT("ArrayOpArrayShuffleDescription", "Output next element of a shuffled array on trigger.");
-				const FVertexInterface NodeInterface = GetDefaultInterface();
-
-				return MetasoundArrayNodesPrivate::CreateArrayNodeClassMetadata(DataTypeName, OperatorName, NodeDisplayName, NodeDescription, NodeInterface);
-			};
-
-			static const FNodeClassMetadata Metadata = CreateNodeClassMetadata();
+			using namespace ArrayShuffleNodePrivate;
+			static const FNodeClassMetadata Metadata = CreateArrayShuffleNodeClassMetadata(GetMetasoundDataTypeName<ArrayType>(), GetMetasoundDataTypeName<ElementType>(), GetMetasoundDataTypeDisplayText<ArrayType>());
 
 			return Metadata;
 		}
@@ -250,22 +232,6 @@ namespace Metasound
 			InOutVertexData.BindReadVertex(METASOUND_GET_PARAM_NAME(OutputShuffledValue), OutValue);
 		}
 
-		virtual FDataReferenceCollection GetInputs() const override
-		{
-			// This should never be called. Bind(...) is called instead. This method
-			// exists as a stop-gap until the API can be deprecated and removed.
-			checkNoEntry();
-			return {};
-		}
-
-		virtual FDataReferenceCollection GetOutputs() const override
-		{
-			// This should never be called. Bind(...) is called instead. This method
-			// exists as a stop-gap until the API can be deprecated and removed.
-			checkNoEntry();
-			return {};
-		}
-
 		void Reset(const IOperator::FResetParams& InParams)
 		{
 			using namespace Frontend;
diff --git a/Engine/Plugins/Runtime/Metasound/Source/MetasoundFrontend/Public/MetasoundAutoConverterNode.h b/Engine/Plugins/Runtime/Metasound/Source/MetasoundFrontend/Public/MetasoundAutoConverterNode.h
--- a/Engine/Plugins/Runtime/Metasound/Source/MetasoundFrontend/Public/MetasoundAutoConverterNode.h
+++ b/Engine/Plugins/Runtime/Metasound/Source/MetasoundFrontend/Public/MetasoundAutoConverterNode.h
@@ -37,7 +37,23 @@ namespace Metasound
 
 		METASOUNDFRONTEND_API FVertexInterface CreateVertexInterface(const FConvertDataTypeInfo& InInfo);
 
+		METASOUNDFRONTEND_API FNodeClassName CreateAutoConverterNodeClassName(const FName& InFromDataType, const FName& InToDataType);
+		METASOUNDFRONTEND_API Frontend::FNodeClassRegistryKey CreateAutoConverterNodeClassRegistryKey(const FName& InFromDataType, const FName& InToDataType);
 		METASOUNDFRONTEND_API FNodeClassMetadata CreateAutoConverterNodeMetadata(const FConvertDataTypeInfo& InInfo);
+
+		template<typename FromDataType, typename ToDataType>
+		FConvertDataTypeInfo MakeConvertDataTypeInfo()
+		{
+			return FConvertDataTypeInfo
+			{
+				GetMetasoundDataTypeName<FromDataType>(),
+				GetMetasoundDataTypeDisplayText<FromDataType>(),
+				GetMetasoundDataTypeName<ToDataType>(),
+				GetMetasoundDataTypeDisplayText<ToDataType>(),
+				TEnumTraits<FromDataType>::bIsEnum,
+				TEnumTraits<ToDataType>::bIsEnum
+			};
+		}
 	}
 
 	// Determines whether an auto converter node will be registered to convert 
@@ -61,19 +77,6 @@ namespace Metasound
 	{
 		static_assert(TIsAutoConvertible<FromDataType, ToDataType>::Value,
 		"Tried to create an auto converter node between two types we can't static_cast between.");
-
-		static AutoConverterNodePrivate::FConvertDataTypeInfo GetConverterDataTypeInfo()
-		{
-			return AutoConverterNodePrivate::FConvertDataTypeInfo
-			{
-				GetMetasoundDataTypeName<FromDataType>(),
-				GetMetasoundDataTypeDisplayText<FromDataType>(),
-				GetMetasoundDataTypeName<ToDataType>(),
-				GetMetasoundDataTypeDisplayText<ToDataType>(),
-				TEnumTraits<FromDataType>::bIsEnum,
-				TEnumTraits<ToDataType>::bIsEnum
-			};
-		}
 		
 	public:
 		static const FVertexName& GetInputName()
@@ -91,13 +94,13 @@ namespace Metasound
 		static FVertexInterface DeclareVertexInterface()
 		{
 			using namespace AutoConverterNodePrivate;
-			return CreateVertexInterface(GetConverterDataTypeInfo());
+			return CreateVertexInterface(MakeConvertDataTypeInfo<FromDataType, ToDataType>());
 		}
 
 		static const FNodeClassMetadata& GetAutoConverterNodeMetadata()
 		{
 			using namespace AutoConverterNodePrivate;
-			static const FNodeClassMetadata Info = CreateAutoConverterNodeMetadata(GetConverterDataTypeInfo());
+			static const FNodeClassMetadata Info = CreateAutoConverterNodeMetadata(MakeConvertDataTypeInfo<FromDataType, ToDataType>());
 
 			return Info;
 		}
diff --git a/Engine/Plugins/Runtime/Metasound/Source/MetasoundFrontend/Public/MetasoundConverterNodeRegistrationMacro.h b/Engine/Plugins/Runtime/Metasound/Source/MetasoundFrontend/Public/MetasoundConverterNodeRegistrationMacro.h
--- a/Engine/Plugins/Runtime/Metasound/Source/MetasoundFrontend/Public/MetasoundConverterNodeRegistrationMacro.h
+++ b/Engine/Plugins/Runtime/Metasound/Source/MetasoundFrontend/Public/MetasoundConverterNodeRegistrationMacro.h
@@ -15,6 +15,7 @@
 #include "MetasoundVertex.h"
 
 template<typename TNodeType, typename TFromType, typename TToType>
+UE_DEPRECATED(5.8, "Converter node can be registered directly with the INodeClassRegistry in the owning modules startup/shutdown methods")
 bool RegisterConversionNode(const Metasound::FVertexName& FromPin, const Metasound::FVertexName& ToPin, const Metasound::FNodeClassMetadata& InNodeMetadata)
 {
 	using FConverterNodeRegistryKey = ::Metasound::Frontend::FConverterNodeRegistryKey;
@@ -25,7 +26,7 @@ bool RegisterConversionNode(const Metasound::FVertexName& FromPin, const Metasou
 	static bool bAlreadyRegisteredThisConverter = false;
 	if (bAlreadyRegisteredThisConverter)
 	{
-		UE_LOG(LogTemp, Display, TEXT("Tried to call REGISTER_METASOUND_CONVERTER twice with the same class. ignoring the second call. Likely because REGISTER_METASOUND_CONVERTER is in a header that's used in multiple modules. Consider moving it to a private header or cpp file."))
+		UE_LOG(LogMetaSound, Error, TEXT("Tried to call REGISTER_METASOUND_CONVERTER twice with the same class. ignoring the second call. Likely because REGISTER_METASOUND_CONVERTER is in a header that's used in multiple modules. Consider moving it to a private header or cpp file."))
 		return false;
 	}
 
@@ -51,10 +52,10 @@ bool RegisterConversionNode(const Metasound::FVertexName& FromPin, const Metasou
 #define MS_CONVERTER_BOOL(Node, FromDataType, ToDataType) MS_CONVERTER_BOOL_HIDDEN(Node, FromDataType, ToDataType) 
 #define MS_CONVERTER_BOOL_HIDDEN(Node, FromDataType, ToDataType) Node ## FromDataType ## ToDataType
 
-#define MS_INVOKE_CONVERTER_REGISTRATION_HIDDEN(NodeType, FromDataType, ToDataType, FromOutputPin, ToInputPin) FMetasoundFrontendRegistryContainer::Get()->EnqueueInitCommand([](){ ::RegisterConversionNode<NodeType, FromDataType, ToDataType >(FString(TEXT( #FromOutputPin )), FString(TEXT( #ToInputPin ))); })
+#define MS_INVOKE_CONVERTER_REGISTRATION_HIDDEN(NodeType, FromDataType, ToDataType, FromOutputPin, ToInputPin) FMetasoundFrontendRegistryContainer::Get()->EnqueueInitCommand([](){ ::Metasound::RegisterConversionNode<NodeType, FromDataType, ToDataType >(FString(TEXT( #FromOutputPin )), FString(TEXT( #ToInputPin ))); })
 
 #define MS_INVOKE_CONSTRUCTOR_REGISTRATION(NodeType, FromDataType, ToDataType, FromOutputPin, ToInputPin) MS_CONVERTER_BOOL(NodeType, FromDataType, ToDataType) = MS_INVOKE_CONVERTER_REGISTRATION_HIDDEN(NodeType, FromDataType, ToDataType, FromOutputPin, ToInputPin)
 
 // This macro exposes a previously registered node as a valid way to convert from one datatype to another.
 // When a user tries to connect a FromDataType output to a ToDataType input pin, this node will pop up as a potential node to insert between the two.
-#define REGISTER_METASOUND_CONVERTER(NodeType, FromDataType, ToDataType, FromOutputPin, ToInputPin) static const bool MS_INVOKE_CONSTRUCTOR_REGISTRATION(NodeType, FromDataType, ToDataType, FromOutputPin, ToInputPin);
+#define REGISTER_METASOUND_CONVERTER(NodeType, FromDataType, ToDataType, FromOutputPin, ToInputPin) UE_DEPRECATED_MACRO(5.8, "REGISTER_METASOUND_CONVERTER will no longer be supported. Use explicit registration during module startup/shutdown.") static const bool MS_INVOKE_CONSTRUCTOR_REGISTRATION(NodeType, FromDataType, ToDataType, FromOutputPin, ToInputPin);
diff --git a/Engine/Plugins/Runtime/Metasound/Source/MetasoundFrontend/Public/MetasoundDataTypeRegistrationMacro.h b/Engine/Plugins/Runtime/Metasound/Source/MetasoundFrontend/Public/MetasoundDataTypeRegistrationMacro.h
--- a/Engine/Plugins/Runtime/Metasound/Source/MetasoundFrontend/Public/MetasoundDataTypeRegistrationMacro.h
+++ b/Engine/Plugins/Runtime/Metasound/Source/MetasoundFrontend/Public/MetasoundDataTypeRegistrationMacro.h
@@ -48,36 +48,6 @@ namespace Metasound::Frontend
 
 	namespace MetasoundDataTypeRegistrationPrivate
 	{
-		template<typename DataType>
-		struct TDataTypeProxyConstructorDeprecation
-		{
-			private:
-				static constexpr bool bIsParsableWithDeprecatedPtr = TIsParsable<DataType, DataFactoryPrivate::ProxyDataPtrType_DEPRECATED>::Value;
-				static constexpr bool bIsParsableWithSharedProxyPtr = TIsParsable<DataType, TSharedPtr<Audio::IProxyData>>::Value;
-			public:
-				static constexpr bool bOnlySupportsDeprecatedProxyPtr = bIsParsableWithDeprecatedPtr && !bIsParsableWithSharedProxyPtr;
-		};
-
-		// ExecutableDataTypes are deprecated in 5.3 in favor of TPostExecutableDataTypes.
-		// This class triggers a deprecation warning when a TExecutableDataType
-		// is registered. 
-		template<typename DataType>
-		struct TExecutableDataTypeDeprecation
-		{
-			TExecutableDataTypeDeprecation()
-			{
-				if constexpr (TExecutableDataType<DataType>::bIsExecutable)
-				{
-					TriggerDeprecationMessage();
-				}
-			}
-
-			void TriggerDeprecationMessage() 
-			{
-				UE_LOG(LogMetaSound, Warning, TEXT("TExecutableDataType<> is deprecated in favor of TPostExecutableDataType<>. Please update your code for data type (%s) as TExecutableDataType<> will be removed in future releases"), *GetMetasoundDataTypeString<DataType>())
-			}
-		};
-
 		// Returns the Array version of a literal type if it exists.
 		template<ELiteralType LiteralType>
 		struct TLiteralArrayEnum 
@@ -129,68 +99,80 @@ namespace Metasound::Frontend
 			static constexpr ELiteralType Value = ELiteralType::UObjectProxyArray;
 		};
 
-		// This utility function can be used to optionally check to see if we can transmit a data type, and autogenerate send and receive nodes for that datatype.
-		template<typename TDataType, typename TEnableIf<TIsTransmittable<TDataType>::Value, bool>::Type = true>
-		void AttemptToRegisterSendAndReceiveNodes(const FModuleInfo& InModuleInfo)
+		
+		// Register the send and receive nodes of a datatype if it is supported
+		template<typename TDataType>
+		void RegisterSendAndReceiveNodes(const FModuleInfo& InModuleInfo)
 		{
-			if (TEnableTransmissionNodeRegistration<TDataType>::Value)
+			if constexpr (TIsTransmittable<TDataType>::Value && TEnableTransmissionNodeRegistration<TDataType>::Value)
 			{
 				RegisterNode<Metasound::TSendNode<TDataType>>(InModuleInfo);
 				RegisterNode<Metasound::TReceiveNode<TDataType>>(InModuleInfo);
 			}
 		}
 
-		template<typename TDataType, typename TEnableIf<!TIsTransmittable<TDataType>::Value, bool>::Type = true>
-		void AttemptToRegisterSendAndReceiveNodes(const FModuleInfo& InModuleInfo)
-		{
-			// This implementation intentionally noops, because Metasound::TIsTransmittable is false for this datatype.
-			// This is either because the datatype is not trivially copyable, and thus can't be buffered between threads,
-			// or it's not an audio buffer type, which we use Audio::FPatchMixerSplitter instances for.
-		}
+	
+		// Unregister the send and receive nodes related to a datatype.
+		UE_INTERNAL
+		UE_API void UnregisterSendAndReceiveNodes(const FName& InDataTypeName, const FModuleInfo& InModuleInfo);
 
-		// This utility function can be used to check to see if we can static cast between two types, and autogenerate a node for that static cast.
-		template<typename TFromDataType, typename TToDataType, typename std::enable_if<TIsAutoConvertible<TFromDataType, TToDataType>::Value, bool>::type = true>
-		void AttemptToRegisterConverter(const FModuleInfo& InModuleInfo)
+		// Register an autoconverter node for a specific data type conversion. Only registers the node if it is supported and not already registered. 
+		template<typename TFromDataType, typename TToDataType>
+		void RegisterAutoConverterNode(INodeClassRegistry& InRegistry, const FModuleInfo& InModuleInfo)
 		{
-			using FConverterNode = Metasound::TAutoConverterNode<TFromDataType, TToDataType>;
-
-			if (TEnableAutoConverterNodeRegistration<TFromDataType, TToDataType>::Value)
+			if constexpr (TIsAutoConvertible<TFromDataType, TToDataType>::Value && TEnableAutoConverterNodeRegistration<TFromDataType, TToDataType>::Value && !std::is_same<TFromDataType, TToDataType>::value)
 			{
-				const FNodeClassMetadata& Metadata = FConverterNode::GetAutoConverterNodeMetadata();
-				const Metasound::Frontend::FNodeRegistryKey Key(Metadata);
+				using namespace AutoConverterNodePrivate;
+
+				const FName& FromDataTypeName = GetMetasoundDataTypeName<TFromDataType>();
+				const FName& ToDataTypeName = GetMetasoundDataTypeName<TToDataType>();
 
-				if (!std::is_same<TFromDataType, TToDataType>::value && !FMetasoundFrontendRegistryContainer::Get()->IsNodeRegistered(Key))
+				FConvertDataTypeInfo ConvertInfo = MakeConvertDataTypeInfo<TFromDataType, TToDataType>();
+				const FNodeClassMetadata Metadata = CreateAutoConverterNodeMetadata(ConvertInfo);
+				const Metasound::Frontend::FNodeClassRegistryKey Key{Metadata};
+
+				if (!InRegistry.IsNodeRegistered(Key))
 				{
+					using FConverterNode = Metasound::TAutoConverterNode<TFromDataType, TToDataType>;
 					RegisterNode<FConverterNode>(Metadata, InModuleInfo);
-					
-					RegisterConversionNode<FConverterNode, TFromDataType, TToDataType>(FConverterNode::GetInputName(), FConverterNode::GetOutputName(), Metadata);
+
+					FConverterNodeInfo ConverterNodeInfo
+					{
+						FConverterNode::GetInputName(), 
+						FConverterNode::GetOutputName(), 
+						FNodeClassRegistryKey{Metadata}
+					};
+					InRegistry.RegisterConversionNode(FConverterNodeClassRegistryKey{FromDataTypeName, ToDataTypeName}, ConverterNodeInfo);
 				}
 			}
 		}
 
-		template<typename TFromDataType, typename TToDataType, typename std::enable_if<!TIsAutoConvertible<TFromDataType, TToDataType>::Value, int>::type = 0>
-		void AttemptToRegisterConverter(const FModuleInfo& InModuleInfo)
-		{
-			// This implementation intentionally noops, because static_cast<TFromDataType>(TToDataType&) is invalid.
-		}
+		// Unregister any autoconverter node that is registered for going converting between specific data types
+		UE_INTERNAL
+		UE_API void UnregisterAutoConverterNode(INodeClassRegistry& InRegistry, const FName& InFromDataTypeName, const FName& InToDataTypeName, const FModuleInfo& InModuleInfo);
 
-		// Here we attempt to infer and autogenerate conversions for basic datatypes.
+		// Register a slew of basic converter nodes for basic types
 		template<typename TDataType>
-		void RegisterConverterNodes(const FModuleInfo& InModuleInfo)
+		void RegisterAutoConverterNodes(const FModuleInfo& InModuleInfo)
 		{
+			INodeClassRegistry& Registry = INodeClassRegistry::GetChecked();
 			// Conversions to this data type:
-			AttemptToRegisterConverter<bool, TDataType>(InModuleInfo);
-			AttemptToRegisterConverter<int32, TDataType>(InModuleInfo);
-			AttemptToRegisterConverter<float, TDataType>(InModuleInfo);
-			AttemptToRegisterConverter<FString, TDataType>(InModuleInfo);
+			RegisterAutoConverterNode<bool, TDataType>(Registry, InModuleInfo);
+			RegisterAutoConverterNode<int32, TDataType>(Registry, InModuleInfo);
+			RegisterAutoConverterNode<float, TDataType>(Registry, InModuleInfo);
+			RegisterAutoConverterNode<FString, TDataType>(Registry, InModuleInfo);
 
 			// Conversions from this data type:
-			AttemptToRegisterConverter<TDataType, bool>(InModuleInfo);
-			AttemptToRegisterConverter<TDataType, int32>(InModuleInfo);
-			AttemptToRegisterConverter<TDataType, float>(InModuleInfo);
-			AttemptToRegisterConverter<TDataType, FString>(InModuleInfo);
+			RegisterAutoConverterNode<TDataType, bool>(Registry, InModuleInfo);
+			RegisterAutoConverterNode<TDataType, int32>(Registry, InModuleInfo);
+			RegisterAutoConverterNode<TDataType, float>(Registry, InModuleInfo);
+			RegisterAutoConverterNode<TDataType, FString>(Registry, InModuleInfo);
 		}
 
+		// Unregister a slew of basic converter nodes for basic types.
+		UE_INTERNAL
+		UE_API void UnregisterAutoConverterNodes(const FName& InDataTypeName, const FModuleInfo& InModuleInfo);
+
 		/** Creates the FDataTypeRegistryInfo for a data type.
 		 * 
 		 * @tparam TDataType - The data type to create info for.
@@ -410,30 +392,6 @@ namespace Metasound::Frontend
 			static constexpr bool bIsParsable = TLiteralTraits<TDataType>::bIsParsableFromAnyLiteralType;
 			static constexpr bool bIsConstructorType = TIsConstructorVertexSupported<TDataType>::Value;
 
-			static bool bAlreadyRegisteredThisDataType = false;
-			if (bAlreadyRegisteredThisDataType)
-			{
-				UE_LOG(LogMetaSound, Display, TEXT("Tried to call REGISTER_METASOUND_DATATYPE twice with the same class %s. ignoring the second call. Likely because REGISTER_METASOUND_DATATYPE is in a header that's used in multiple modules. Consider moving it to a private header or cpp file."), TDataReferenceTypeInfo<TDataType>::TypeName)
-				return false;
-			}
-
-	
-			if constexpr (MetasoundDataTypeRegistrationPrivate::TDataTypeProxyConstructorDeprecation<TDataType>::bOnlySupportsDeprecatedProxyPtr)
-			{
-				// TUniquePtr<Audio::IProxyData> deprecated in 5.2. Log a warning
-				// during data type registration to warn users to update their MetaSound
-				// data type constructor. 
-				UE_LOG(LogMetaSound, Warning, TEXT("MetaSound data type \"%s\" supports construction from deprecated TUniquePtr<Audio::IProxyData>. Please update the constructor to accept a \"const TSharedPtr<Audio::IProxyData>& \""), TDataReferenceTypeInfo<TDataType>::TypeName);
-			}
-			
-			// TExecutableDataTypes are deprecated as of 5.3. This call triggers
-			// a deprecation warning in case the TExecutableDataType<> template
-			// was specialized.
-			MetasoundDataTypeRegistrationPrivate::TExecutableDataTypeDeprecation<TDataType>();
-
-			bAlreadyRegisteredThisDataType = true;
-
-
 			class FDataTypeRegistryEntry : public FDataTypeRegistryEntryBase
 			{
 				void InitRawAssignmentFunction()
@@ -810,17 +768,26 @@ namespace Metasound::Frontend
 				}
 			};
 
-			bool bSucceeded = Frontend::IDataTypeRegistry::Get().RegisterDataType(MakeUnique<FDataTypeRegistryEntry>());
+			bool bSucceeded = IDataTypeRegistry::Get().RegisterDataType(MakeUnique<FDataTypeRegistryEntry>());
 
 			if (bSucceeded)
 			{
-				RegisterConverterNodes<TDataType>(InModuleInfo);
-				AttemptToRegisterSendAndReceiveNodes<TDataType>(InModuleInfo);
+				RegisterAutoConverterNodes<TDataType>(InModuleInfo);
+				RegisterSendAndReceiveNodes<TDataType>(InModuleInfo);
+			}
+			else
+			{
+				UE_LOG(LogMetaSound, Error, TEXT("Failed to register MetaSound DataType %s"), *GetMetasoundDataTypeName<TDataType>().ToString());
 			}
 			
 			return bSucceeded;
 		}
 
+		// Unregister a datatype as well as some autogenerated nodes. This should unregister all nodes that were registered
+		// during RegisterDataTypeWIthFrontendInternal<>(...)
+		UE_INTERNAL
+		UE_API bool UnregisterDataTypeWithFrontendInternal(const FName& InDataTypeName, const FModuleInfo& InModuleInfo);
+
 		/** Registers an array of a data type with the MetaSound Frontend. This allows 
 		 * an array of the data type to be used in Input, Output, Send and Receive 
 		 * nodes by informing the Frontend how to instantiate an instance. 
@@ -840,7 +807,6 @@ namespace Metasound::Frontend
 
 			if constexpr (TEnableAutoArrayTypeRegistration<TDataType>::Value)
 			{
-				constexpr bool bIsArrayType = true;
 				bool bSuccess = RegisterDataTypeWithFrontendInternal<TArrayType, TLiteralArrayEnum<PreferredArgType>::Value>(InModuleInfo);
 				bSuccess = bSuccess && RegisterArrayNodes<TArrayType>(InModuleInfo);
 				bSuccess = bSuccess && RegisterDataTypeWithFrontendInternal<TVariable<TArrayType>>(InModuleInfo);
@@ -851,6 +817,32 @@ namespace Metasound::Frontend
 				return true;
 			}
 		}
+
+		template<typename TDataType>
+		bool UnregisterDataTypeArrayWithFrontend(const FModuleInfo& InModuleInfo)
+		{
+			using namespace MetasoundDataTypeRegistrationPrivate;
+
+			if constexpr (TEnableAutoArrayTypeRegistration<TDataType>::Value)
+			{
+				using TArrayType = TArray<TDataType>;
+
+				const FName& DataTypeName = GetMetasoundDataTypeName<TArrayType>();
+
+				// Unregistration methods are called in reverse order of registration 
+				// functions.
+PRAGMA_DISABLE_INTERNAL_WARNINGS
+				bool bSuccess = UnregisterDataTypeWithFrontendInternal(GetMetasoundDataTypeName<TVariable<TArrayType>>(), InModuleInfo);
+				bSuccess = bSuccess && UnregisterArrayNodes(DataTypeName, InModuleInfo);
+				bSuccess = bSuccess && UnregisterDataTypeWithFrontendInternal(DataTypeName, InModuleInfo);
+PRAGMA_ENABLE_INTERNAL_WARNINGS
+				return bSuccess;
+			}
+			else
+			{
+				return true;
+			}
+		}
 	} // namespace MetasoundDataTypeRegistrationPrivate
 	
 	/** Registers a data type with the MetaSound Frontend. This allows the data type 
@@ -879,6 +871,27 @@ namespace Metasound::Frontend
 		return bSuccess;
 	}
 
+	/** Unregisters a data type with the MetaSound Frontend. 
+	 *
+	 * @tparam TDataType - The data type to register.
+	 *
+	 * @return True on success, false on failure.
+	 */
+	template<typename TDataType>
+	bool UnregisterDataType(const FModuleInfo& InModuleInfo)
+	{
+		using namespace MetasoundDataTypeRegistrationPrivate;
+
+PRAGMA_DISABLE_INTERNAL_WARNINGS
+		// Unregistration is done in reverse order of registration
+		bool bSuccess = UnregisterDataTypeArrayWithFrontend<TDataType>(InModuleInfo);
+		bSuccess = bSuccess && UnregisterDataTypeWithFrontendInternal(GetMetasoundDataTypeName<TVariable<TDataType>>(), InModuleInfo);
+		bSuccess = bSuccess && UnregisterDataTypeWithFrontendInternal(GetMetasoundDataTypeName<TDataType>(), InModuleInfo);
+PRAGMA_ENABLE_INTERNAL_WARNINGS
+
+		return bSuccess;
+	}
+
 	/** Registers a data type with the MetaSound Frontend. This allows the data type 
 	 * to be used in Input, Output, Send and Receive  nodes by informing the 
 	 * Frontend how to instantiate an instance. 
@@ -948,7 +961,7 @@ namespace Metasound
 
 #define ENQUEUE_METASOUND_DATATYPE_REGISTRATION_COMMAND(DataType, DataTypeName, ...) \
 	static_assert(::Metasound::Frontend::TMetasoundDataTypeRegistration<DataType>::bCanRegister, CANNOT_REGISTER_METASOUND_DATA_TYPE_ASSERT_STRING(DataType)); \
-	METASOUND_IMPLEMENT_REGISTRATION_ACTION(UE_JOIN(_DataTypeRegistration,__COUNTER__), (::Metasound::Frontend::RegisterDataType<DataType, ##__VA_ARGS__>), nullptr)
+	METASOUND_IMPLEMENT_REGISTRATION_ACTION(UE_JOIN(_DataTypeRegistration,__COUNTER__), (::Metasound::Frontend::RegisterDataType<DataType, ##__VA_ARGS__>), (::Metasound::Frontend::UnregisterDataType<DataType>))
 
 // This should be used to expose a datatype as a potential input or output for a metasound graph.
 // The first argument to the macro is the class to expose.
diff --git a/Engine/Plugins/Runtime/Metasound/Source/MetasoundFrontend/Public/MetasoundFrontendDataTypeRegistry.h b/Engine/Plugins/Runtime/Metasound/Source/MetasoundFrontend/Public/MetasoundFrontendDataTypeRegistry.h
--- a/Engine/Plugins/Runtime/Metasound/Source/MetasoundFrontend/Public/MetasoundFrontendDataTypeRegistry.h
+++ b/Engine/Plugins/Runtime/Metasound/Source/MetasoundFrontend/Public/MetasoundFrontendDataTypeRegistry.h
@@ -310,13 +310,19 @@ namespace Metasound
 			virtual ~IDataTypeRegistry() = default;
 
 			/** Register a data type
-			 * @param InName - Name of data type.
 			 * @param InEntry - TUniquePtr to data type registry entry.
 			 *
 			 * @return True on success, false on failure.
 			 */
 			virtual bool RegisterDataType(TUniquePtr<IDataTypeRegistryEntry>&& InEntry) = 0;
 
+			/** Unregister a data type.
+			 * @param InDataTypeName
+			 *
+			 * @return True on success, false on failure.
+			 */
+			virtual bool UnregisterDataType(const FName& InDataTypeName) = 0;
+
 			/** Provides all names of registered DataTypes. */
 			virtual void GetRegisteredDataTypeNames(TArray<FName>& OutNames) const = 0;
 
diff --git a/Engine/Plugins/Runtime/Metasound/Source/MetasoundFrontend/Public/MetasoundFrontendNodeClassRegistry.h b/Engine/Plugins/Runtime/Metasound/Source/MetasoundFrontend/Public/MetasoundFrontendNodeClassRegistry.h
--- a/Engine/Plugins/Runtime/Metasound/Source/MetasoundFrontend/Public/MetasoundFrontendNodeClassRegistry.h
+++ b/Engine/Plugins/Runtime/Metasound/Source/MetasoundFrontend/Public/MetasoundFrontendNodeClassRegistry.h
@@ -343,6 +343,12 @@ namespace Metasound::Frontend
 		
 		/** Unregister a node migration from one plugin and/or module to another. This should be used during module shutdown. */
 		virtual bool UnregisterNodeMigration(const FNodeMigrationInfo& InMigrationInfo) = 0;
+
+		/** Name of the plugin registering this node. */
+		virtual FName GetOwningPluginName(const FNodeClassRegistryKey& InKey) const = 0;
+		
+		/** Name of the module registering this node. */
+		virtual FName GetOwningModuleName(const FNodeClassRegistryKey& InKey) const = 0;
 #endif
 
 		/** Returns true if the provided registry key corresponds to a valid registered node. */
@@ -356,6 +362,7 @@ namespace Metasound::Frontend
 		virtual bool IsNodeNative(const FNodeClassRegistryKey& InKey) const { return false; }
 
 		// Iterates class types in registry.  If InClassType is set to a valid class type (optional), only iterates classes of the given type
+		UE_DEPRECATED(5.8, "This will be used internally only due to thread safety issues. Use Metasound::Frontend::ISearchEngine instead.")
 		virtual void IterateRegistry(FIterateMetasoundFrontendClassFunction InIterFunc, EMetasoundFrontendClassType InClassType = EMetasoundFrontendClassType::Invalid) const = 0;
 
 		// Query for MetaSound Frontend document objects.
@@ -401,7 +408,11 @@ namespace Metasound::Frontend
 		UE_DEPRECATED(5.6, "Node classes should be constructed with FNodeData")
 		virtual TUniquePtr<Metasound::INode> CreateNode(const FNodeClassRegistryKey& InKey, Metasound::FDefaultNamedVertexWithLiteralNodeConstructorParams&&) const { return nullptr; }
 
-		virtual bool RegisterConversionNode(const FConverterNodeClassRegistryKey& InNodeKey, const FConverterNodeClassInfo& InNodeInfo) = 0;
+		/** Registers a node as a converter node */
+		virtual bool RegisterConversionNode(const FConverterNodeClassRegistryKey& InConversionKey, const FConverterNodeClassInfo& InNodeInfo) = 0;
+		
+		/** Unregisters a node as a converter node */
+		virtual bool UnregisterConversionNode(const FConverterNodeClassRegistryKey& InConversionKey, const FNodeClassRegistryKey& InNodeKey) = 0;
 
 		// Returns a list of possible nodes to use to convert from FromDataType to ToDataType.
 		// Returns an empty array if none are available.
diff --git a/Engine/Plugins/Runtime/Metasound/Source/MetasoundFrontend/Public/MetasoundNodeRegistrationMacro.h b/Engine/Plugins/Runtime/Metasound/Source/MetasoundFrontend/Public/MetasoundNodeRegistrationMacro.h
--- a/Engine/Plugins/Runtime/Metasound/Source/MetasoundFrontend/Public/MetasoundNodeRegistrationMacro.h
+++ b/Engine/Plugins/Runtime/Metasound/Source/MetasoundFrontend/Public/MetasoundNodeRegistrationMacro.h
@@ -188,7 +188,6 @@ namespace Metasound::Frontend
 		};
 
 		UE_API bool RegisterNodeInternal(TUniquePtr<INodeClassRegistryEntry>&& InEntry);
-		UE_API bool UnregisterNodeInternal(const FNodeClassMetadata& InMetadata);
 	} // namespace NodeRegistrationPrivate
 
 	template<typename TNodeType, typename ConfigurationType=void>
@@ -241,14 +240,20 @@ namespace Metasound::Frontend
 		return RegisterNode<TNodeType, ConfigurationType>(FModuleInfo{});
 	}
 
+	/** Unregister a node using a registry key. */
+	UE_API bool UnregisterNode(const FNodeClassRegistryKey& InRegistryKey, const FModuleInfo& InOwningModuleInfo);
+	
+	/** Unregister a node using FNodeClassMetadata. */
+	UE_API bool UnregisterNode(const FNodeClassMetadata& InMetadata, const FModuleInfo& InOwningModuleInfo);
 
+	/** Unregister a node by node template param. */
 	template<typename TNodeType>
 	bool UnregisterNode(const FModuleInfo& InOwningModuleInfo)
 	{
 		using namespace NodeRegistrationPrivate;
 		if constexpr(TIsCreateNodeClassMetadataDeclared<TNodeType>::Value)
 		{
-			return UnregisterNodeInternal(TNodeType::CreateNodeClassMetadata());
+			return UnregisterNode(FNodeClassRegistryKey{TNodeType::CreateNodeClassMetadata()}, InOwningModuleInfo);
 		}
 		else
 		{
@@ -257,7 +262,7 @@ namespace Metasound::Frontend
 			// Register a node using a prototype node.
 			FNodeInitData InitData;
 			TUniquePtr<Metasound::INode> Node = MakeUnique<TNodeType>(InitData);
-			return UnregisterNodeInternal(Node->GetMetadata());
+			return UnregisterNode(Node->GetMetadata(), InOwningModuleInfo);
 		}
 	}
 
diff --git a/Engine/Plugins/Runtime/Metasound/Source/MetasoundFrontend/Public/MetasoundReceiveNode.h b/Engine/Plugins/Runtime/Metasound/Source/MetasoundFrontend/Public/MetasoundReceiveNode.h
--- a/Engine/Plugins/Runtime/Metasound/Source/MetasoundFrontend/Public/MetasoundReceiveNode.h
+++ b/Engine/Plugins/Runtime/Metasound/Source/MetasoundFrontend/Public/MetasoundReceiveNode.h
@@ -37,6 +37,9 @@ namespace Metasound
 
 	namespace ReceiveNodePrivate
 	{
+		UE_INTERNAL
+		METASOUNDFRONTEND_API Frontend::FNodeClassRegistryKey CreateNodeClassRegistryKey(const FName& InDataTypeName);
+
 		template<typename TDataType>
 		class TReceiverOperator : public TExecutableOperator<TReceiverOperator<TDataType>>
 		{
diff --git a/Engine/Plugins/Runtime/Metasound/Source/MetasoundFrontend/Public/MetasoundSendNode.h b/Engine/Plugins/Runtime/Metasound/Source/MetasoundFrontend/Public/MetasoundSendNode.h
--- a/Engine/Plugins/Runtime/Metasound/Source/MetasoundFrontend/Public/MetasoundSendNode.h
+++ b/Engine/Plugins/Runtime/Metasound/Source/MetasoundFrontend/Public/MetasoundSendNode.h
@@ -27,16 +27,26 @@ namespace Metasound
 
 	namespace SendNodePrivate
 	{
+		UE_INTERNAL
+		METASOUNDFRONTEND_API FVertexInterface CreateVertexInterface(const FName& InDataTypeName);
+
+		UE_INTERNAL
+		METASOUNDFRONTEND_API FNodeClassName CreateNodeClassName(const FName& InDataTypeName);
+
+		UE_INTERNAL
+		METASOUNDFRONTEND_API Frontend::FNodeClassRegistryKey CreateNodeClassRegistryKey(const FName& InDataTypeName);
+
+		UE_INTERNAL
+		METASOUNDFRONTEND_API FNodeClassMetadata CreateNodeClassMetadata(const FName& InDataTypeName, const FText& InDataTypeDisplayText);
+
+		UE_INTERNAL
+		METASOUNDFRONTEND_API const FName& GetSendInputName(const FName& InDataTypeName);
+
 		template<typename TDataType>
 		class TSendOperator : public TExecutableOperator<TSendOperator<TDataType>>
 		{
 		public:
 
-			static const FVertexName& GetSendInputName()
-			{
-				static const FVertexName& SendInput = GetMetasoundDataTypeName<TDataType>();
-				return SendInput;
-			}
 
 			TSendOperator(TDataReadReference<TDataType> InInputData, TDataReadReference<FSendAddress> InSendAddress, const FOperatorSettings& InOperatorSettings)
 				: InputData(InInputData)
@@ -56,8 +66,10 @@ namespace Metasound
 			virtual void BindInputs(FInputVertexInterfaceData& InOutVertexData) override
 			{
 				using namespace SendVertexNames; 
+				using namespace SendNodePrivate;
+
 				InOutVertexData.BindReadVertex<FSendAddress>(METASOUND_GET_PARAM_NAME(AddressInput), SendAddress);
-				InOutVertexData.BindReadVertex<TDataType>(GetSendInputName(), InputData);
+				InOutVertexData.BindReadVertex<TDataType>(GetSendInputName(GetMetasoundDataTypeName<TDataType>()), InputData);
 			}
 
 			virtual void BindOutputs(FOutputVertexInterfaceData& InOutVertexData) override
@@ -87,62 +99,24 @@ namespace Metasound
 
 			static FVertexInterface DeclareVertexInterface()
 			{
-				using namespace SendVertexNames; 
-				static const FDataVertexMetadata AddressInputMetadata
-				{
-					  FText::GetEmpty() // description
-					, METASOUND_GET_PARAM_DISPLAYNAME(AddressInput) // display name
-				};
-
-				return FVertexInterface(
-					FInputVertexInterface(
-						TInputDataVertex<FSendAddress>(METASOUND_GET_PARAM_NAME(AddressInput), AddressInputMetadata),
-						TInputDataVertex<TDataType>(GetSendInputName(), FDataVertexMetadata{ FText::GetEmpty() })
-					),
-					FOutputVertexInterface(
-					)
-				);
+				return CreateVertexInterface(GetMetasoundDataTypeName<TDataType>());
 			}
 
 			static const FNodeClassMetadata& GetNodeInfo()
 			{
-				auto InitNodeInfo = []() -> FNodeClassMetadata
-				{
-					const FVertexName& InputName = GetSendInputName();
-					FNodeClassMetadata Info;
-
-					Info.ClassName = { "Send", GetMetasoundDataTypeName<TDataType>(), FName() };
-					Info.MajorVersion = 1;
-					Info.MinorVersion = 0;
-					Info.DisplayName = METASOUND_LOCTEXT_FORMAT("Metasound_SendNodeDisplayNameFormat", "Send {0}", GetMetasoundDataTypeDisplayText<TDataType>());
-					Info.Description = METASOUND_LOCTEXT("Metasound_SendNodeDescription", "Sends data from a send node with the same name.");
-					Info.Author = PluginAuthor;
-					Info.PromptIfMissing = PluginNodeMissingPrompt;
-					Info.DefaultInterface = DeclareVertexInterface();
-					Info.CategoryHierarchy = { METASOUND_LOCTEXT("Metasound_TransmissionNodeCategory", "Transmission") };
-					Info.Keywords = { };
-
-					// Then send & receive nodes do not work as expected, particularly 
-					// around multiple-consumer scenarios. Deprecate them to avoid
-					// metasound assets from relying on send & receive nodes. 
-					EnumAddFlags(Info.AccessFlags, ENodeClassAccessFlags::Deprecated);
-
-					return Info;
-				};
-
-				static const FNodeClassMetadata Info = InitNodeInfo();
-
+				static const FNodeClassMetadata Info = CreateNodeClassMetadata(GetMetasoundDataTypeName<TDataType>(), GetMetasoundDataTypeDisplayText<TDataType>());
 				return Info;
 			}
 
 			static TUniquePtr<IOperator> CreateOperator(const FBuildOperatorParams& InParams, FBuildResults& OutResults)
 			{
 				using namespace SendVertexNames;
+				using namespace SendNodePrivate;
 
-				if (InParams.InputData.IsVertexBound(GetSendInputName()))
+				if (InParams.InputData.IsVertexBound(GetSendInputName(GetMetasoundDataTypeName<TDataType>())))
 				{
 					return MakeUnique<SendNodePrivate::TSendOperator<TDataType>>(
-						InParams.InputData.GetOrCreateDefaultDataReadReference<TDataType>(GetSendInputName(), InParams.OperatorSettings),
+						InParams.InputData.GetOrCreateDefaultDataReadReference<TDataType>(GetSendInputName(GetMetasoundDataTypeName<TDataType>()), InParams.OperatorSettings),
 						InParams.InputData.GetOrCreateDefaultDataReadReference<FSendAddress>(METASOUND_GET_PARAM_NAME(AddressInput), InParams.OperatorSettings),
 						InParams.OperatorSettings
 					);
@@ -188,6 +162,8 @@ namespace Metasound
 		};
 	}
 
+
+
 	template<typename TDataType>
 	class TSendNode : public TNodeFacade<SendNodePrivate::TSendOperator<TDataType>>
 	{
@@ -196,11 +172,13 @@ namespace Metasound
 		using FSuper = TNodeFacade<SendNodePrivate::TSendOperator<TDataType>>;
 		using FSuper::FSuper;
 
+		UE_DEPRECATED(5.8, "This will no longer be supported")
 		static const FVertexName& GetSendInputName()
 		{
-			return FOperator::GetSendInputName();
+			return SendNodePrivate::GetSendInputName(GetMetasoundDataTypeName<TDataType>());
 		}
 
+		UE_DEPRECATED(5.8, "This will no longer be supported")
 		static FVertexInterface DeclareVertexInterface()
 		{
 			return FOperator::DeclareVertexInterface();
diff --git a/Engine/Plugins/Runtime/Metasound/Source/MetasoundGraphCore/Public/MetasoundNodeInterface.h b/Engine/Plugins/Runtime/Metasound/Source/MetasoundGraphCore/Public/MetasoundNodeInterface.h
--- a/Engine/Plugins/Runtime/Metasound/Source/MetasoundGraphCore/Public/MetasoundNodeInterface.h
+++ b/Engine/Plugins/Runtime/Metasound/Source/MetasoundGraphCore/Public/MetasoundNodeInterface.h
@@ -144,8 +144,8 @@ namespace Metasound
 		UE_API FNodeClassMetadata(const FNodeClassMetadata& InOther);
 		UE_API FNodeClassMetadata(
 			FNodeClassName InClassName,
-			int32 InMajorVersion = -1,
-			int32 InMinorVersion = -1,
+			int32 InMajorVersion,
+			int32 InMinorVersion,
 			FText InDisplayName = { },
 			FText InDescription = { },
 			FString InAuthor = { },
