Readd conform object in MetaSound autoupdate that was causing issues with preset output formats

#rnx
#Synced-CL 49174767

diff --git a/Engine/Plugins/Runtime/Metasound/Source/MetasoundFrontend/Private/MetasoundFrontendTransform.cpp b/Engine/Plugins/Runtime/Metasound/Source/MetasoundFrontend/Private/MetasoundFrontendTransform.cpp
--- a/Engine/Plugins/Runtime/Metasound/Source/MetasoundFrontend/Private/MetasoundFrontendTransform.cpp
+++ b/Engine/Plugins/Runtime/Metasound/Source/MetasoundFrontend/Private/MetasoundFrontendTransform.cpp
@@ -510,6 +510,9 @@ namespace Metasound
 			if (TConstStructView<FMetaSoundFrontendDocumentConfiguration> DocConfig = InOutBuilderToTransform.GetConstDocumentConfiguration(); DocConfig.IsValid())
 			{
 				InOutBuilderToTransform.ConfigureDocument(true /* bLogFailure */);
+				
+				TScriptInterface<IMetaSoundDocumentInterface> DocInterface(&InOutBuilderToTransform.CastDocumentObjectChecked<UObject>());
+				DocInterface->ConformObjectToDocument();
 				InOutBuilderToTransform.SynchronizeDependencyMetadata();
 			}
 
