Fix for race condition on triggering async metasound asset versioning

#rnx
#Synced-CL 49197653

diff --git a/Engine/Plugins/Runtime/Metasound/Source/MetasoundEngine/Private/MetasoundEngineAsset.h b/Engine/Plugins/Runtime/Metasound/Source/MetasoundEngine/Private/MetasoundEngineAsset.h
--- a/Engine/Plugins/Runtime/Metasound/Source/MetasoundEngine/Private/MetasoundEngineAsset.h
+++ b/Engine/Plugins/Runtime/Metasound/Source/MetasoundEngine/Private/MetasoundEngineAsset.h
@@ -346,7 +346,8 @@ PRAGMA_DISABLE_DEPRECATION_WARNINGS
 				{
 					// Versioning asset is applied after async load of soft object
 					// references to ensure document properties are ready for versioning.
-					IMetaSoundAssetManager::GetChecked().RequestAsyncLoadReferencedAssets(InMetaSound, &ReferenceAssetClassCache);
+					UE_LOG(LogMetaSound, Display, TEXT("Delaying asset versioning due to need to async load soft references for MetaSound %s"), *InMetaSound.GetOwningAssetName());
+					IMetaSoundAssetManager::GetChecked().RequestAsyncLoadReferencedAssets(InMetaSound, &ReferenceAssetClassCache); 
 				}
 
 				// There is no way to determine if a commandlet is going to read directly from a MetaSound object
diff --git a/Engine/Plugins/Runtime/Metasound/Source/MetasoundEngine/Public/MetasoundDocumentBuilderRegistry.h b/Engine/Plugins/Runtime/Metasound/Source/MetasoundEngine/Public/MetasoundDocumentBuilderRegistry.h
--- a/Engine/Plugins/Runtime/Metasound/Source/MetasoundEngine/Public/MetasoundDocumentBuilderRegistry.h
+++ b/Engine/Plugins/Runtime/Metasound/Source/MetasoundEngine/Public/MetasoundDocumentBuilderRegistry.h
@@ -3,6 +3,7 @@
 
 #include "Algo/Find.h"
 #include "Async/Async.h"
+#include "CoreGlobals.h"
 #include "MetasoundBuilderBase.h"
 #include "MetasoundDocumentInterface.h"
 #include "MetasoundFrontendDocument.h"
@@ -67,6 +68,8 @@ namespace Metasound::Engine
 		template <typename BuilderClass = UMetaSoundBuilderBase>
 		BuilderClass& FindOrBeginBuilding(UObject& InMetaSoundObject, bool bIsTransacting = false) const
 		{
+			check(IsInGameThread());
+
 			if (!bIsTransacting)
 			{
 				check(InMetaSoundObject.IsAsset());
@@ -75,16 +78,16 @@ namespace Metasound::Engine
 			TScriptInterface<IMetaSoundDocumentInterface> DocInterface = &InMetaSoundObject;
 			check(DocInterface.GetObject());
 
+#if WITH_EDITORONLY_DATA
+			WaitUntilVersioningComplete(InMetaSoundObject);
+#endif // WITH_EDITORONLY_DATA
+
 			FScopeLock Lock(&BuildersCriticalSection);
 			if (UMetaSoundBuilderBase* Builder = FindBuilderObject(&InMetaSoundObject))
 			{
 				return *CastChecked<BuilderClass>(Builder);
 			}
 
-#if WITH_EDITORONLY_DATA
-			WaitUntilVersioningComplete(InMetaSoundObject);
-#endif // WITH_EDITORONLY_DATA
-
 			FNameBuilder BuilderName;
 			BuilderName.Append(InMetaSoundObject.GetName());
 			BuilderName.Append(TEXT("_Builder"));
diff --git a/Engine/Plugins/Runtime/Metasound/Source/MetasoundFrontend/Private/MetasoundFrontendDocumentVersioning.cpp b/Engine/Plugins/Runtime/Metasound/Source/MetasoundFrontend/Private/MetasoundFrontendDocumentVersioning.cpp
--- a/Engine/Plugins/Runtime/Metasound/Source/MetasoundFrontend/Private/MetasoundFrontendDocumentVersioning.cpp
+++ b/Engine/Plugins/Runtime/Metasound/Source/MetasoundFrontend/Private/MetasoundFrontendDocumentVersioning.cpp
@@ -1276,61 +1276,69 @@ PRAGMA_ENABLE_DEPRECATION_WARNINGS
 			IDocumentBuilderRegistry::GetChecked().FinishBuilding(ClassName, AssetPath);
 		}
 
-		InMetaSound.SetInternalFlags(EInternalObjectFlags::Async);
-		FTaskData NewTaskData
 		{
-			.AssetPtr = TStrongObjectPtr<UObject>(&InMetaSound),
-			.Task = UE::Tasks::Launch(UE_SOURCE_LOCATION, [this, AssetId, bIsDeterministic]()
+			FScopeLock Lock(&ActiveVersionCritSection);
+			
+			if (ActiveVersionTaskData.Contains(AssetId))
 			{
-				FTaskData TaskData;
+				// Early out if an async task has already begun
+				return;
+			}
+
+			UE_LOG(LogMetaSound, Verbose, TEXT("Beginning async versioning for %s"), *AssetPath.ToString());
+			InMetaSound.SetInternalFlags(EInternalObjectFlags::Async);
+			FTaskData NewTaskData
+			{
+				.AssetPtr = TStrongObjectPtr<UObject>(&InMetaSound),
+				.Task = UE::Tasks::Launch(UE_SOURCE_LOCATION, [this, AssetId, bIsDeterministic]()
 				{
-					FScopeLock VersionLock(&ActiveVersionCritSection);
-					if (const FTaskData* TaskEntry = ActiveVersionTaskData.Find(AssetId))
-					{
-						TaskData = *TaskEntry;
-					}
-					else
+					FTaskData TaskData;
 					{
-						return;
+						FScopeLock VersionLock(&ActiveVersionCritSection);
+						if (const FTaskData* TaskEntry = ActiveVersionTaskData.Find(AssetId))
+						{
+							TaskData = *TaskEntry;
+						}
+						else
+						{
+							return;
+						}
 					}
-				}
-
-				{
-					FGCScopeGuard ScopeGuard;
 
-					if (UObject* VersioningMetaSound = TaskData.AssetPtr.Get())
 					{
-						FDocumentIDGenerator::FScopeDeterminism DeterminismScope(bIsDeterministic);
-						FMetaSoundFrontendDocumentBuilder DocBuilder(VersioningMetaSound);
-						const bool bVersionedDoc = VersionDocument(DocBuilder);
-						const bool bVersionedInterfaces = VersionInterfacesInternal(DocBuilder);
+						FGCScopeGuard ScopeGuard;
 
-						if (bVersionedDoc || bVersionedInterfaces)
+						if (UObject* VersioningMetaSound = TaskData.AssetPtr.Get())
 						{
-							DocBuilder.GetMetasoundAsset().SetVersionedOnLoad();
-						}
+							FDocumentIDGenerator::FScopeDeterminism DeterminismScope(bIsDeterministic);
+							FMetaSoundFrontendDocumentBuilder DocBuilder(VersioningMetaSound);
+							const bool bVersionedDoc = VersionDocument(DocBuilder);
+							const bool bVersionedInterfaces = VersionInterfacesInternal(DocBuilder);
 
-						ForEachObjectWithOuter(VersioningMetaSound, [](UObject* Object)
-						{
-							if (Object->HasAnyInternalFlags(EInternalObjectFlags::Async))
+							if (bVersionedDoc || bVersionedInterfaces)
 							{
-								METASOUND_VERSIONING_LOG(VeryVerbose, TEXT("Clearing object '%s' async flag (object likely generated while MetaSound versioning)"), *Object->GetName());
-								Object->ClearInternalFlags(EInternalObjectFlags::Async);
+								DocBuilder.GetMetasoundAsset().SetVersionedOnLoad();
 							}
-						});
-						VersioningMetaSound->ClearInternalFlags(EInternalObjectFlags::Async);
+
+							ForEachObjectWithOuter(VersioningMetaSound, [](UObject* Object)
+							{
+								if (Object->HasAnyInternalFlags(EInternalObjectFlags::Async))
+								{
+									METASOUND_VERSIONING_LOG(Verbose, TEXT("Clearing object '%s' async flag (object likely generated while MetaSound versioning)"), *Object->GetName());
+									Object->ClearInternalFlags(EInternalObjectFlags::Async);
+								}
+							});
+							VersioningMetaSound->ClearInternalFlags(EInternalObjectFlags::Async);
+						}
 					}
-				}
 
-				{
-					FScopeLock VersionLock(&ActiveVersionCritSection);
-					ActiveVersionTaskData.Remove(AssetId);
-				}
-			})
-		};
+					{
+						FScopeLock VersionLock(&ActiveVersionCritSection);
+						ActiveVersionTaskData.Remove(AssetId);
+					}
+				})
+			};
 
-		{
-			FScopeLock Lock(&ActiveVersionCritSection);
 			ActiveVersionTaskData.Add(AssetId, MoveTemp(NewTaskData));
 		}
 	}
@@ -1461,11 +1469,12 @@ PRAGMA_ENABLE_DEPRECATION_WARNINGS
 						}
 						else
 						{
+							METASOUND_VERSIONING_LOG(Verbose, TEXT("Skipping wait for %s to versioning reference in async task: no task executing"), *InMetaSound.GetFullName());
 							continue;
 						}
 					}
 
-					METASOUND_VERSIONING_LOG(VeryVerbose, TEXT("Waiting for %s to finish versioning on async task..."), *InMetaSound.GetFullName());
+					METASOUND_VERSIONING_LOG(Verbose, TEXT("Waiting for %s to finish versioning on async task..."), *InMetaSound.GetFullName());
 					TaskToWaitFor.Wait();
 				}
 			}
@@ -1495,11 +1504,12 @@ PRAGMA_ENABLE_DEPRECATION_WARNINGS
 			}
 			else
 			{
+				METASOUND_VERSIONING_LOG(Verbose, TEXT("Skipping wait for %s to version on async task: no task executing"), *InMetaSound.GetFullName());
 				return;
 			}
 		}
 
-		METASOUND_VERSIONING_LOG(VeryVerbose, TEXT("Waiting for %s to finish versioning on async task..."), *InMetaSound.GetFullName());
+		METASOUND_VERSIONING_LOG(Verbose, TEXT("Waiting for %s to finish versioning on async task..."), *InMetaSound.GetFullName());
 		TaskToWaitFor.Wait();
 	}
 } // namespace Metasound::Frontend
diff --git a/Engine/Plugins/Runtime/Metasound/Source/MetasoundFrontend/Private/MetasoundFrontendPages.cpp b/Engine/Plugins/Runtime/Metasound/Source/MetasoundFrontend/Private/MetasoundFrontendPages.cpp
--- a/Engine/Plugins/Runtime/Metasound/Source/MetasoundFrontend/Private/MetasoundFrontendPages.cpp
+++ b/Engine/Plugins/Runtime/Metasound/Source/MetasoundFrontend/Private/MetasoundFrontendPages.cpp
@@ -10,6 +10,7 @@
 #include "Misc/Guid.h"
 #include "MetasoundFrontendDocument.h"
 #include "MetasoundFrontendDocumentBuilder.h"
+#include "MetasoundFrontendDocumentVersioning.h"
 #include "MetasoundTrace.h"
 
 namespace Metasound::Frontend
@@ -153,6 +154,12 @@ namespace Metasound::Frontend
 		METASOUND_LLM_SCOPE;
 		METASOUND_TRACE_CPUPROFILER_EVENT_SCOPE(Metasound::Frontend::StripUnusedPages);
 
+		checkf(InBuilder.GetConstDocumentChecked().Metadata.Version.Number == FVersioningManager::GetMaxDocumentVersion(),
+			TEXT("Document must be versioned prior to stripping pages. Expected %s, found %s on asset %s"), 
+			*FVersioningManager::GetMaxDocumentVersion().ToString(),
+			*InBuilder.GetConstDocumentChecked().Metadata.Version.Number.ToString(), 
+			*InBuilder.GetDebugName());
+
 		if (InTargetPages.Num() < 1)
 		{
 			UE_LOG(LogMetaSound, Error, TEXT("Must have at least one target page to strip unused pages in asset %s"), *InBuilder.GetDebugName());
