Fix crash when setting a MetaSound constructor input or output to a type that does not support constructor types (ex. Trigger)

#Synced-CL 49174767

diff --git a/Engine/Plugins/Runtime/Metasound/Source/MetasoundFrontend/Private/MetasoundFrontendDocumentBuilder.cpp b/Engine/Plugins/Runtime/Metasound/Source/MetasoundFrontend/Private/MetasoundFrontendDocumentBuilder.cpp
--- a/Engine/Plugins/Runtime/Metasound/Source/MetasoundFrontend/Private/MetasoundFrontendDocumentBuilder.cpp
+++ b/Engine/Plugins/Runtime/Metasound/Source/MetasoundFrontend/Private/MetasoundFrontendDocumentBuilder.cpp
@@ -6030,8 +6030,9 @@ bool FMetaSoundFrontendDocumentBuilder::SetGraphInputAccessType(FName InputName,
 bool FMetaSoundFrontendDocumentBuilder::SetGraphInputDataType(FName InputName, FName DataType)
 {
 	using namespace Metasound;
+	using namespace Metasound::Frontend;
 
-	if (Frontend::IDataTypeRegistry::Get().IsRegistered(DataType))
+	if (IDataTypeRegistry::Get().IsRegistered(DataType))
 	{
 		const int32* Index = DocumentCache->GetInterfaceCache().FindInputIndex(InputName);
 		if (!Index)
@@ -6044,6 +6045,19 @@ bool FMetaSoundFrontendDocumentBuilder::SetGraphInputDataType(FName InputName, F
 		FMetasoundFrontendClassInput& GraphInput = RootGraph.GetDefaultInterface().Inputs[*Index];
 		if (GraphInput.TypeName != DataType)
 		{
+			// If input is currently a constructor (value) type, but is being set to a data type 
+			// that does not support constructor types, set it back to reference
+			if (GraphInput.AccessType == EMetasoundFrontendVertexAccessType::Value)
+			{
+				FDataTypeRegistryInfo Info;
+				if (IDataTypeRegistry::Get().GetDataTypeInfo(DataType, Info))
+				{
+					if (!Info.bIsConstructorType)
+					{
+						GraphInput.AccessType = EMetasoundFrontendVertexAccessType::Reference;
+					}
+				}
+			}
 			GraphInput.TypeName = DataType;
 			GraphInput.ResetDefaults();
 
@@ -6417,6 +6431,19 @@ bool FMetaSoundFrontendDocumentBuilder::SetGraphOutputDataType(FName OutputName,
 	FMetasoundFrontendClassOutput& GraphOutput = RootGraph.GetDefaultInterface().Outputs[*Index];
 	if (GraphOutput.TypeName != DataType)
 	{
+		// If output is currently a constructor (value) type, but is being set to a data type 
+		// that does not support constructor types, set it back to reference
+		if (GraphOutput.AccessType == EMetasoundFrontendVertexAccessType::Value)
+		{
+			FDataTypeRegistryInfo Info;
+			if (IDataTypeRegistry::Get().GetDataTypeInfo(DataType, Info))
+			{
+				if (!Info.bIsConstructorType)
+				{
+					GraphOutput.AccessType = EMetasoundFrontendVertexAccessType::Reference;
+				}
+			}
+		}
 		GraphOutput.TypeName = DataType;
 		RootGraph.IterateGraphPages([this, &GraphOutput, &DataType](FMetasoundFrontendGraph& Graph)
 		{
