Add [MustPass] to selected MetaSound Automation Tests

#rnx
#Synced-CL 49218521

diff --git a/Engine/Plugins/Runtime/Metasound/Source/MetasoundEngineTest/Private/EngineTestMetaSoundAutomatedNodeTest_ValidOutputs.cpp b/Engine/Plugins/Runtime/Metasound/Source/MetasoundEngineTest/Private/EngineTestMetaSoundAutomatedNodeTest_ValidOutputs.cpp
--- a/Engine/Plugins/Runtime/Metasound/Source/MetasoundEngineTest/Private/EngineTestMetaSoundAutomatedNodeTest_ValidOutputs.cpp
+++ b/Engine/Plugins/Runtime/Metasound/Source/MetasoundEngineTest/Private/EngineTestMetaSoundAutomatedNodeTest_ValidOutputs.cpp
@@ -25,6 +25,16 @@ void FMetasoundAutomatedNodeTest_ValidOutputs::GetTests(TArray<FString>& OutBeau
 	constexpr bool bIncludeDeprecatedNodes = false;
 	GetAllRegisteredNativeNodes(OutBeautifiedNames, OutTestCommands, bIncludeDeprecatedNodes);
 
+	FAutomationTestFramework& TestFramework = FAutomationTestFramework::Get();
+	for (const FString& ComplexTestName : OutBeautifiedNames)
+	{
+		bool Registered = TestFramework.RegisterComplexAutomationTestTags(this, ComplexTestName, TEXT("[MustPass]"));
+		if (!Registered)
+		{
+			UE_LOG(LogMetaSound, Warning, TEXT("Failed to register automation test tags for test %s"), *ComplexTestName);
+		}
+	}
+
 	UE_LOG(LogMetaSound, Verbose, TEXT("Found %d metasound nodes to test"), OutTestCommands.Num());
 }
 
diff --git a/Engine/Plugins/Runtime/Metasound/Source/MetasoundEngineTest/Private/EngineTestMetaSoundBuilder.h b/Engine/Plugins/Runtime/Metasound/Source/MetasoundEngineTest/Private/EngineTestMetaSoundBuilder.h
--- a/Engine/Plugins/Runtime/Metasound/Source/MetasoundEngineTest/Private/EngineTestMetaSoundBuilder.h
+++ b/Engine/Plugins/Runtime/Metasound/Source/MetasoundEngineTest/Private/EngineTestMetaSoundBuilder.h
@@ -126,7 +126,9 @@ bool FBuilderRemoveFromRootLatentCommand::Update()
 }
 
 // Creates a collection of MetaSound patches from builders and exercises bindings all input and output interfaces by connecting and disconnecting using various builder API calls.
+
 IMPLEMENT_SIMPLE_AUTOMATION_TEST(FAudioMetaSoundBuilderInterfaceBindingConnectAndDisconnect, "Audio.Metasound.Builder.InterfaceBindingConnectAndDisconnect", EAutomationTestFlags::EditorContext | EAutomationTestFlags::EngineFilter)
+REGISTER_SIMPLE_AUTOMATION_TEST_TAGS(FAudioMetaSoundBuilderInterfaceBindingConnectAndDisconnect, "Audio.Metasound.Builder.InterfaceBindingConnectAndDisconnect", "[MustPass]")
 bool FAudioMetaSoundBuilderInterfaceBindingConnectAndDisconnect::RunTest(const FString& Parameters)
 {
 	using namespace EngineTestMetaSoundPatchBuilderPrivate;
@@ -196,6 +198,7 @@ bool FAudioMetaSoundBuilderInterfaceBindingConnectAndDisconnect::RunTest(const F
 
 // This test creates a MetaSound patch, then adds and connects sin oscillator, attempts to retrieve a default input set, then clears it, and finally retrieves the class default
 IMPLEMENT_SIMPLE_AUTOMATION_TEST(FAudioMetaSoundNodeClassQueryFunctions, "Audio.Metasound.Builder.NodeClassQueryFunctions", EAutomationTestFlags::EditorContext | EAutomationTestFlags::EngineFilter)
+REGISTER_SIMPLE_AUTOMATION_TEST_TAGS(FAudioMetaSoundNodeClassQueryFunctions, "Audio.Metasound.Builder.NodeClassQueryFunctions", "[MustPass]")
 bool FAudioMetaSoundNodeClassQueryFunctions::RunTest(const FString& Parameters)
 {
 	using namespace EngineTestMetaSoundPatchBuilderPrivate;
diff --git a/Engine/Plugins/Runtime/Metasound/Source/MetasoundEngineTest/Private/EngineTestMetasoundSource.cpp b/Engine/Plugins/Runtime/Metasound/Source/MetasoundEngineTest/Private/EngineTestMetasoundSource.cpp
--- a/Engine/Plugins/Runtime/Metasound/Source/MetasoundEngineTest/Private/EngineTestMetasoundSource.cpp
+++ b/Engine/Plugins/Runtime/Metasound/Source/MetasoundEngineTest/Private/EngineTestMetasoundSource.cpp
@@ -390,6 +390,7 @@ bool FMetaSoundSourceBuilderConnectNodesLatentCommand::Update()
 
 // This test creates a MetaSound from the legacy controller document editing system and attempts to play it directly.
 IMPLEMENT_SIMPLE_AUTOMATION_TEST(FAudioMetasoundSourceTest, "Audio.Metasound.PlayMetasoundSource", EAutomationTestFlags::EditorContext | EAutomationTestFlags::EngineFilter)
+REGISTER_SIMPLE_AUTOMATION_TEST_TAGS(FAudioMetasoundSourceTest, "Audio.Metasound.PlayMetasoundSource", "[MustPass]")
 bool FAudioMetasoundSourceTest::RunTest(const FString& Parameters)
 {
 	using namespace EngineTestMetaSoundSourcePrivate;
@@ -417,6 +418,7 @@ bool FAudioMetasoundSourceTest::RunTest(const FString& Parameters)
 
 // This test creates a MetaSound source from a SourceBuilder, adds a simple sine tone generator with a connected graph input frequency, and attempts to audition it.
 IMPLEMENT_SIMPLE_AUTOMATION_TEST(FAudioMetasoundSourceBuilderTest, "Audio.Metasound.Builder.AuditionMetasoundSource", EAutomationTestFlags::EditorContext | EAutomationTestFlags::EngineFilter)
+REGISTER_SIMPLE_AUTOMATION_TEST_TAGS(FAudioMetasoundSourceBuilderTest, "Audio.Metasound.Builder.AuditionMetasoundSource", "[MustPass]")
 bool FAudioMetasoundSourceBuilderTest::RunTest(const FString& Parameters)
 {
 	using namespace EngineTestMetaSoundPatchBuilderPrivate;
@@ -564,6 +566,7 @@ bool FAudioMetasoundSourceBuilderLiveUpdateNode::RunTest(const FString& Paramete
 // disconnects frequency input, sets the sinosc frequency literal value to a new value, and finally removes the literal value default to have it return to the
 // class default.
 IMPLEMENT_SIMPLE_AUTOMATION_TEST(FAudioMetasoundSourceBuilderLiveUpdateLiteral, "Audio.Metasound.Builder.LiveUpdateLiteralMetaSoundSource", EAutomationTestFlags::EditorContext | EAutomationTestFlags::EngineFilter)
+REGISTER_SIMPLE_AUTOMATION_TEST_TAGS(FAudioMetasoundSourceBuilderLiveUpdateLiteral, "Audio.Metasound.Builder.LiveUpdateLiteralMetaSoundSource", "[MustPass]")
 bool FAudioMetasoundSourceBuilderLiveUpdateLiteral::RunTest(const FString& Parameters)
 {
 	using namespace EngineTestMetaSoundPatchBuilderPrivate;
@@ -843,6 +846,7 @@ bool FAudioMetasoundSourceBuilderAddRemovePagedVariableNodes::RunTest(const FStr
 
 // This test creates a MetaSound source from a SourceBuilder, then adds, removes and connects a variable and its respective nodes (get/get delayed/set) to multiple pages. Does not playback or live-audition.
 IMPLEMENT_SIMPLE_AUTOMATION_TEST(FAudioMetasoundSourceBuilderAddRemoveNodes, "Audio.Metasound.Builder.AddRemoveNodes", EAutomationTestFlags::EditorContext | EAutomationTestFlags::EngineFilter)
+REGISTER_SIMPLE_AUTOMATION_TEST_TAGS(FAudioMetasoundSourceBuilderAddRemoveNodes, "Audio.Metasound.Builder.AddRemoveNodes", "[MustPass]")
 bool FAudioMetasoundSourceBuilderAddRemoveNodes::RunTest(const FString& Parameters)
 {
 	using namespace Audio;
@@ -924,6 +928,7 @@ bool FAudioMetasoundSourceBuilderAddRemoveNodes::RunTest(const FString& Paramete
 
 // This test attempts to load class info from asset tag data set on a MetaSound source and check its validity via the MetaSound Asset Subsystem.
 IMPLEMENT_SIMPLE_AUTOMATION_TEST(FAudioMetaSoundFindAssetClassInfo, "Audio.Metasound.Asset.FindAssetClassInfo", EAutomationTestFlags::EditorContext | EAutomationTestFlags::EngineFilter)
+REGISTER_SIMPLE_AUTOMATION_TEST_TAGS(FAudioMetaSoundFindAssetClassInfo, "Audio.Metasound.Asset.FindAssetClassInfo", "[MustPass]")
 bool FAudioMetaSoundFindAssetClassInfo::RunTest(const FString& Parameters)
 {
 #if WITH_EDITOR
diff --git a/Engine/Plugins/Runtime/Metasound/Source/MetasoundEngineTest/Private/MetasoundDynamicGeneratorTests.cpp b/Engine/Plugins/Runtime/Metasound/Source/MetasoundEngineTest/Private/MetasoundDynamicGeneratorTests.cpp
--- a/Engine/Plugins/Runtime/Metasound/Source/MetasoundEngineTest/Private/MetasoundDynamicGeneratorTests.cpp
+++ b/Engine/Plugins/Runtime/Metasound/Source/MetasoundEngineTest/Private/MetasoundDynamicGeneratorTests.cpp
@@ -158,6 +158,7 @@ namespace Metasound::Test::Generator::Dynamic
 		FMetasoundGeneratorDynamicVertexInterfaceUpdatedTest,
 		"Audio.Metasound.Generator.Dynamic.VertexInterfaceUpdated",
 		EAutomationTestFlags::EditorContext | EAutomationTestFlags::EngineFilter)
+	REGISTER_SIMPLE_AUTOMATION_TEST_TAGS(FMetasoundGeneratorDynamicVertexInterfaceUpdatedTest, "Audio.Metasound.Generator.Dynamic.VertexInterfaceUpdated", "[MustPass]")
 	bool FMetasoundGeneratorDynamicVertexInterfaceUpdatedTest::RunTest(const FString& Parameters)
 	{
 		// Make a dynamic generator
diff --git a/Engine/Plugins/Runtime/Metasound/Source/MetasoundEngineTest/Private/MetasoundInterfaceUpdateTests.cpp b/Engine/Plugins/Runtime/Metasound/Source/MetasoundEngineTest/Private/MetasoundInterfaceUpdateTests.cpp
--- a/Engine/Plugins/Runtime/Metasound/Source/MetasoundEngineTest/Private/MetasoundInterfaceUpdateTests.cpp
+++ b/Engine/Plugins/Runtime/Metasound/Source/MetasoundEngineTest/Private/MetasoundInterfaceUpdateTests.cpp
@@ -49,7 +49,11 @@ namespace Metasound::Test::InterfaceUpdate
 		FMetasoundInterfaceMinorUpdateTest,
 		"Audio.Metasound.Interface.MinorVersionUpdate",
 		EAutomationTestFlags::EditorContext | EAutomationTestFlags::EngineFilter)
-		bool FMetasoundInterfaceMinorUpdateTest::RunTest(const FString&)
+	REGISTER_SIMPLE_AUTOMATION_TEST_TAGS(
+		FMetasoundInterfaceMinorUpdateTest,
+		"Audio.Metasound.Interface.MinorVersionUpdate",
+		"[MustPass]")
+	bool FMetasoundInterfaceMinorUpdateTest::RunTest(const FString&)
 	{
 		using namespace Audio;
 		using namespace Metasound::Frontend;
@@ -170,4 +174,4 @@ namespace Metasound::Test::InterfaceUpdate
 	}
 }
 
-#endif
\ No newline at end of file
+#endif
diff --git a/Engine/Plugins/Runtime/Metasound/Source/MetasoundEngineTest/Private/MetasoundOutputTests.cpp b/Engine/Plugins/Runtime/Metasound/Source/MetasoundEngineTest/Private/MetasoundOutputTests.cpp
--- a/Engine/Plugins/Runtime/Metasound/Source/MetasoundEngineTest/Private/MetasoundOutputTests.cpp
+++ b/Engine/Plugins/Runtime/Metasound/Source/MetasoundEngineTest/Private/MetasoundOutputTests.cpp
@@ -11,9 +11,13 @@
 namespace Metasound::Test::GeneratorOutput
 {
 	IMPLEMENT_SIMPLE_AUTOMATION_TEST(
-	FMetasoundOutputIsTypeTest,
-	"Audio.Metasound.Output.IsType",
-	EAutomationTestFlags::EditorContext | EAutomationTestFlags::EngineFilter)
+		FMetasoundOutputIsTypeTest,
+		"Audio.Metasound.Output.IsType",
+		EAutomationTestFlags::EditorContext | EAutomationTestFlags::EngineFilter)
+	REGISTER_SIMPLE_AUTOMATION_TEST_TAGS(
+		FMetasoundOutputIsTypeTest,
+		"Audio.Metasound.Output.IsType",
+		"[MustPass]")
 	bool FMetasoundOutputIsTypeTest::RunTest(const FString&)
 	{
 		FMetaSoundOutput Output;
@@ -144,4 +148,4 @@ namespace Metasound::Test::GeneratorOutput
 	}
 }
 
-#endif
\ No newline at end of file
+#endif
