- Fix for node ID not being serialized on vertex metadata.
- Minor: Future proof output removal should vertex data be added for outputs in a config, and remove where necessary in builder API "RemoveGraphOutput" call

Issue: Value was not being saved on vertex metadata as wrong constructor was used.  This is risky to provide to the client API as they could incorrectly not set it and its duplicate data anyway. Fixing would require a doc version.

Solution: Vertex metadata is stored in a map keyed off node ID anyway, so just update experimental API to pass vertex through to metadata where necessary, which has more useful context than the NodeID and avoids storing additional NodeID field on vertex/user error of being responsible for caching.

#rnx
#Synced-CL 49193787

diff --git a/Engine/Plugins/Runtime/Metasound/Source/MetasoundEditor/Private/MetasoundNodeDetailCustomization.cpp b/Engine/Plugins/Runtime/Metasound/Source/MetasoundEditor/Private/MetasoundNodeDetailCustomization.cpp
--- a/Engine/Plugins/Runtime/Metasound/Source/MetasoundEditor/Private/MetasoundNodeDetailCustomization.cpp
+++ b/Engine/Plugins/Runtime/Metasound/Source/MetasoundEditor/Private/MetasoundNodeDetailCustomization.cpp
@@ -1783,10 +1783,13 @@ PRAGMA_ENABLE_DEPRECATION_WARNINGS
 							{
 								if (FMetaSoundFrontendDocumentBuilder* DocBuilder = GetBuilder())
 								{
-									ThisMetadata->OnPropertyChanged(InEvent, *DocBuilder);
-
 									if (UMetasoundEditorGraphInput* Input = Cast<UMetasoundEditorGraphInput>(GraphMember.Get()))
 									{
+										if (const FMetasoundFrontendClassVertex* InputVertex = Input->GetFrontendClassVertex())
+										{
+											ThisMetadata->OnPropertyChanged(InEvent, *DocBuilder, *InputVertex);
+										}
+										
 										if (UMetasoundEditorGraphMemberDefaultLiteral* Literal = Input->GetLiteral())
 										{
 											Literal->ForceRefresh();
diff --git a/Engine/Plugins/Runtime/Metasound/Source/MetasoundFrontend/Private/DocumentConfigurations/MetasoundFrontendPresetDocumentConfiguration.cpp b/Engine/Plugins/Runtime/Metasound/Source/MetasoundFrontend/Private/DocumentConfigurations/MetasoundFrontendPresetDocumentConfiguration.cpp
--- a/Engine/Plugins/Runtime/Metasound/Source/MetasoundFrontend/Private/DocumentConfigurations/MetasoundFrontendPresetDocumentConfiguration.cpp
+++ b/Engine/Plugins/Runtime/Metasound/Source/MetasoundFrontend/Private/DocumentConfigurations/MetasoundFrontendPresetDocumentConfiguration.cpp
@@ -101,12 +101,12 @@ bool FMetaSoundFrontendPresetDocumentConfiguration::ConfigureDocument(FMetaSound
 }
 #endif // WITH_EDITORONLY_DATA
 
-TInstancedStruct<FMetaSoundFrontendDocumentConfigurationVertexMetadata> FMetaSoundFrontendPresetDocumentConfiguration::ConstructVertexMetadata(const FMetaSoundFrontendDocumentBuilder& Builder, const FGuid& NodeID) const
+TInstancedStruct<FMetaSoundFrontendDocumentConfigurationVertexMetadata> FMetaSoundFrontendPresetDocumentConfiguration::ConstructVertexMetadata(const FMetaSoundFrontendDocumentBuilder& Builder, const FMetasoundFrontendClassVertex& InVertex) const
 {
 	using namespace Metasound::Frontend;
 
 	// Presets only need ConfigVertexMetadata for inputs, so don't construct otherwise
-	if (PresetDocConfigPrivate::IsInheritableInputNode(Builder, NodeID))
+	if (PresetDocConfigPrivate::IsInheritableInputNode(Builder, InVertex.NodeID))
 	{
 		return TInstancedStruct<FPresetVertexMetadata>::Make();
 	}
@@ -155,15 +155,15 @@ bool FMetaSoundFrontendPresetConfigurationVertexMetadata::IsDefaultEditable() co
 }
 #endif // WITH_EDITOR
 
-void FMetaSoundFrontendPresetConfigurationVertexMetadata::OnDefaultReset(const FMetaSoundFrontendDocumentBuilder& Builder)
+void FMetaSoundFrontendPresetConfigurationVertexMetadata::OnDefaultReset(const FMetaSoundFrontendDocumentBuilder& Builder, const FMetasoundFrontendClassVertex& InVertex)
 {
 	using namespace Metasound::Frontend;
 
 	bOverrideInheritedDefault = false;
-	PresetDocConfigPrivate::BroadcastInputInheritanceUpdated(const_cast<FMetaSoundFrontendDocumentBuilder&>(Builder), GetNodeID());
+	PresetDocConfigPrivate::BroadcastInputInheritanceUpdated(const_cast<FMetaSoundFrontendDocumentBuilder&>(Builder), InVertex.NodeID);
 }
 
-void FMetaSoundFrontendPresetConfigurationVertexMetadata::OnDefaultUpdated(const FMetaSoundFrontendDocumentBuilder& Builder)
+void FMetaSoundFrontendPresetConfigurationVertexMetadata::OnDefaultUpdated(const FMetaSoundFrontendDocumentBuilder& Builder, const FMetasoundFrontendClassVertex& InVertex)
 {
 	using namespace Metasound::Frontend;
 
@@ -172,7 +172,7 @@ void FMetaSoundFrontendPresetConfigurationVertexMetadata::OnDefaultUpdated(const
 	if (!Builder.IsConfiguringDocument())
 	{
 		bOverrideInheritedDefault = true;
-		PresetDocConfigPrivate::BroadcastInputInheritanceUpdated(const_cast<FMetaSoundFrontendDocumentBuilder&>(Builder), GetNodeID());
+		PresetDocConfigPrivate::BroadcastInputInheritanceUpdated(const_cast<FMetaSoundFrontendDocumentBuilder&>(Builder), InVertex.NodeID);
 	}
 }
 #undef LOCTEXT_NAMESPACE
diff --git a/Engine/Plugins/Runtime/Metasound/Source/MetasoundFrontend/Private/DocumentConfigurations/MetasoundFrontendVertexDocumentConfiguration.cpp b/Engine/Plugins/Runtime/Metasound/Source/MetasoundFrontend/Private/DocumentConfigurations/MetasoundFrontendVertexDocumentConfiguration.cpp
--- a/Engine/Plugins/Runtime/Metasound/Source/MetasoundFrontend/Private/DocumentConfigurations/MetasoundFrontendVertexDocumentConfiguration.cpp
+++ b/Engine/Plugins/Runtime/Metasound/Source/MetasoundFrontend/Private/DocumentConfigurations/MetasoundFrontendVertexDocumentConfiguration.cpp
@@ -10,26 +10,28 @@
 
 #define LOCTEXT_NAMESPACE "MetaSound"
 
-FMetaSoundFrontendDocumentConfigurationVertexMetadata* FMetaSoundFrontendDocumentVertexConfiguration::AddVertexMetadata(const FMetaSoundFrontendDocumentBuilder& Builder, const FGuid& NodeID)
+FMetaSoundFrontendDocumentConfigurationVertexMetadata* FMetaSoundFrontendDocumentVertexConfiguration::AddVertexMetadata(const FMetaSoundFrontendDocumentBuilder& Builder, const FMetasoundFrontendClassVertex& InVertex)
 {
-	if (TInstancedStruct<FMetaSoundFrontendDocumentConfigurationVertexMetadata>* Entry = VertexMetadata.Find(NodeID))
+	if (TInstancedStruct<FMetaSoundFrontendDocumentConfigurationVertexMetadata>* Entry = VertexMetadata.Find(InVertex.NodeID))
 	{
 		return Entry->GetMutablePtr();
 	}
 
 PRAGMA_DISABLE_EXPERIMENTAL_WARNINGS
-	TInstancedStruct<FMetaSoundFrontendDocumentConfigurationVertexMetadata> NewMetadata = ConstructVertexMetadata(Builder, NodeID);
+	TInstancedStruct<FMetaSoundFrontendDocumentConfigurationVertexMetadata> NewMetadata = ConstructVertexMetadata(Builder, InVertex);
 PRAGMA_ENABLE_EXPERIMENTAL_WARNINGS
 
 	if (NewMetadata.IsValid())
 	{
-		return VertexMetadata.Add(NodeID, MoveTemp(NewMetadata)).GetMutablePtr();
+		return VertexMetadata.Add(InVertex.NodeID, MoveTemp(NewMetadata)).GetMutablePtr();
 	}
 
 	return nullptr;
 }
 
-TInstancedStruct<FMetaSoundFrontendDocumentConfigurationVertexMetadata> FMetaSoundFrontendDocumentVertexConfiguration::ConstructVertexMetadata(const FMetaSoundFrontendDocumentBuilder& Builder, const FGuid& NodeID) const
+TInstancedStruct<FMetaSoundFrontendDocumentConfigurationVertexMetadata> FMetaSoundFrontendDocumentVertexConfiguration::ConstructVertexMetadata(
+	const FMetaSoundFrontendDocumentBuilder& Builder,
+	const FMetasoundFrontendClassVertex& InVertex) const
 {
 	return { };
 }
diff --git a/Engine/Plugins/Runtime/Metasound/Source/MetasoundFrontend/Private/MetasoundFrontendDocumentBuilder.cpp b/Engine/Plugins/Runtime/Metasound/Source/MetasoundFrontend/Private/MetasoundFrontendDocumentBuilder.cpp
--- a/Engine/Plugins/Runtime/Metasound/Source/MetasoundFrontend/Private/MetasoundFrontendDocumentBuilder.cpp
+++ b/Engine/Plugins/Runtime/Metasound/Source/MetasoundFrontend/Private/MetasoundFrontendDocumentBuilder.cpp
@@ -1291,7 +1291,7 @@ const FMetasoundFrontendNode* FMetaSoundFrontendDocumentBuilder::AddGraphInput(F
 		if (FDocumentVertexConfiguration* DocConfig = Document.Configuration.GetMutablePtr<FDocumentVertexConfiguration>())
 		{
 PRAGMA_DISABLE_EXPERIMENTAL_WARNINGS
-			DocConfig->AddVertexMetadata(*this, NewNode->GetID());
+			DocConfig->AddVertexMetadata(*this, NewInput);
 PRAGMA_ENABLE_EXPERIMENTAL_WARNINGS
 		}
 
@@ -1411,7 +1411,7 @@ const FMetasoundFrontendNode* FMetaSoundFrontendDocumentBuilder::AddGraphOutput(
 		if (FDocumentVertexConfiguration* DocConfig = Document.Configuration.GetMutablePtr<FDocumentVertexConfiguration>())
 		{
 PRAGMA_DISABLE_EXPERIMENTAL_WARNINGS
-			DocConfig->AddVertexMetadata(*this, NewOutput.NodeID);
+			DocConfig->AddVertexMetadata(*this, NewOutput);
 PRAGMA_ENABLE_EXPERIMENTAL_WARNINGS
 		}
 
@@ -5265,6 +5265,8 @@ bool FMetaSoundFrontendDocumentBuilder::RemoveGraphInput(FName InputName, bool b
 
 bool FMetaSoundFrontendDocumentBuilder::RemoveGraphOutput(FName OutputName)
 {
+	using namespace Metasound::Frontend;
+
 	bool bNodesRemoved = true;
 	FGuid ClassID;
 	FGuid NodeID;
@@ -5289,6 +5291,13 @@ bool FMetaSoundFrontendDocumentBuilder::RemoveGraphOutput(FName OutputName)
 
 	if (bNodesRemoved)
 	{
+#if WITH_EDITORONLY_DATA
+		if (FDocumentVertexConfiguration* DocConfig = Document.Configuration.GetMutablePtr<FDocumentVertexConfiguration>())
+		{
+			DocConfig->RemoveVertexMetadata(NodeID);
+		}
+#endif // WITH_EDITORONLY_DATA
+
 		TArray<FMetasoundFrontendClassOutput>& Outputs = Document.RootGraph.GetDefaultInterface().Outputs;
 		auto OutputNameMatches = [OutputName](const FMetasoundFrontendClassOutput& Output) { return Output.Name == OutputName; };
 		const int32 Index = Outputs.IndexOfByPredicate(OutputNameMatches);
@@ -5815,7 +5824,7 @@ bool FMetaSoundFrontendDocumentBuilder::ResetGraphInputDefault(FName InputName)
 		{
 			if (FMetaSoundFrontendDocumentConfigurationVertexMetadata* VertexMetadata = Config->FindVertexMetadata(Input.NodeID))
 			{
-				VertexMetadata->OnDefaultReset(*this);
+				VertexMetadata->OnDefaultReset(*this, Input);
 			}
 		}
 
@@ -6147,7 +6156,7 @@ bool FMetaSoundFrontendDocumentBuilder::SetGraphInputDefault(FName InputName, FM
 
 			if (const FDocumentConfigurationVertexMetadata* ConfigVertexMetadata = FindConstVertexMetadata<FDocumentConfigurationVertexMetadata>(Input.NodeID))
 			{
-				const_cast<FMetaSoundFrontendDocumentConfigurationVertexMetadata*>(ConfigVertexMetadata)->OnDefaultUpdated(*this);
+				const_cast<FMetaSoundFrontendDocumentConfigurationVertexMetadata*>(ConfigVertexMetadata)->OnDefaultUpdated(*this, Input);
 			}
 
 			return true;
@@ -6185,7 +6194,7 @@ bool FMetaSoundFrontendDocumentBuilder::SetGraphInputDefaults(FName InputName, T
 			{
 				if (FDocumentConfigurationVertexMetadata* VertexMetadata = DocConfig->FindVertexMetadata(Input.NodeID))
 				{
-					VertexMetadata->OnDefaultUpdated(*this);
+					VertexMetadata->OnDefaultUpdated(*this, Input);
 				}
 			}
 
diff --git a/Engine/Plugins/Runtime/Metasound/Source/MetasoundFrontend/Private/MetasoundFrontendDocumentVersioning.cpp b/Engine/Plugins/Runtime/Metasound/Source/MetasoundFrontend/Private/MetasoundFrontendDocumentVersioning.cpp
--- a/Engine/Plugins/Runtime/Metasound/Source/MetasoundFrontend/Private/MetasoundFrontendDocumentVersioning.cpp
+++ b/Engine/Plugins/Runtime/Metasound/Source/MetasoundFrontend/Private/MetasoundFrontendDocumentVersioning.cpp
@@ -1137,7 +1137,7 @@ PRAGMA_ENABLE_DEPRECATION_WARNINGS
 					for (const FMetasoundFrontendClassInput& Input : Document.RootGraph.GetDefaultInterface().Inputs)
 					{
 						// Some inputs don't support inheriting a default (ex. triggers), so ignore if no metadata is added.
-						if (FPresetVertexMetadata* VertexMetadata = PresetConfig.GetMutable().AddVertexMetadata<FPresetVertexMetadata>(OutBuilder, Input.NodeID))
+						if (FPresetVertexMetadata* VertexMetadata = PresetConfig.GetMutable().AddVertexMetadata<FPresetVertexMetadata>(OutBuilder, Input))
 						{
 PRAGMA_DISABLE_DEPRECATION_WARNINGS
 							const bool bOverrideInheritedDefault = !PresetOptions.InputsInheritingDefault.Contains(Input.Name);
diff --git a/Engine/Plugins/Runtime/Metasound/Source/MetasoundFrontend/Public/DocumentConfigurations/MetasoundFrontendDocumentVertexConfiguration.h b/Engine/Plugins/Runtime/Metasound/Source/MetasoundFrontend/Public/DocumentConfigurations/MetasoundFrontendDocumentVertexConfiguration.h
--- a/Engine/Plugins/Runtime/Metasound/Source/MetasoundFrontend/Public/DocumentConfigurations/MetasoundFrontendDocumentVertexConfiguration.h
+++ b/Engine/Plugins/Runtime/Metasound/Source/MetasoundFrontend/Public/DocumentConfigurations/MetasoundFrontendDocumentVertexConfiguration.h
@@ -3,6 +3,8 @@
 
 #include "MetasoundFrontendDocumentConfiguration.h"
 
+#include "MetasoundFrontendDocument.h"
+
 #include "MetasoundFrontendDocumentVertexConfiguration.generated.h"
 
 #define UE_API METASOUNDFRONTEND_API
@@ -17,22 +19,11 @@ struct FMetaSoundFrontendDocumentConfigurationVertexMetadata
 {
 	GENERATED_BODY()
 
-private:
-	UPROPERTY()
-	FGuid NodeID;
-
 public:
 	FMetaSoundFrontendDocumentConfigurationVertexMetadata() = default;
-	FMetaSoundFrontendDocumentConfigurationVertexMetadata(const FGuid& InNodeID)
-		: NodeID(InNodeID)
-	{
-	}
 
 	virtual ~FMetaSoundFrontendDocumentConfigurationVertexMetadata() = default;
 
-	UE_EXPERIMENTAL(5.8, "Document Configuration is in active development and subject to change")
-	const FGuid& GetNodeID() const { return NodeID; }
-
 #if WITH_EDITOR
 	// If true, default values are directly editable if the parent configuration lists member editing as enabled.
 	UE_EXPERIMENTAL(5.8, "Document Configuration is in active development and subject to change")
@@ -40,14 +31,14 @@ public:
 #endif // WITH_EDITOR
 
 	UE_EXPERIMENTAL(5.8, "Document Configuration is in active development and subject to change")
-	virtual void OnDefaultReset(const FMetaSoundFrontendDocumentBuilder& Builder) { }
+	virtual void OnDefaultReset(const FMetaSoundFrontendDocumentBuilder& Builder, const FMetasoundFrontendClassVertex& InVertex) { }
 
 	UE_EXPERIMENTAL(5.8, "Document Configuration is in active development and subject to change")
-	virtual void OnDefaultUpdated(const FMetaSoundFrontendDocumentBuilder& Builder) { }
+	virtual void OnDefaultUpdated(const FMetaSoundFrontendDocumentBuilder& Builder, const FMetasoundFrontendClassVertex& InVertex) { }
 
 #if WITH_EDITOR
 	UE_EXPERIMENTAL(5.8, "Document Configuration is in active development and subject to change")
-	virtual void OnPropertyChanged(const FPropertyChangedEvent& InEvent, FMetaSoundFrontendDocumentBuilder& OutBuilder) { }
+	virtual void OnPropertyChanged(const FPropertyChangedEvent& InEvent, FMetaSoundFrontendDocumentBuilder& OutBuilder, const FMetasoundFrontendClassVertex& InVertex) { }
 #endif // WITH_EDITOR
 };
 
@@ -59,28 +50,28 @@ struct FMetaSoundFrontendDocumentVertexConfiguration : public FMetaSoundFrontend
 
 	virtual ~FMetaSoundFrontendDocumentVertexConfiguration() = default;
 
-	// Creates and stores new VertexMetadata with the provided NodeID. Overwrites if metadata already exists with
+	// Creates and stores new VertexMetadata. Overwrites if metadata already exists with
 	// default constructed VertexMetadata if configuration supports Metadata.
-	UE_INTERNAL FMetaSoundFrontendDocumentConfigurationVertexMetadata* AddVertexMetadata(const FMetaSoundFrontendDocumentBuilder& Builder, const FGuid& NodeID);
+	UE_INTERNAL FMetaSoundFrontendDocumentConfigurationVertexMetadata* AddVertexMetadata(const FMetaSoundFrontendDocumentBuilder& Builder, const FMetasoundFrontendClassVertex& InVertex);
 
 	// Creates and stores new VertexMetadata with the provided type and name, asserting that the type created by the overloaded configuration
 	// matches the template type. Overwrites if metadata already exists with default constructed VertexMetadata. Returns null if configuration
 	// does not support Metadata.
 	template <typename TConfigVertexMetadata>
-	UE_INTERNAL TConfigVertexMetadata* AddVertexMetadata(const FMetaSoundFrontendDocumentBuilder& Builder, const FGuid& NodeID)
+	UE_INTERNAL TConfigVertexMetadata* AddVertexMetadata(const FMetaSoundFrontendDocumentBuilder& Builder, const FMetasoundFrontendClassVertex& InVertex)
 	{
-		if (TInstancedStruct<FMetaSoundFrontendDocumentConfigurationVertexMetadata>* Entry = VertexMetadata.Find(NodeID))
+		if (TInstancedStruct<FMetaSoundFrontendDocumentConfigurationVertexMetadata>* Entry = VertexMetadata.Find(InVertex.NodeID))
 		{
 			return &Entry->GetMutable<TConfigVertexMetadata>();
 		}
 
 PRAGMA_DISABLE_EXPERIMENTAL_WARNINGS
-		TInstancedStruct<FMetaSoundFrontendDocumentConfigurationVertexMetadata> NewMetadata = ConstructVertexMetadata(Builder, NodeID);
+		TInstancedStruct<FMetaSoundFrontendDocumentConfigurationVertexMetadata> NewMetadata = ConstructVertexMetadata(Builder, InVertex);
 PRAGMA_ENABLE_EXPERIMENTAL_WARNINGS
 
 		if (NewMetadata.IsValid())
 		{
-			return &VertexMetadata.Add(NodeID, MoveTemp(NewMetadata)).GetMutable<TConfigVertexMetadata>();
+			return &VertexMetadata.Add(InVertex.NodeID, MoveTemp(NewMetadata)).GetMutable<TConfigVertexMetadata>();
 		}
 
 		return nullptr;
@@ -129,7 +120,7 @@ PRAGMA_ENABLE_EXPERIMENTAL_WARNINGS
 
 protected:
 	UE_EXPERIMENTAL(5.8, "Document Configuration is in active development and subject to change")
-	UE_API virtual TInstancedStruct<FMetaSoundFrontendDocumentConfigurationVertexMetadata> ConstructVertexMetadata(const FMetaSoundFrontendDocumentBuilder& Builder, const FGuid& NodeID) const;
+	UE_API virtual TInstancedStruct<FMetaSoundFrontendDocumentConfigurationVertexMetadata> ConstructVertexMetadata(const FMetaSoundFrontendDocumentBuilder& Builder, const FMetasoundFrontendClassVertex& InVertex) const;
 
 private:
 	UPROPERTY()
diff --git a/Engine/Plugins/Runtime/Metasound/Source/MetasoundFrontend/Public/DocumentConfigurations/MetasoundFrontendPresetDocumentConfiguration.h b/Engine/Plugins/Runtime/Metasound/Source/MetasoundFrontend/Public/DocumentConfigurations/MetasoundFrontendPresetDocumentConfiguration.h
--- a/Engine/Plugins/Runtime/Metasound/Source/MetasoundFrontend/Public/DocumentConfigurations/MetasoundFrontendPresetDocumentConfiguration.h
+++ b/Engine/Plugins/Runtime/Metasound/Source/MetasoundFrontend/Public/DocumentConfigurations/MetasoundFrontendPresetDocumentConfiguration.h
@@ -2,6 +2,7 @@
 #pragma once
 
 #include "MetasoundDocumentInterface.h"
+#include "MetasoundFrontendDocument.h"
 #include "MetasoundFrontendDocumentVertexConfiguration.h"
 #include "UObject/NameTypes.h"
 
@@ -43,7 +44,7 @@ struct FMetaSoundFrontendPresetDocumentConfiguration final : public FMetaSoundFr
 #endif // WITH_EDITOR
 
 protected:
-	UE_INTERNAL UE_API virtual TInstancedStruct<FMetaSoundFrontendDocumentConfigurationVertexMetadata> ConstructVertexMetadata(const FMetaSoundFrontendDocumentBuilder& Builder, const FGuid& NodeID) const override;
+	UE_INTERNAL UE_API virtual TInstancedStruct<FMetaSoundFrontendDocumentConfigurationVertexMetadata> ConstructVertexMetadata(const FMetaSoundFrontendDocumentBuilder& Builder, const FMetasoundFrontendClassVertex& InVertex) const override;
 
 public:
 	UPROPERTY(EditAnywhere, BlueprintReadOnly, Category = "Preset Options")
@@ -61,8 +62,8 @@ struct FMetaSoundFrontendPresetConfigurationVertexMetadata : public FMetaSoundFr
 	virtual bool IsDefaultEditable() const override;
 #endif // WITH_EDITOR
 
-	virtual void OnDefaultReset(const FMetaSoundFrontendDocumentBuilder& Builder) override;
-	virtual void OnDefaultUpdated(const FMetaSoundFrontendDocumentBuilder& Builder) override;
+	virtual void OnDefaultReset(const FMetaSoundFrontendDocumentBuilder& Builder, const FMetasoundFrontendClassVertex& InVertex) override;
+	virtual void OnDefaultUpdated(const FMetaSoundFrontendDocumentBuilder& Builder, const FMetasoundFrontendClassVertex& InVertex) override;
 
 	// Enables overriding the input's inherited default value otherwise provided by the parent graph.
 	// Setting to true disables the configuration utilizing the input's default value if updated on the parent asset.
