MetaSound editor graph schema loctext key fix

#rnx
#Synced-CL 49389138

diff --git a/Engine/Plugins/Runtime/Metasound/Source/MetasoundEditor/Private/MetasoundEditorGraphSchema.cpp b/Engine/Plugins/Runtime/Metasound/Source/MetasoundEditor/Private/MetasoundEditorGraphSchema.cpp
--- a/Engine/Plugins/Runtime/Metasound/Source/MetasoundEditor/Private/MetasoundEditorGraphSchema.cpp
+++ b/Engine/Plugins/Runtime/Metasound/Source/MetasoundEditor/Private/MetasoundEditorGraphSchema.cpp
@@ -1309,8 +1309,8 @@ FMetasoundGraphSchemaAction_NewGraphVariable::FMetasoundGraphSchemaAction_NewGra
 		InDataType,
 		InLiteral,
 		FText(),
-		LOCTEXT("AddNewGraphOutputName", "Add Graph Variable"),
-		LOCTEXT("AddNewGraphOutputTooltip", "Adds new graph variable with given default value"),
+		LOCTEXT("AddNewGraphVariableName", "Add Graph Variable"),
+		LOCTEXT("AddNewGraphVariableTooltip", "Adds new graph variable with given default value"),
 		Metasound::Editor::EPrimaryContextGroup::Common)
 {
 }
@@ -2952,7 +2952,7 @@ void UMetasoundEditorGraphSchema::DroppedNonMetaSoundAssetsOnGraph(const TArray<
 	bool bModifiedObjects = false;
 	check(MetaSoundGraph);
 	UObject& MetaSound = MetaSoundGraph->GetMetasoundChecked();
-	FScopedTransaction Transaction(LOCTEXT("DropMetaSoundOnGraph", "Drop Asset On Graph"));
+	FScopedTransaction Transaction(LOCTEXT("DropAssetOnGraph", "Drop Asset On Graph"));
 
 	// Collect dropped objects
 	FName DataType; 
