PR #14212: Fixed: assert error msg for MetasoundFacade CreateOperator method

#Synced-CL 49627059

diff --git a/Engine/Plugins/Runtime/Metasound/Source/MetasoundGraphCore/Public/MetasoundFacade.h b/Engine/Plugins/Runtime/Metasound/Source/MetasoundGraphCore/Public/MetasoundFacade.h
--- a/Engine/Plugins/Runtime/Metasound/Source/MetasoundGraphCore/Public/MetasoundFacade.h
+++ b/Engine/Plugins/Runtime/Metasound/Source/MetasoundGraphCore/Public/MetasoundFacade.h
@@ -23,7 +23,7 @@ namespace Metasound
 				template<typename T, T> 
 				struct Helper;
 
-				// Check for "static TUniquePtr<IOperator> U::CreateOperator(const FBuildOperatorParams& Inparams, FBuildErrorArray& OutErrors)"
+				// Check for "static TUniquePtr<IOperator> U::CreateOperator(const FBuildOperatorParams& Inparams, FBuildResults& OutErrors)"
 				template<typename T>
 				static uint8 Check(Helper<TUniquePtr<IOperator>(*)(const FBuildOperatorParams& InParams, FBuildResults& OutResults), &T::CreateOperator>*);
 
@@ -72,7 +72,7 @@ namespace Metasound
 		static_assert(std::is_base_of<IOperator, OperatorType>::value, "To use the FNodeFacade constructor, the OperatorType must be derived from IOperator");
 
 		// Require static TUniquePtr<IOperator> OperatorType::CreateOperator(const FCreateOperatorParams&, TArray<TUniquePtr<IOperatorBuildError>>&) exists.
-		static_assert(MetasoundFacadePrivate::TIsFactoryMethodDeclared<OperatorType>::Value, "To use the FNodeFacade constructor, the OperatorType must have the static function \"static TUniquePtr<IOperator> OperatorType::CreateOperator(const FCreateOperatorParams&, TArray<TUniquePtr<IOperatorBuildError>>&)\"");
+		static_assert(MetasoundFacadePrivate::TIsFactoryMethodDeclared<OperatorType>::Value, "To use the FNodeFacade constructor, the OperatorType must have the static function \"static TUniquePtr<IOperator> OperatorType::CreateOperator(const FBuildOperatorParams&, TArray<TUniquePtr<IOperatorBuildError>>&)\"");
 
 
 		// Require static const FNodeClassMetadata& OperatorType::GetNodeInfo() exists.
