Fix MetaSound preset input nodes' order not matching referenced node vertex sort order

#Synced-CL 49629634

diff --git a/Engine/Plugins/Runtime/Metasound/Source/MetasoundFrontend/Private/MetasoundFrontendTransform.cpp b/Engine/Plugins/Runtime/Metasound/Source/MetasoundFrontend/Private/MetasoundFrontendTransform.cpp
--- a/Engine/Plugins/Runtime/Metasound/Source/MetasoundFrontend/Private/MetasoundFrontendTransform.cpp
+++ b/Engine/Plugins/Runtime/Metasound/Source/MetasoundFrontend/Private/MetasoundFrontendTransform.cpp
@@ -807,7 +807,8 @@ namespace Metasound
 			TArray<const FMetasoundFrontendVertex*> ReferencedNodeInputVertices;
 
 #if WITH_EDITORONLY_DATA
-			TArray<const FMetasoundFrontendNode*> InputTemplateNodes;
+			// Input template node IDs with associated connected vertex (from ReferencedNodeInputVertices)
+			TMap<const FMetasoundFrontendVertex*, const FGuid> InputTemplateNodeIDs;
 #endif // WITH_EDITORONLY_DATA
 
 			for (const FMetasoundFrontendClassInput& ClassInput : InClassInputs)
@@ -847,7 +848,7 @@ namespace Metasound
 				const FMetasoundFrontendNode* TemplateNode = InOutBuilderToTransform.AddNodeByTemplate(*InputTemplate, MoveTemp(Params));
 
 				check(TemplateNode);
-				InputTemplateNodes.Add(TemplateNode);
+				InputTemplateNodeIDs.Add(ReferencedNodeInputVertex, TemplateNode->GetID());
 
 				const FGuid TemplateNodeInputVertexID = TemplateNode->Interface.Inputs.Last().VertexID;
 				const FGuid TemplateNodeOutputVertexID = TemplateNode->Interface.Outputs.Last().VertexID;
@@ -894,12 +895,16 @@ namespace Metasound
 
 			InOutBuilderToTransform.SetInputStyle(MoveTemp(Style));
 
-			// Set editor node locations
+			// Set editor node locations, sorted to match referenced node input vertex order
 			FVector2D InputNodeLocation = FVector2D::ZeroVector;
-			for (const FMetasoundFrontendNode* TemplateNode : InputTemplateNodes)
+			for (const FMetasoundFrontendVertex* NodeInputVertex : ReferencedNodeInputVertices)
 			{
-				InOutBuilderToTransform.SetNodeLocation(TemplateNode->GetID(), InputNodeLocation);
-				InputNodeLocation += DisplayStyle::NodeLayout::DefaultOffsetY;
+				const FGuid* TemplateNodeID = InputTemplateNodeIDs.Find(NodeInputVertex);
+				if (TemplateNodeID)
+				{
+					InOutBuilderToTransform.SetNodeLocation(*TemplateNodeID, InputNodeLocation);
+					InputNodeLocation += DisplayStyle::NodeLayout::DefaultOffsetY;
+				}
 			}
 #endif // WITH_EDITOR
 		}
