Doc Config Widget/MVVM Support

[RN] MetaSound Configurations now support user-generated widgets and direct MVVM support
#Synced-CL 49750797

diff --git a/Engine/Plugins/Runtime/AudioWidgets/Source/AudioWidgets/Public/AudioMaterialSlate/AudioMaterialSlider.h b/Engine/Plugins/Runtime/AudioWidgets/Source/AudioWidgets/Public/AudioMaterialSlate/AudioMaterialSlider.h
--- a/Engine/Plugins/Runtime/AudioWidgets/Source/AudioWidgets/Public/AudioMaterialSlate/AudioMaterialSlider.h
+++ b/Engine/Plugins/Runtime/AudioWidgets/Source/AudioWidgets/Public/AudioMaterialSlate/AudioMaterialSlider.h
@@ -19,7 +19,7 @@ DECLARE_DYNAMIC_MULTICAST_DELEGATE_OneParam(FOnSliderFloatValueChangedEvent, flo
  *
  * * No Children
  */
-UCLASS(MinimalAPI)
+UCLASS(MinimalAPI, BlueprintType, Blueprintable, CustomFieldNotify)
 class UAudioMaterialSlider : public UWidget
 {
 	GENERATED_BODY()
@@ -97,7 +97,7 @@ public:
 public:
 
 	/** Called when the value is changed by slider. */
-	UPROPERTY(BlueprintAssignable, Category = "Widget Event")
+	UPROPERTY(BlueprintAssignable, FieldNotify, Category = "Widget Event")
 	FOnSliderFloatValueChangedEvent OnValueChanged;
 
 protected:
@@ -111,7 +111,7 @@ protected:
 private:
 
 	/**Default Value of the slider*/
-	UPROPERTY(EditAnywhere, BlueprintSetter = SetValue, BlueprintGetter = GetValue, Category = "Appearance", meta = (UIMin = "0", UIMax = "1"))
+	UPROPERTY(BlueprintReadWrite, EditAnywhere, FieldNotify, Setter = SetValue, Getter = GetValue, Category = "Appearance", meta = (UIMin = "0", UIMax = "1", AllowPrivateAccess))
 	float Value = 1.f;
 
 	/**Orientation of the slider*/
diff --git a/Engine/Plugins/Runtime/Metasound/Metasound.uplugin b/Engine/Plugins/Runtime/Metasound/Metasound.uplugin
--- a/Engine/Plugins/Runtime/Metasound/Metasound.uplugin
+++ b/Engine/Plugins/Runtime/Metasound/Metasound.uplugin
@@ -67,6 +67,10 @@
 		{
 			"Name": "WaveTable",
 			"Enabled": true
+		},
+		{
+			"Name": "ModelViewViewModel",
+			"Enabled": true		
 		}
 	]
 }
diff --git a/Engine/Plugins/Runtime/Metasound/Source/MetasoundEditor/MetasoundEditor.Build.cs b/Engine/Plugins/Runtime/Metasound/Source/MetasoundEditor/MetasoundEditor.Build.cs
--- a/Engine/Plugins/Runtime/Metasound/Source/MetasoundEditor/MetasoundEditor.Build.cs
+++ b/Engine/Plugins/Runtime/Metasound/Source/MetasoundEditor/MetasoundEditor.Build.cs
@@ -33,9 +33,7 @@ namespace UnrealBuildTool.Rules
 					"ToolMenus",
 					"ToolWidgets",
 					"WaveTable",
-					"WaveTableEditor",
-					"UMG"
-
+					"WaveTableEditor"
 				}
 			);
 
@@ -60,11 +58,13 @@ namespace UnrealBuildTool.Rules
 					"Engine",
 					"GraphEditor",
 					"InputCore",
+					"ModelViewViewModel",
 					"PropertyEditor",
 					"RenderCore",
 					"Slate",
 					"SlateCore",
 					"StructUtilsEditor",
+					"UMG",
 					"UnrealEd"
 				}
 			);
diff --git a/Engine/Plugins/Runtime/Metasound/Source/MetasoundEditor/Private/DocumentConfigurationViewModels/MetasoundDocumentConfigurationViewModelBase.cpp b/Engine/Plugins/Runtime/Metasound/Source/MetasoundEditor/Private/DocumentConfigurationViewModels/MetasoundDocumentConfigurationViewModelBase.cpp
new file mode 100644
--- /dev/null
+++ b/Engine/Plugins/Runtime/Metasound/Source/MetasoundEditor/Private/DocumentConfigurationViewModels/MetasoundDocumentConfigurationViewModelBase.cpp
@@ -0,0 +1,100 @@
+// Copyright Epic Games, Inc. All Rights Reserved.
+
+#pragma once
+
+#include "DocumentConfigurationViewModels/MetasoundDocumentConfigurationViewModelBase.h"
+#include "MetasoundFrontendDocumentModifyDelegates.h"
+#include "View/MVVMView.h"
+
+
+void UMetaSoundDocumentConfigurationViewModelBase::ConstructWidgetViewModel(
+	UUserWidget* UserWidget,
+	UMetaSoundBuilderBase& InBuilder,
+	TConstStructView<FMetaSoundFrontendDocumentConfiguration> Config,
+	TSubclassOf<UMetaSoundDocumentConfigurationViewModelBase> VMClass)
+{
+	using namespace Metasound::Engine;
+
+	if (!UserWidget)
+	{
+		return;
+	}
+
+	UMVVMView* View = UserWidget->GetExtension<UMVVMView>();
+	const UScriptStruct* ConfigStruct = Config.GetScriptStruct();
+	
+	if (View && ConfigStruct && VMClass.Get() != nullptr)
+	{
+		if (UMetaSoundDocumentConfigurationViewModelBase* ViewModel = NewObject<UMetaSoundDocumentConfigurationViewModelBase>(GetTransientPackage(), VMClass))
+		{
+			if (ViewModel->IsSupportedConfiguration(*ConfigStruct))
+			{
+				ViewModel->Initialize(&InBuilder);
+				View->SetViewModelByClass(ViewModel);
+			}
+		}
+	}
+}
+
+bool UMetaSoundDocumentConfigurationViewModelBase::IsSupportedConfiguration(const UScriptStruct& InStruct) const
+{
+	return true;
+}
+
+void UMetaSoundDocumentConfigurationViewModelBase::Initialize(UMetaSoundBuilderBase* InBuilder)
+{
+	using namespace Metasound::Frontend;
+	Builder = InBuilder;
+
+	BuilderListener = MakeShared<FBuilderListener>(this);
+	Builder->AddTransactionListener(BuilderListener->AsShared());
+
+	Reload();
+}
+
+void UMetaSoundDocumentConfigurationViewModelBase::Reload()
+{
+}
+
+void UMetaSoundDocumentConfigurationViewModelBase::FBuilderListener::OnBuilderReloaded(Metasound::Frontend::FDocumentModifyDelegates& OutDelegates)
+{
+	if (TStrongObjectPtr<UMetaSoundDocumentConfigurationViewModelBase> PinnedParent = Parent.Pin())
+	{
+		PinnedParent->Reload();
+	}
+
+	OutDelegates.OnDocumentConfigurationChanged.Remove(ConfigChangeDelegateHandle);
+	ConfigChangeDelegateHandle = OutDelegates.OnDocumentConfigurationChanged.AddSP(this, &UMetaSoundDocumentConfigurationViewModelBase::FBuilderListener::OnConfigChanged);
+}
+
+void UMetaSoundDocumentConfigurationViewModelBase::FBuilderListener::OnConfigChanged(const Metasound::Frontend::FDocumentConfigChangedArgs& InArgs)
+{
+	using namespace Metasound::Frontend;
+
+	// Property change covered by config's direct property changed delegate
+	// Covers if the whole struct is replaced, in which case just re-initialize
+	// or if editor transaction stack has been rolled forward or back.
+	if (InArgs.Type != Metasound::Frontend::EDocumentConfigChangeType::Property)
+	{
+		if (TStrongObjectPtr<UMetaSoundDocumentConfigurationViewModelBase> PinnedParent = Parent.Pin())
+		{
+			if (UMetaSoundBuilderBase* ThisBuilder = PinnedParent->Builder.Get())
+			{
+				FMetaSoundFrontendDocumentBuilder& DocBuilder = ThisBuilder->GetBuilder();
+				TConstStructView<FDocumentConfiguration> DocConfig = DocBuilder.GetConstDocumentConfiguration();
+				if (const UScriptStruct* ScriptStruct = DocConfig.GetScriptStruct();
+					ScriptStruct && PinnedParent->IsSupportedConfiguration(*ScriptStruct))
+				{
+					PinnedParent->Reload();
+				}
+				else
+				{
+					// Builder is now associated with MetaSound that no longer supports this configuration, so stop listening.
+					// This UObject (and owning object, likely an associated UUserWidget) will be GC'ed at an indeterminate
+					// later point.
+					DocBuilder.GetDocumentDelegates().OnDocumentConfigurationChanged.Remove(ConfigChangeDelegateHandle);
+				}
+			}
+		}
+	}
+}
diff --git a/Engine/Plugins/Runtime/Metasound/Source/MetasoundEditor/Private/DocumentConfigurationViewModels/RandomizerConfigurationViewModel.cpp b/Engine/Plugins/Runtime/Metasound/Source/MetasoundEditor/Private/DocumentConfigurationViewModels/RandomizerConfigurationViewModel.cpp
new file mode 100644
--- /dev/null
+++ b/Engine/Plugins/Runtime/Metasound/Source/MetasoundEditor/Private/DocumentConfigurationViewModels/RandomizerConfigurationViewModel.cpp
@@ -0,0 +1,159 @@
+// Copyright Epic Games, Inc. All Rights Reserved.
+
+#include "DocumentConfigurationViewModels/RandomizerConfigurationViewModel.h"
+
+#include "DocumentConfigurations/MetasoundRandomizerDocumentConfiguration.h"
+
+
+bool UMetaSoundRandomizerViewModel::IsSupportedConfiguration(const UScriptStruct& InStruct) const
+{
+	return InStruct.IsChildOf(FMetaSoundRandomizerDocumentConfiguration::StaticStruct());
+}
+
+bool UMetaSoundRandomizerViewModel::GetIsOneShot() const
+{
+	return bIsOneShot;
+}
+
+const TArray<TObjectPtr<USoundWave>>& UMetaSoundRandomizerViewModel::GetSounds() const
+{
+	return Sounds;
+}
+
+void UMetaSoundRandomizerViewModel::Reload()
+{
+	Super::Reload();
+
+	if (const FMetaSoundRandomizerDocumentConfiguration* Config = GetConstConfiguration<FMetaSoundRandomizerDocumentConfiguration>())
+	{
+		UE_MVVM_SET_PROPERTY_VALUE(bIsOneShot, Config->bIsOneShot);
+		UE_MVVM_SET_PROPERTY_VALUE(Pitch, Config->Pitch);
+		UE_MVVM_SET_PROPERTY_VALUE(Sounds, Config->Sounds);
+
+		Config->GetPropertyChangedDelegate().Remove(PropertyChangedHandle);
+		PropertyChangedHandle = Config->GetPropertyChangedDelegate().AddUObject(this, &UMetaSoundRandomizerViewModel::OnPropertyChanged);
+	}
+}
+
+void UMetaSoundRandomizerViewModel::OnPropertyChanged(const FPropertyChangedEvent& InEvent)
+{
+	if (const FMetaSoundRandomizerDocumentConfiguration* Config = GetConstConfiguration<FMetaSoundRandomizerDocumentConfiguration>())
+	{
+		if (InEvent.GetPropertyName() == GET_MEMBER_NAME_CHECKED(FMetaSoundRandomizerDocumentConfiguration, bIsOneShot))
+		{
+			UE_MVVM_SET_PROPERTY_VALUE(bIsOneShot, Config->bIsOneShot);
+		}
+		else if (InEvent.GetPropertyName() == GET_MEMBER_NAME_CHECKED(FMetaSoundRandomizerDocumentConfiguration, Pitch))
+		{
+			UE_MVVM_SET_PROPERTY_VALUE(Pitch, Config->Pitch);
+		}
+		else if (InEvent.GetPropertyName() == GET_MEMBER_NAME_CHECKED(FVector2f, X))
+		{
+			UE_MVVM_SET_PROPERTY_VALUE(Pitch, Config->Pitch);
+			BroadcastFieldValueChanged(FFieldNotificationClassDescriptor::GetPitchMin);
+		}
+		else if (InEvent.GetPropertyName() == GET_MEMBER_NAME_CHECKED(FVector2f, Y))
+		{
+			UE_MVVM_SET_PROPERTY_VALUE(Pitch, Config->Pitch);
+			BroadcastFieldValueChanged(FFieldNotificationClassDescriptor::GetPitchMax);
+		}
+		else if (InEvent.GetPropertyName() == GET_MEMBER_NAME_CHECKED(FMetaSoundRandomizerDocumentConfiguration, Sounds))
+		{
+			UE_MVVM_SET_PROPERTY_VALUE(Sounds, Config->Sounds);
+		}
+	}
+}
+
+FVector2f UMetaSoundRandomizerViewModel::GetPitch() const
+{
+	return Pitch;
+}
+
+float UMetaSoundRandomizerViewModel::GetPitchMin() const
+{
+	return Pitch.X;
+}
+
+float UMetaSoundRandomizerViewModel::GetPitchMax() const
+{
+	return Pitch.Y;
+}
+
+void UMetaSoundRandomizerViewModel::SetIsOneShot(bool bInIsOneShot)
+{
+	if (Builder && UE_MVVM_SET_PROPERTY_VALUE(bIsOneShot, bInIsOneShot))
+	{
+		Builder->SetConfigurationProperties<FMetaSoundRandomizerDocumentConfiguration>([this](FMetaSoundRandomizerDocumentConfiguration& RandomConfig)
+		{
+			RandomConfig.bIsOneShot = bIsOneShot;
+		});
+	}
+}
+
+void UMetaSoundRandomizerViewModel::SetPitch(const FVector2f& InPitch)
+{
+	if (Builder && UE_MVVM_SET_PROPERTY_VALUE(Pitch, InPitch))
+	{
+		Builder->SetConfigurationProperties<FMetaSoundRandomizerDocumentConfiguration>([this](FMetaSoundRandomizerDocumentConfiguration& RandomConfig)
+		{
+			RandomConfig.Pitch = Pitch;
+			BroadcastFieldValueChanged(FFieldNotificationClassDescriptor::GetPitchMin);
+			BroadcastFieldValueChanged(FFieldNotificationClassDescriptor::GetPitchMax);
+		});
+	}
+}
+
+void UMetaSoundRandomizerViewModel::SetPitchMax(float PitchMax)
+{
+	FVector2f NewPitch = Pitch;
+	NewPitch.Y = PitchMax;
+	if (Builder && UE_MVVM_SET_PROPERTY_VALUE(Pitch, NewPitch))
+	{
+		Builder->SetConfigurationProperties<FMetaSoundRandomizerDocumentConfiguration>([this](FMetaSoundRandomizerDocumentConfiguration& RandomConfig)
+		{
+			RandomConfig.Pitch = Pitch;
+			BroadcastFieldValueChanged(FFieldNotificationClassDescriptor::GetPitchMax);
+		});
+	}
+}
+
+void UMetaSoundRandomizerViewModel::SetPitchMin(float PitchMin)
+{
+	FVector2f NewPitch = Pitch;
+	NewPitch.X = PitchMin;
+	if (Builder && UE_MVVM_SET_PROPERTY_VALUE(Pitch, NewPitch))
+	{
+		Builder->SetConfigurationProperties<FMetaSoundRandomizerDocumentConfiguration>([this](FMetaSoundRandomizerDocumentConfiguration& RandomConfig)
+		{
+			RandomConfig.Pitch = Pitch;
+			BroadcastFieldValueChanged(FFieldNotificationClassDescriptor::GetPitchMin);
+		});
+	}
+}
+
+void UMetaSoundRandomizerViewModel::SetNumSounds(int32 NumSounds)
+{
+	if (NumSounds < 0)
+	{
+		return;
+	}
+
+	TArray<TObjectPtr<USoundWave>> NewSounds = Sounds;
+	NewSounds.SetNumZeroed(NumSounds);
+	SetSounds(MoveTemp(NewSounds));
+}
+
+void UMetaSoundRandomizerViewModel::SetSounds(TArray<TObjectPtr<USoundWave>> InSounds)
+{
+	if (Builder)
+	{
+		FMetaSoundRandomizerDocumentConfiguration::RemoveInvalidSounds(Builder->GetConstBuilder(), InSounds);
+		if (UE_MVVM_SET_PROPERTY_VALUE(Sounds, InSounds))
+		{
+			Builder->SetConfigurationProperties<FMetaSoundRandomizerDocumentConfiguration>([this](FMetaSoundRandomizerDocumentConfiguration& RandomConfig)
+			{
+				RandomConfig.Sounds = Sounds;
+			});
+		}
+	}
+}
diff --git a/Engine/Plugins/Runtime/Metasound/Source/MetasoundEditor/Private/MetasoundAssetDefinitions.cpp b/Engine/Plugins/Runtime/Metasound/Source/MetasoundEditor/Private/MetasoundAssetDefinitions.cpp
--- a/Engine/Plugins/Runtime/Metasound/Source/MetasoundEditor/Private/MetasoundAssetDefinitions.cpp
+++ b/Engine/Plugins/Runtime/Metasound/Source/MetasoundEditor/Private/MetasoundAssetDefinitions.cpp
@@ -425,7 +425,7 @@ namespace Metasound::Editor
 
 			if (Section)
 			{
-				const FText ConfigToolTip = Config.GetPtr()->GetToolTip();
+				const FText ConfigToolTip = Config.GetPtr()->GetEditorOptions().ToolTip;
 				auto AddEntryLambda = [CreateAction = MoveTemp(CreateFromConfigAction), ScriptStruct, ConfigToolTip, ClassName, AssetIcon](FToolMenuSection& InSection)
 				{
 					using namespace Metasound::Editor;
@@ -823,6 +823,8 @@ namespace MenuExtension_MetaSoundTemplate
 		const UClass* MetaSoundClass = TClass::StaticClass();
 		check(MetaSoundClass);
 
+		const TSoftClassPtr<UClass> MetaSoundSoftClass(MetaSoundClass);
+
 		for (TObjectIterator<UStruct> It; It; ++It)
 		{
 			UScriptStruct* Struct = Cast<UScriptStruct>(*It);
@@ -851,7 +853,7 @@ namespace MenuExtension_MetaSoundTemplate
 
 			TInstancedStruct<FMetaSoundFrontendDocumentConfiguration> ConfigInstance;
 			ConfigInstance.InitializeAsScriptStruct(Struct);
-			TArray<const UClass*> SupportedMetaSoundClasses = ConfigInstance.GetPtr()->GetSupportedMetaSoundClasses();
+			TArray<TSoftClassPtr<UObject>> SupportedMetaSoundClasses = ConfigInstance.Get().GetEditorOptions().MetaSoundClasses;
 			if (SupportedMetaSoundClasses.IsEmpty())
 			{
 				IMetasoundUObjectRegistry::Get().IterateRegisteredUClasses([&SupportedMetaSoundClasses](UClass& RegClass)
@@ -860,7 +862,7 @@ namespace MenuExtension_MetaSoundTemplate
 				});
 			}
 
-			if (SupportedMetaSoundClasses.Contains(TClass::StaticClass()))
+			if (SupportedMetaSoundClasses.Contains(MetaSoundSoftClass))
 			{
 				FToolUIAction DocConfigCreateAction;
 				DocConfigCreateAction.ExecuteAction = FToolMenuExecuteAction::CreateLambda([Struct](const FToolMenuContext& MenuContext)
@@ -870,17 +872,23 @@ namespace MenuExtension_MetaSoundTemplate
 					return AssetDefinitionsPrivate::ExecuteCreateMetaSoundWithDocConfig<TClass, TFactoryClass>(MenuContext, MoveTemp(NewInstance));
 				});
 
-				TArray<const UClass*> AssetActionClasses = ConfigInstance.GetPtr()->GetSupportedAssetActionClasses();
+				const TArray<TSoftClassPtr<UObject>>& AssetActionClasses = ConfigInstance.GetPtr()->GetEditorOptions().AssetActionClasses;
 				if (AssetActionClasses.IsEmpty())
 				{
 					AssetDefinitionsPrivate::AddConfigActions(AssetIcon, *MetaSoundClass, *MetaSoundClass, MoveTemp(DocConfigCreateAction), ConfigInstance);
 				}
 				else
 				{
-					for (const UClass* Class : AssetActionClasses)
+					for (const TSoftClassPtr<UObject>& Class : AssetActionClasses)
 					{
-						check(Class);
-						AssetDefinitionsPrivate::AddConfigActions(AssetIcon, *Class, *MetaSoundClass, DocConfigCreateAction, ConfigInstance);
+						const UClass* ClassPtr = Class.Get();
+						if (!ClassPtr)
+						{
+							ClassPtr = Class.LoadSynchronous();
+						}
+						check(ClassPtr);
+
+						AssetDefinitionsPrivate::AddConfigActions(AssetIcon, *ClassPtr, *MetaSoundClass, DocConfigCreateAction, ConfigInstance);
 					}
 				}
 			}
diff --git a/Engine/Plugins/Runtime/Metasound/Source/MetasoundEditor/Private/MetasoundDetailCustomization.cpp b/Engine/Plugins/Runtime/Metasound/Source/MetasoundEditor/Private/MetasoundDetailCustomization.cpp
--- a/Engine/Plugins/Runtime/Metasound/Source/MetasoundEditor/Private/MetasoundDetailCustomization.cpp
+++ b/Engine/Plugins/Runtime/Metasound/Source/MetasoundEditor/Private/MetasoundDetailCustomization.cpp
@@ -201,8 +201,8 @@ namespace Metasound::Editor
 								FMetaSoundFrontendDocumentBuilder& DocBuilder = Builder->GetBuilder();
 								if (DocBuilder.IsValid())
 								{
-									DocBuilder.GetDocumentDelegates().OnDocumentConfigurationChanged.Broadcast({ });
-									DocBuilder.ConfigureDocument();
+									DocBuilder.ConfigureDocument(true /* bLogFailure */);
+									DocBuilder.GetDocumentDelegates().OnDocumentConfigurationChanged.Broadcast({ EDocumentConfigChangeType::Struct });
 								}
 							}
 						});
@@ -215,15 +215,36 @@ namespace Metasound::Editor
 								FMetaSoundFrontendDocumentBuilder& DocBuilder = Builder->GetBuilder();
 								if (DocBuilder.IsValid())
 								{
-									if (const FDocumentConfiguration* Config = DocBuilder.GetConstDocumentConfiguration().GetPtr())
+									if (FDocumentConfiguration* Config = DocBuilder.GetDocumentConfiguration().GetPtr())
 									{
-										const_cast<FDocumentConfiguration*>(Config)->OnPropertyChanged(InEvent, DocBuilder);
+										DocBuilder.ConfigureDocument(true /* bLogFailure */);
+										Config->OnPropertyChanged(InEvent, DocBuilder);
+										DocBuilder.GetDocumentDelegates().OnDocumentConfigurationChanged.Broadcast({ EDocumentConfigChangeType::Property });
 									}
 								}
 							}
 						});
 						DocConfigHandle->SetOnChildPropertyValueChangedWithData(OnConfigPropChanged);
-						GeneralCategoryBuilder.AddProperty(DocConfigHandle);
+						GeneralCategoryBuilder.AddProperty(DocConfigHandle)
+							.Visibility(TAttribute<EVisibility>::CreateSPLambda(this, [this]()
+							{
+								if (Builder.IsValid() && !FEditor::IgnoreDocConfigVisibilitySettings())
+								{
+									FMetaSoundFrontendDocumentBuilder& DocBuilder = Builder->GetBuilder();
+									if (DocBuilder.IsValid())
+									{
+										if (const FDocumentConfiguration* Config = DocBuilder.GetConstDocumentConfiguration().GetPtr())
+										{
+											return Config->GetEditorOptions().bConfigurationPropertiesVisible
+												? EVisibility::Visible
+												: EVisibility::Collapsed;
+										}
+									}
+								}
+
+								return EVisibility::Visible;
+							}
+						));
 					}
 				}
 
diff --git a/Engine/Plugins/Runtime/Metasound/Source/MetasoundEditor/Private/MetasoundEditor.cpp b/Engine/Plugins/Runtime/Metasound/Source/MetasoundEditor/Private/MetasoundEditor.cpp
--- a/Engine/Plugins/Runtime/Metasound/Source/MetasoundEditor/Private/MetasoundEditor.cpp
+++ b/Engine/Plugins/Runtime/Metasound/Source/MetasoundEditor/Private/MetasoundEditor.cpp
@@ -10,6 +10,7 @@
 #include "Blueprint/UserWidget.h"
 #include "Components/AudioComponent.h"
 #include "DetailLayoutBuilder.h"
+#include "DocumentConfigurationViewModels/MetasoundDocumentConfigurationViewModelBase.h"
 #include "DocumentConfigurations/MetasoundDocumentConfigurator.h"
 #include "DocumentConfigurations/MetasoundFrontendDocumentConfiguration.h"
 #include "DocumentConfigurations/MetasoundFrontendPresetDocumentConfiguration.h"
@@ -43,6 +44,7 @@
 #include "MetasoundAssetSubsystem.h"
 #include "MetasoundBuilderSubsystem.h"
 #include "MetasoundDocumentBuilderRegistry.h"
+#include "MetasoundDocumentConfigurationWidgetInterface.h"
 #include "MetasoundDocumentInterface.h"
 #include "MetasoundEditorCommands.h"
 #include "MetasoundEditorDocumentClipboardUtils.h"
@@ -100,7 +102,9 @@
 #include "Widgets/SWindow.h"
 #include "Styling/StyleColors.h"
 #include "UObject/Class.h"
+#include "UObject/NameTypes.h"
 #include "UObject/UObjectIterator.h"
+#include "WorkflowOrientedApp/WorkflowTabFactory.h"
 
 #include UE_INLINE_GENERATED_CPP_BY_NAME(MetasoundEditor)
 
@@ -134,6 +138,129 @@ namespace Metasound
 				TEXT("If true, all settings on a Document Configuration for what is and is not visible in editor are ignored and all tabs are invokable.")
 				TEXT("Default: false"),
 				ECVF_Default);
+
+			TFunction<bool(UUserWidget*)> GetPresetWidgetFilterFunction(const FMetaSoundFrontendDocumentBuilder& DocBuilder)
+			{
+				using namespace Frontend;
+
+				if (!DocBuilder.IsValid() || DocBuilder.GetIsGraphEditable())
+				{
+					return { };
+				}
+
+				UWorld* World = GEditor ? GEditor->GetEditorWorldContext().World() : nullptr;
+				UAudioWidgetSubsystem* WidgetSubsystem = GEngine ? GEngine->GetEngineSubsystem<UAudioWidgetSubsystem>() : nullptr;
+				if (!World || !WidgetSubsystem)
+				{
+					return { };
+				}
+
+				// Referenced assets
+				FMetasoundAssetBase& ThisAsset = DocBuilder.GetMetasoundAsset();
+				TArray<FMetasoundAssetBase*> ReferencedAssets;
+				const bool bIsValidPreset = IMetaSoundAssetManager::GetChecked().GetReferencedPresetHierarchy(ThisAsset, ReferencedAssets);
+				if (!bIsValidPreset)
+				{
+					return { };
+				}
+
+				TSet<FNodeClassRegistryKey> ReferencedRegistryKeys;
+				Algo::Transform(ReferencedAssets, ReferencedRegistryKeys, [](const FMetasoundAssetBase* ReferencedAsset)
+				{
+					if (ensure(ReferencedAsset))
+					{
+						TScriptInterface<const IMetaSoundDocumentInterface> RefDocInterface = ReferencedAsset->GetOwningAsset();
+						const FMetasoundFrontendDocument& ReferencedDocument = RefDocInterface->GetConstDocument();
+						const FMetasoundFrontendGraphClass& ReferencedGraphClass = ReferencedDocument.RootGraph;
+						return FNodeClassRegistryKey(ReferencedGraphClass);
+					}
+					return FNodeClassRegistryKey();
+				});
+
+				// Filter whether the widget supports this MetaSound
+				return [ThisRegistryKey = FNodeRegistryKey(DocBuilder.GetConstDocumentChecked().RootGraph), ReferencedRegistryKeys](UUserWidget* UserWidget)
+				{
+					bool bSupportAllPresets = false;
+					TArray<TScriptInterface<IMetaSoundDocumentInterface>> ExcludedMetaSounds;
+					TArray<TScriptInterface<IMetaSoundDocumentInterface>> IncludedMetaSounds;
+					IMetaSoundPresetWidgetInterface::Execute_GetSupportedMetaSounds(UserWidget, bSupportAllPresets, ExcludedMetaSounds, IncludedMetaSounds);
+
+					// Show if all presets supported and this MetaSound is not excluded 
+					if (bSupportAllPresets)
+					{
+						if (!ExcludedMetaSounds.IsEmpty())
+						{
+							for (TScriptInterface<IMetaSoundDocumentInterface> ExcludedMetaSound : ExcludedMetaSounds)
+							{
+								if (!ExcludedMetaSound)
+								{
+									continue;
+								}
+
+								// This MetaSound or the referenced MetaSound is excluded 
+								const FMetasoundFrontendGraphClass& ExcludedGraphClass = ExcludedMetaSound->GetConstDocument().RootGraph;
+								const FNodeRegistryKey ExcludedRegistryKey(ExcludedGraphClass);
+
+								if (ThisRegistryKey == ExcludedRegistryKey || ReferencedRegistryKeys.Contains(ExcludedRegistryKey))
+								{
+									return false;
+								}
+							}
+						}
+						return true;
+					}
+
+					// This MetaSound is in the inclusion list 
+					for (TScriptInterface<IMetaSoundDocumentInterface> IncludedMetaSound : IncludedMetaSounds)
+					{
+						if (!IncludedMetaSound)
+						{
+							continue;
+						}
+
+						// This MetaSound or the referenced MetaSound is included 
+						const FMetasoundFrontendGraphClass& IncludedGraphClass = IncludedMetaSound->GetConstDocument().RootGraph;
+						const FNodeRegistryKey IncludedRegistryKey(IncludedGraphClass);
+
+						if (ThisRegistryKey == IncludedRegistryKey || ReferencedRegistryKeys.Contains(IncludedRegistryKey))
+						{
+							return true;
+						}
+					}
+					return false;
+				};
+			}
+
+			void SanitizeTabInfo(const UScriptStruct& ConfigStruct, FMetaSoundDocumentConfigurationWidgetTabInfo& InOutTabInfo)
+			{
+				if (InOutTabInfo.TabName.IsNone())
+				{
+					FNameBuilder NameBuilder;
+					NameBuilder.Appendf(TEXT("MetasoundEditor_%s"), *ConfigStruct.GetName());
+					InOutTabInfo.TabName = *NameBuilder;
+					if (InOutTabInfo.TabDisplayName.IsEmpty())
+					{
+						InOutTabInfo.TabDisplayName = ConfigStruct.GetDisplayNameText();
+					}
+				}
+
+				if (InOutTabInfo.TabIcon.IsNone())
+				{
+					InOutTabInfo.TabIcon = "MetasoundEditor.MetasoundPatch.Thumbnail";
+				}
+
+				if (InOutTabInfo.ToolTip.IsEmpty())
+				{
+					InOutTabInfo.ToolTip = FText::Format(
+						LOCTEXT("ConfigWidgetTab_DefaultToolTipFormat", "Custom control for MetaSound '{0}' configuration"),
+						ConfigStruct.GetDisplayNameText());
+				}
+
+				if (InOutTabInfo.StyleName.IsNone())
+				{
+					InOutTabInfo.StyleName = "MetaSoundStyle";
+				}
+			};
 		} // namespace AssetEditorPrivate
 
 		bool IsPreviewingPageInputDefault(const FMetaSoundFrontendDocumentBuilder& Builder, const FMetasoundFrontendClassInput& InClassInput, const FGuid& InPageID)
@@ -207,7 +334,7 @@ namespace Metasound
 			TConstStructView<FMetaSoundFrontendDocumentConfiguration> ConfigView = Builder.GetConstDocumentConfiguration();
 			if (ConfigView.IsValid())
 			{
-				return ConfigView.GetPtr()->IsPageGraphEditingEnabled();
+				return ConfigView.GetPtr()->GetEditorOptions().bPageGraphEditingEnabled;
 			}
 
 			return true;
@@ -253,7 +380,8 @@ namespace Metasound
 			const FName Interfaces = "MetasoundEditor_Interfaces";
 			const FName Pages = "MetasoundEditor_Pages";
 			const FName Find = "MetasoundEditor_Find";
-			const FName DocConfigWidget = "MetasoundEditor_DocConfigWidget";
+
+			const FName DocConfigWidgets_Factory = "MetasoundEditor_DocConfigWidgetFactory";
 		} // namespace TabNamesPrivate
 
 		static const TArray<FText> NodeSectionNames
@@ -957,6 +1085,99 @@ namespace Metasound
 			}
 		};
 
+		FEditor::FDocConfigWidgetSummoner::FDocConfigWidgetSummoner(FName FactoryName, FEditor::FBlueprintId InBlueprintId, TSharedPtr<class FAssetEditorToolkit> InHostingApp)
+			: FWorkflowTabFactory(FactoryName, InHostingApp)
+			, BlueprintId(InBlueprintId)
+		{
+			using namespace Frontend;
+
+			bIsSingleton = false;
+
+			TabIcon = FSlateIcon(FAppStyle::GetAppStyleSetName(), "LevelEditor.Tabs.Details");
+
+			TabLabel = LOCTEXT("MetaSoundConfigurationView_MenuTitle", "Configuration");
+			TabIdentifier = "Configuration";
+			ViewMenuTooltip = LOCTEXT("MetaSoundConfigurationView_ToolTip", "Custom view for the current MetaSound configuration");
+
+			TSharedPtr<FEditor> Parent = GetParentEditor();
+			if (ensure(Parent.IsValid()))
+			{
+				if (FEditor::FUserWidgetData* WidgetData = Parent->DocConfigWidgetData.Find(BlueprintId))
+				{
+					TabLabel = WidgetData->TabInfo.TabDisplayName;
+					TabIdentifier = WidgetData->TabInfo.TabName;
+					TabIcon = FSlateIcon(WidgetData->TabInfo.StyleName, WidgetData->TabInfo.TabIcon);
+					if (const FDocumentConfiguration* Config = Parent->GetConfiguration())
+					{
+						ViewMenuTooltip = Config->GetEditorOptions().ToolTip;
+					}
+				}
+			}
+
+			ViewMenuDescription = TabLabel;
+		}
+
+		TSharedRef<SWidget> FEditor::FDocConfigWidgetSummoner::CreateTabBody(const FWorkflowTabSpawnInfo& Info) const
+		{
+			if (const FUserWidgetData* WidgetData = GetWidgetData())
+			{
+				if (UUserWidget* UserWidget = WidgetData->Widget.Get())
+				{
+					return UserWidget->TakeWidget();
+				}
+			}
+
+			return SNullWidget::NullWidget;
+		}
+
+		TSharedRef<SDockTab> FEditor::FDocConfigWidgetSummoner::SpawnTab(const FWorkflowTabSpawnInfo& Info) const
+		{
+			TSharedRef<SDockTab> DockTab = SNew(SDockTab);
+
+			DockTab->SetLabel(TabLabel);
+
+			if (const FUserWidgetData* WidgetData = GetWidgetData())
+			{
+				if (UUserWidget* UserWidget = WidgetData->Widget.Get())
+				{
+					DockTab->SetContent(UserWidget->TakeWidget());
+				}
+			}
+
+			return DockTab;
+		}
+
+		TSharedPtr<FEditor> FEditor::FDocConfigWidgetSummoner::GetParentEditor() const
+		{
+			return StaticCastWeakPtr<FEditor>(HostingApp).Pin();
+		}
+
+		const FEditor::FUserWidgetData* FEditor::FDocConfigWidgetSummoner::GetWidgetData() const
+		{
+			TSharedPtr<FEditor> Parent = GetParentEditor();
+			if (ensure(Parent.IsValid()))
+			{
+				return Parent->DocConfigWidgetData.Find(BlueprintId);
+			}
+
+			return nullptr;
+		}
+
+		FText FEditor::FDocConfigWidgetSummoner::GetTabToolTipText(const FWorkflowTabSpawnInfo& Info) const
+		{
+			using namespace Frontend;
+
+			TSharedPtr<FEditor> Parent = GetParentEditor();
+			if (Parent.IsValid())
+			{
+				if (const FDocumentConfiguration* Config = Parent->GetConfiguration())
+				{
+					return Config->GetEditorOptions().ToolTip;
+				}
+			}
+			return { };
+		}
+
 		const FName FEditor::EditorName = "MetaSoundEditor";
 
 		FEditor::FEditor()
@@ -1071,33 +1292,6 @@ namespace Metasound
 			.SetGroup(WorkspaceMenuCategoryRef)
 			.SetIcon(FSlateIcon(FAppStyle::GetAppStyleSetName(), "Kismet.Tabs.Palette"));
 
-			InTabManager->RegisterTabSpawner(TabNamesPrivate::DocConfigWidget, FOnSpawnTab::CreateLambda([this](const FSpawnTabArgs&)
-			{
-				if (DocConfigWidget.IsValid())
-				{
-					// Refresh widget if widget blueprint has been recompiled
-					if (DocConfigWidget->GetClass()->HasAnyClassFlags(EClassFlags::CLASS_NewerVersionExists))
-					{
-						DocConfigWidget = TStrongObjectPtr<UUserWidget>(CreateDocConfigWidget());
-					}
-				}
-
-				if (DocConfigWidget.IsValid())
-				{
-					IMetaSoundPresetWidgetInterface::Execute_OnConstructed(DocConfigWidget.Get(), Builder.Get());
-					return SNew(SDockTab)
-						.Label(LOCTEXT("DocConfigWidgetTitle", "Configuration Widget"))
-						[
-							DocConfigWidget.Get()->TakeWidget()
-						];
-				}
-
-				return SNew(SDockTab);
-			}))
-			.SetDisplayName(LOCTEXT("DocConfigWidgetTab", "Configuration"))
-			.SetGroup(WorkspaceMenuCategoryRef)
-			.SetIcon(FSlateIcon(FAppStyle::GetAppStyleSetName(), "Kismet.Tabs.Palette"));
-
 			InTabManager->RegisterTabSpawner(TabNamesPrivate::Interfaces, FOnSpawnTab::CreateLambda([InInterfacesDetails = InterfacesDetails](const FSpawnTabArgs&)
 			{
 				return SNew(SDockTab).Label(LOCTEXT("MetasoundInterfacesDetailsTitle", "Interfaces")) [ InInterfacesDetails.ToSharedRef() ];
@@ -1128,7 +1322,18 @@ namespace Metasound
 			InTabManager->UnregisterTabSpawner(TabNamesPrivate::Pages);
 			InTabManager->UnregisterTabSpawner(TabNamesPrivate::Interfaces);
 			InTabManager->UnregisterTabSpawner(TabNamesPrivate::Find);
-			InTabManager->UnregisterTabSpawner(TabNamesPrivate::DocConfigWidget);
+
+			for (auto FactoryIt = DocConfigTabSet.CreateConstIterator(); FactoryIt; ++FactoryIt)
+			{
+				const FName TabId = FactoryIt.Value()->GetIdentifier();
+				InTabManager->UnregisterTabSpawner(TabId);
+			}
+			DocConfigTabSet.Clear();
+		}
+
+		bool FEditor::IgnoreDocConfigVisibilitySettings()
+		{
+			return AssetEditorPrivate::bIgnoreDocConfigVisibilitySettings;
 		}
 
 		bool FEditor::IsLiveAuditionEnabled()
@@ -1170,13 +1375,17 @@ namespace Metasound
 			FLevelEditorModule& LevelEditor = FModuleManager::LoadModuleChecked<FLevelEditorModule>("LevelEditor");
 			LevelEditor.OnMapChanged().RemoveAll(this);
 
-			if (DocConfigWidget.IsValid())
+			for (FBPIdWidgetDataPair& Pair : DocConfigWidgetData)
 			{
-				if (UBlueprint* OwningBlueprint = UBlueprint::GetBlueprintFromClass(DocConfigWidget->GetClass()))
+				if (UUserWidget* UserWidget = Pair.Value.Widget.Get())
 				{
-					OwningBlueprint->OnCompiled().RemoveAll(this);
+					if (UBlueprint* OwningBlueprint = UBlueprint::GetBlueprintFromClass(UserWidget->GetClass()))
+					{
+						OwningBlueprint->OnCompiled().RemoveAll(this);
+					}
 				}
 			}
+			DocConfigWidgetData.Reset();
 
 			check(GEditor);
 			GEditor->UnregisterForUndo(this);
@@ -1252,18 +1461,21 @@ namespace Metasound
 			TSharedRef<FTabManager::FStack> DetailsStack = FTabManager::NewStack()
 				->SetSizeCoefficient(0.50f)
 				->SetHideTabWell(false)
-				->AddTab(TabNamesPrivate::Details, ETabState::OpenedTab);
-
-			DetailsStack->AddTab(TabNamesPrivate::Pages, ETabState::OpenedTab);
+				->AddTab(TabNamesPrivate::Details, ETabState::OpenedTab)
+				->AddTab(TabNamesPrivate::Pages, ETabState::OpenedTab);
 
 			GraphCanvasTabStack = FTabManager::NewStack()
 				->SetSizeCoefficient(0.8f)
 				->SetHideTabWell(true)
 				->SetExtensionId(TabNamesPrivate::GraphCanvas)
-				->AddTab(TabNamesPrivate::GraphCanvas, ETabState::OpenedTab)
-				->AddTab(TabNamesPrivate::DocConfigWidget, ETabState::OpenedTab);
+				->AddTab(TabNamesPrivate::GraphCanvas, ETabState::OpenedTab);
+
+			for (const FBPIdWidgetDataPair& Pair : DocConfigWidgetData)
+			{
+				GraphCanvasTabStack->AddTab(Pair.Value.TabInfo.TabName, ETabState::ClosedTab);
+			}
 
-			const TSharedRef<FTabManager::FLayout> StandaloneDefaultLayout = FTabManager::NewLayout("Standalone_MetasoundEditor_Layout_v16")
+			const TSharedRef<FTabManager::FLayout> StandaloneDefaultLayout = FTabManager::NewLayout("Standalone_MetasoundEditor_Layout_v17")
 			->AddArea
 			(
 				FTabManager::NewPrimaryArea()
@@ -1332,6 +1544,8 @@ namespace Metasound
 			ExtendToolbarInternal();
 			RegenerateMenusAndToolbars();
 
+			ReloadAllDocConfigTabs();
+
 			NotifyDocumentVersioned();
 
 			if (bRestartSound)
@@ -1429,7 +1643,7 @@ namespace Metasound
 					const FMetasoundFrontendDocument& Document = DocBuilder.GetConstDocumentChecked();
 					if (const FMetaSoundFrontendDocumentConfiguration* DocConfig = Document.Configuration.GetPtr())
 					{
-						if (!DocConfig->IsPageGraphEditingEnabled())
+						if (!DocConfig->GetEditorOptions().bPageGraphEditingEnabled)
 						{
 							return false;
 						}
@@ -1791,8 +2005,6 @@ namespace Metasound
 			{
 				FLevelEditorModule& LevelEditor = FModuleManager::LoadModuleChecked<FLevelEditorModule>("LevelEditor");
 				LevelEditor.OnMapChanged().AddSP(this, &FEditor::ChangeDocConfigWidgetWorld);
-
-				DocConfigWidget = TStrongObjectPtr<UUserWidget>(CreateDocConfigWidget());
 			}
 
 			Palette = SNew(SMetasoundPalette);
@@ -1833,141 +2045,191 @@ namespace Metasound
 			}
 		}
 
-		UUserWidget* FEditor::CreateDocConfigWidget()
+		void FEditor::ReloadAllDocConfigTabs()
 		{
-			using namespace Metasound::Frontend;
+			using namespace Frontend;
 
-			if (!Builder.IsValid() || Builder->GetConstBuilder().GetIsGraphEditable())
+			if (TabManager.IsValid())
 			{
-				return nullptr;
+				for (auto FactoryIt = DocConfigTabSet.CreateConstIterator(); FactoryIt; ++FactoryIt)
+				{
+					const FName TabId = FactoryIt.Value()->GetIdentifier();
+					if (TSharedPtr<SDockTab> Tab = TabManager->FindExistingLiveTab(TabId))
+					{
+						Tab->RequestCloseTab();
+					}
+					TabManager->UnregisterTabSpawner(TabId);
+				}
 			}
 
-			UWorld* World = GEditor ? GEditor->GetEditorWorldContext().World() : nullptr;
-			UAudioWidgetSubsystem* WidgetSubsystem = GEngine ? GEngine->GetEngineSubsystem<UAudioWidgetSubsystem>() : nullptr;
-			if (!World || !WidgetSubsystem)
-			{
-				return nullptr;
-			}
+			DocConfigTabSet.Clear();
 
-			// This asset 
-			const FMetaSoundFrontendDocumentBuilder& DocBuilder = Builder->GetConstBuilder();
-			if (!DocBuilder.IsValid())
+			for (TPair<FBlueprintId, FUserWidgetData>& WidgetData : DocConfigWidgetData)
 			{
-				return nullptr;
+				if (UUserWidget* Widget = WidgetData.Value.Widget.Get())
+				{
+					if (UBlueprint* OwningBlueprint = UBlueprint::GetBlueprintFromClass(Widget->GetClass()))
+					{
+						OwningBlueprint->OnCompiled().RemoveAll(this);
+					}
+				}
 			}
+			DocConfigWidgetData.Reset();
 
-			const FMetasoundFrontendGraphClass& ThisGraphClass = DocBuilder.GetConstDocumentChecked().RootGraph;
-			const FNodeClassRegistryKey ThisRegistryKey(ThisGraphClass);
-
-			// Referenced assets
-			FMetasoundAssetBase& ThisAsset = DocBuilder.GetMetasoundAsset();
-			TArray<FMetasoundAssetBase*> ReferencedAssets;
-			IMetaSoundAssetManager::GetChecked().GetReferencedPresetHierarchy(ThisAsset, ReferencedAssets);
+			RebuildDocConfigUserWidgets();
 
-			TSet<FNodeClassRegistryKey> ReferencedRegistryKeys;
-			Algo::Transform(ReferencedAssets, ReferencedRegistryKeys, [](const FMetasoundAssetBase* ReferencedAsset)
+			TMap<FName, FSlateIcon> TabIcons;
+			for (const FBPIdWidgetDataPair& Pair : DocConfigWidgetData)
 			{
-				if (ensure(ReferencedAsset))
+				const FName TabName = Pair.Value.TabInfo.TabName;
+				FNameBuilder FactoryName(TabNamesPrivate::DocConfigWidgets_Factory);
+				if (!TabName.IsNone())
 				{
-					TScriptInterface<const IMetaSoundDocumentInterface> RefDocInterface = ReferencedAsset->GetOwningAsset();
-					const FMetasoundFrontendDocument& ReferencedDocument = RefDocInterface->GetConstDocument();
-					const FMetasoundFrontendGraphClass& ReferencedGraphClass = ReferencedDocument.RootGraph;
-					return FNodeClassRegistryKey(ReferencedGraphClass);
+					FactoryName.Appendf(TEXT("_%s"), *TabName.ToString());
 				}
-				return FNodeClassRegistryKey();
-			});
+				DocConfigTabSet.RegisterFactory(MakeShared<FDocConfigWidgetSummoner>(*FactoryName, Pair.Key, AsShared()));
+				TabIcons.Add(TabName, FSlateIcon(Pair.Value.TabInfo.StyleName, Pair.Value.TabInfo.TabIcon));
+			}
 
-			// Filter whether the widget supports this MetaSound
-			auto FilterFunction = [&](UUserWidget* UserWidget)
+			if (TabManager.IsValid())
 			{
-				bool bSupportAllPresets = false;
-				TArray<TScriptInterface<IMetaSoundDocumentInterface>> ExcludedMetaSounds;
-				TArray<TScriptInterface<IMetaSoundDocumentInterface>> IncludedMetaSounds;
-				IMetaSoundPresetWidgetInterface::Execute_GetSupportedMetaSounds(UserWidget, bSupportAllPresets, ExcludedMetaSounds, IncludedMetaSounds);
-
-				// Show if all presets supported and this MetaSound is not excluded 
-				if (bSupportAllPresets)
+				for (auto FactoryIt = DocConfigTabSet.CreateConstIterator(); FactoryIt; ++FactoryIt)
 				{
-					if (!ExcludedMetaSounds.IsEmpty())
-					{
-						for (TScriptInterface<IMetaSoundDocumentInterface> ExcludedMetaSound : ExcludedMetaSounds)
-						{
-							if (!ExcludedMetaSound)
-							{
-								continue;
-							}
+					FTabSpawnerEntry& SpawnerEntry = FactoryIt.Value()->RegisterTabSpawner(TabManager.ToSharedRef(), nullptr)
+					.SetGroup(WorkspaceMenuCategory.ToSharedRef());
 
-							// This MetaSound or the referenced MetaSound is excluded 
-							const FMetasoundFrontendGraphClass& ExcludedGraphClass = ExcludedMetaSound->GetConstDocument().RootGraph;
-							const FNodeRegistryKey ExcludedRegistryKey(ExcludedGraphClass);
-
-							if (ThisRegistryKey == ExcludedRegistryKey || ReferencedRegistryKeys.Contains(ExcludedRegistryKey))
-							{
-								return false;
-							}
-						}
+					const FName TabId = FactoryIt.Value()->GetIdentifier();
+					if (const FSlateIcon* TabIcon = TabIcons.Find(TabId))
+					{
+						SpawnerEntry.SetIcon(*TabIcon);
 					}
-					return true;
+
 				}
+			}
+		}
+
+		void FEditor::RebuildDocConfigUserWidgets(const FGuid& TargetBlueprintId)
+		{
+			using namespace AssetEditorPrivate;
+			using namespace Frontend;
+
+			UMetaSoundBuilderBase* BuilderBase = Builder.Get();
+			if (!BuilderBase)
+			{
+				return;
+			}
+
+			const FMetaSoundFrontendDocumentBuilder& DocBuilder = BuilderBase->GetConstBuilder();
+			if (!DocBuilder.IsValid())
+			{
+				return;
+			}
+
+			TConstStructView<FDocumentConfiguration> DocConfig = DocBuilder.GetConstDocumentConfiguration();
+			if (!DocConfig.IsValid() || !DocConfig.GetPtr()->GetEditorOptions().bCustomTabEnabled)
+			{
+				return;
+			}
+
+			if (!GEditor || !GEngine)
+			{
+				return;
+			}
+
+			UWorld* World = GEditor->GetEditorWorldContext().World();
+			UAudioWidgetSubsystem* WidgetSubsystem = GEngine->GetEngineSubsystem<UAudioWidgetSubsystem>();
+			if (!World || !WidgetSubsystem)
+			{
+				return;
+			}
 
-				// This MetaSound is in the inclusion list 
-				for (TScriptInterface<IMetaSoundDocumentInterface> IncludedMetaSound : IncludedMetaSounds)
+			// Preset separated out for back compat
+			if (TFunction<bool(UUserWidget*)> PresetFilterFunction = GetPresetWidgetFilterFunction(DocBuilder); PresetFilterFunction.IsSet())
+			{
+				TFunction<FMetaSoundDocumentConfigurationWidgetTabInfo(UUserWidget*)> OnPresetConstruct = [&DocConfig, &BuilderBase, &TargetBlueprintId](UUserWidget* UserWidget)
 				{
-					if (!IncludedMetaSound)
-					{
-						continue;
-					}
+					IMetaSoundPresetWidgetInterface::Execute_OnConstructed(UserWidget, BuilderBase);
 
-					// This MetaSound or the referenced MetaSound is included 
-					const FMetasoundFrontendGraphClass& IncludedGraphClass = IncludedMetaSound->GetConstDocument().RootGraph;
-					const FNodeRegistryKey IncludedRegistryKey(IncludedGraphClass);
+					FMetaSoundDocumentConfigurationWidgetTabInfo TabInfo;
+					IMetaSoundPresetWidgetInterface::Execute_OnInitializingTab(UserWidget, TabInfo);
+					return TabInfo;
+				};
+				RebuildUserWidgets(*World, *WidgetSubsystem,
+					UMetaSoundPresetWidgetInterface::StaticClass(),
+					MoveTemp(PresetFilterFunction),
+					MoveTemp(OnPresetConstruct),
+					TargetBlueprintId);
+				return;
+			}
 
-					if (ThisRegistryKey == IncludedRegistryKey || ReferencedRegistryKeys.Contains(IncludedRegistryKey))
-					{
-						return true;
-					}
+			TFunction<bool(UUserWidget*)> DocConfigFilterFunction = [&DocConfig, BuilderBase](UUserWidget* UserWidget)
+			{
+				bool bIsSupportedConfig = true;
+				if (const UScriptStruct* ScriptStruct = DocConfig.GetScriptStruct())
+				{
+					IMetaSoundDocumentConfigurationWidgetInterface::Execute_IsSupportedConfiguration(
+						UserWidget,
+						ScriptStruct->GetFName(),
+						BuilderBase,
+						bIsSupportedConfig);
 				}
-				return false;
+				return bIsSupportedConfig;
 			};
 
-			// Create widget
-			TArray<UUserWidget*> UserWidgets = WidgetSubsystem->CreateUserWidgets(*World, UMetaSoundPresetWidgetInterface::StaticClass(), FilterFunction);
-			// Currently only support a single user widget per preset 
-			if (UserWidgets.Num() > 0)
+			TFunction<FMetaSoundDocumentConfigurationWidgetTabInfo(UUserWidget*)> OnDocConfigTabConstruct = [&DocConfig, BuilderBase](UUserWidget* UserWidget)
 			{
-				UUserWidget* UserWidget = UserWidgets.Last();
-				check(UserWidget);
-
-				// Set recompile refresh delegate 
-				if (UBlueprint* OwningBlueprint = UBlueprint::GetBlueprintFromClass(UserWidget->GetClass()))
+				TScriptInterface<IMetaSoundDocumentConfigurationWidgetInterface> WidgetInterface(UserWidget);
+				if (ensure(WidgetInterface.GetObject()))
 				{
-					OwningBlueprint->OnCompiled().RemoveAll(this);
-					OwningBlueprint->OnCompiled().AddSP(this, &FEditor::OnPresetWidgetBlueprintCompiled);
+					IMetaSoundDocumentConfigurationWidgetInterface::Execute_OnBuilderSet(UserWidget, BuilderBase);
 				}
-				return UserWidget;
-			}
-			return nullptr;
+
+				TSubclassOf<UMetaSoundDocumentConfigurationViewModelBase> VMClass;
+				IMetaSoundDocumentConfigurationWidgetInterface::Execute_OnInitializingViewModel(UserWidget, VMClass);
+				UMetaSoundDocumentConfigurationViewModelBase::ConstructWidgetViewModel(UserWidget, *BuilderBase, DocConfig, VMClass);
+
+				FMetaSoundDocumentConfigurationWidgetTabInfo TabInfo;
+				IMetaSoundDocumentConfigurationWidgetInterface::Execute_OnInitializingTab(UserWidget, TabInfo);
+				return TabInfo;
+			};
+
+			RebuildUserWidgets(*World, *WidgetSubsystem,
+				UMetaSoundDocumentConfigurationWidgetInterface::StaticClass(),
+				MoveTemp(DocConfigFilterFunction),
+				MoveTemp(OnDocConfigTabConstruct),
+				TargetBlueprintId);
 		}
 
-		void FEditor::OnPresetWidgetBlueprintCompiled(UBlueprint* Blueprint)
+		void FEditor::OnDocConfigWidgetBlueprintCompiled(UBlueprint* Blueprint)
 		{
-			if (!Blueprint || !DocConfigWidget)
-			{
-				return;
-			}
-			if (UBlueprint* PresetWidgetBlueprint = UBlueprint::GetBlueprintFromClass(DocConfigWidget->GetClass()))
+			using namespace AssetEditorPrivate;
+
+			if (Blueprint)
 			{
-				if (PresetWidgetBlueprint == Blueprint)
+				const FGuid BlueprintId = Blueprint->GetBlueprintGuid();
+				if (const FUserWidgetData* WidgetData = DocConfigWidgetData.Find(BlueprintId))
+				{
+					if (TabManager.IsValid())
+					{
+						if (TSharedPtr<SDockTab> DockTab = TabManager->FindExistingLiveTab(WidgetData->TabInfo.TabName))
+						{
+							DockTab->SetContent(SNullWidget::NullWidget);
+						}
+					}
+				}
+
+				RebuildDocConfigUserWidgets(BlueprintId);
+
+				if (const FUserWidgetData* WidgetData = DocConfigWidgetData.Find(BlueprintId))
 				{
-					// Regenerate user preset widget
-					TSharedPtr<SDockTab> DocConfigWidgetTab = TabManager->FindExistingLiveTab(TabNamesPrivate::DocConfigWidget);
-					if (DocConfigWidgetTab.IsValid())
+					if (TabManager.IsValid())
 					{
-						DocConfigWidget = TStrongObjectPtr<UUserWidget>(CreateDocConfigWidget());
-						if (DocConfigWidget.IsValid())
+						if (TSharedPtr<SDockTab> DockTab = TabManager->FindExistingLiveTab(WidgetData->TabInfo.TabName))
 						{
-							IMetaSoundPresetWidgetInterface::Execute_OnConstructed(DocConfigWidget.Get(), Builder.Get());
-							DocConfigWidgetTab->SetContent(DocConfigWidget.Get()->TakeWidget());
+							if (UUserWidget* UserWidget = WidgetData->Widget.Get())
+							{
+								DockTab->SetContent(UserWidget->TakeWidget());
+							}
 						}
 					}
 				}
@@ -1976,20 +2238,33 @@ namespace Metasound
 		
 		void FEditor::ChangeDocConfigWidgetWorld(UWorld* World, EMapChangeType MapChangeType)
 		{
-			if (MapChangeType == EMapChangeType::TearDownWorld)
+			switch (MapChangeType)
 			{
-				// We need to Delete the UMG widget if we are tearing down the World it was built with.
-				if (DocConfigWidget && World == DocConfigWidget->GetWorld())
+				case EMapChangeType::TearDownWorld:
+				{
+					// Moves the UMG widgets to the transient package if the
+					// World they were built with is being torn down.
+					for (FBPIdWidgetDataPair& Pair : DocConfigWidgetData)
+					{
+						UUserWidget* DocConfigWidget = Pair.Value.Widget.Get();
+						if (DocConfigWidget && World == DocConfigWidget->GetWorld())
+						{
+							DocConfigWidget->Rename(nullptr, GetTransientPackage());
+						}
+					}
+				}
+				break;
+
+				case EMapChangeType::LoadMap:
+				case EMapChangeType::NewMap:
 				{
-					DocConfigWidget->Rename(nullptr, GetTransientPackage());
-					DocConfigWidget = nullptr;
+					ReloadAllDocConfigTabs();
 				}
-			}
-			else if (MapChangeType != EMapChangeType::SaveMap)
-			{
-				// Recreate the widget if loading a map or opening a new map
-				// since the widget is parented to the world
-				DocConfigWidget = TStrongObjectPtr<UUserWidget>(CreateDocConfigWidget());
+				break;
+
+				case EMapChangeType::SaveMap:
+				default:
+				break;
 			}
 		}
 
@@ -2286,6 +2561,8 @@ namespace Metasound
 							constexpr bool bLogErrors = true;
 							DocBuilder->SetDocumentConfiguration(MoveTemp(ConfigInstance));
 							DocBuilder->ConfigureDocument(bLogErrors);
+
+							ReloadAllDocConfigTabs();
 							UpdateEnabledTabs();
 						}
 					});
@@ -2720,6 +2997,21 @@ namespace Metasound
 			}
 		}
 
+		void FEditor::FUserWidgetData::SetAuditionState(UAudioComponent* AudioComponent, bool bIsAuditioning) const
+		{
+			if (UUserWidget* UserWidget = Widget.Get())
+			{
+				if (UMetaSoundPresetWidgetInterface::StaticClass() == Class)
+				{
+					IMetaSoundPresetWidgetInterface::Execute_OnAuditionStateChanged(UserWidget, AudioComponent, bIsAuditioning);
+				}
+				else if (ensure(UMetaSoundDocumentConfigurationWidgetInterface::StaticClass() == Class))
+				{
+					IMetaSoundDocumentConfigurationWidgetInterface::Execute_OnAuditionStateChanged(UserWidget, AudioComponent, bIsAuditioning);
+				}
+			}
+		}
+
 		void FEditor::Play()
 		{
 			using namespace Engine;
@@ -2752,10 +3044,13 @@ namespace Metasound
 							CastChecked<UMetaSoundSourceBuilder>(Builder.Get())->Audition(nullptr, PreviewComp, GeneratorHandleDelegate, AssetEditorPrivate::bEnableLiveAudition);
 						}
 
-						if (DocConfigWidget.IsValid() && IsPlaying())
+						if (IsPlaying())
 						{
 							constexpr bool bIsAuditioning = true;
-							IMetaSoundPresetWidgetInterface::Execute_OnAuditionStateChanged(DocConfigWidget.Get(), PreviewComp, bIsAuditioning);
+							for (FBPIdWidgetDataPair& Pair : DocConfigWidgetData)
+							{
+								Pair.Value.SetAuditionState(PreviewComp, bIsAuditioning);
+							}
 						}
 
 						SetPreviewID(PreviewComp->GetUniqueID());
@@ -2784,10 +3079,10 @@ namespace Metasound
 							UpdatePageInfo(bIsPlaying);
 							GraphConnectionManager = RebuildConnectionManager();
 
-							if (DocConfigWidget.IsValid())
+							constexpr bool bIsAuditioning = false;
+							for (FBPIdWidgetDataPair& Pair : DocConfigWidgetData)
 							{
-								constexpr bool bIsAuditioning = false;
-								IMetaSoundPresetWidgetInterface::Execute_OnAuditionStateChanged(DocConfigWidget.Get(), nullptr, bIsAuditioning);
+								Pair.Value.SetAuditionState(nullptr, bIsAuditioning);
 							}
 
 							return EActiveTimerReturnType::Stop;
@@ -2912,14 +3207,13 @@ namespace Metasound
 				if (const FMetaSoundFrontendDocumentBuilder& DocBuilder = Builder->GetBuilder(); ensure(DocBuilder.IsValid()))
 				{
 					TConstStructView<FMetaSoundFrontendDocumentConfiguration> Config = DocBuilder.GetConstDocumentConfiguration();
-					return Config.IsValid() && !Config.GetPtr()->IsGraphEditorVisible();
+					return Config.IsValid() && !Config.GetPtr()->GetEditorOptions().bGraphEditorVisible;
 				}
 			}
 
 			return false;
 		}
 
-
 		void FEditor::ToggleShowHideGraph()
 		{
 			AssetEditorPrivate::bShowGraphOverride = !AssetEditorPrivate::bShowGraphOverride;
@@ -3022,9 +3316,10 @@ namespace Metasound
 
 									GetMetasoundObject()->Modify();
 									EMetaSoundBuilderResult Result = EMetaSoundBuilderResult::Failed;
-										DocBuilder.SetDocumentConfiguration({});
-										DocBuilder.ConfigureDocument();
-										UpdateEnabledTabs();
+									DocBuilder.SetDocumentConfiguration({ });
+									DocBuilder.ConfigureDocument();
+									ReloadAllDocConfigTabs();
+									UpdateEnabledTabs();
 								}
 								
 								if (UToolMenus* ToolMenus = UToolMenus::Get())
@@ -4961,7 +5256,8 @@ PRAGMA_ENABLE_DEPRECATION_WARNINGS
 
 				if (!MembersModified.IsEmpty() || bForceRefreshViews)
 				{
-					if (const Frontend::FDocumentConfiguration* Config = GetConfiguration(); !Config || Config->IsMemberEditingEnabled())
+					if (const Frontend::FDocumentConfiguration* Config = GetConfiguration();
+						!Config || Config->GetEditorOptions().bMemberEditingEnabled)
 					{
 						UMetasoundEditorGraphMember* SelectedMember = RefreshGraphMemberMenu();
 
@@ -5328,6 +5624,8 @@ PRAGMA_ENABLE_DEPRECATION_WARNINGS
 
 		void FEditor::UpdateEnabledTabs()
 		{
+			using namespace Frontend;
+
 			if (!TabManager)
 			{
 				return;
@@ -5366,28 +5664,36 @@ PRAGMA_ENABLE_DEPRECATION_WARNINGS
 			if (FMetaSoundFrontendDocumentBuilder* DocBuilder = GetFrontendBuilder())
 			{
 				const bool bHasProjectPageValues = DocBuilder->GetConstDocumentChecked().RootGraph.GetConstGraphPages().Num() > 1;
+
 				const bool bEnablePageTab = IsPageGraphEditingEnabled(*DocBuilder, bHasProjectPageValues);
 
-				bool bVisibilitySet = false;
-				if (const FMetaSoundFrontendDocumentConfiguration* Config = GetConfiguration())
+				bool bConfigVisibilitySet = false;
+				if (const FDocumentConfiguration* Config = GetConfiguration())
 				{
 					if (!AssetEditorPrivate::bIgnoreDocConfigVisibilitySettings)
 					{
-						const bool bShowGraph = Config->IsGraphEditorVisible() || AssetEditorPrivate::bShowGraphOverride;
+						const FDocumentConfiguration::FEditorOptions& EdOptions = Config->GetEditorOptions();
+						const bool bShowGraph = EdOptions.bGraphEditorVisible || AssetEditorPrivate::bShowGraphOverride;
 						ResolveTabEnablement(TabNamesPrivate::Find, bShowGraph);
 						ResolveTabEnablement(TabNamesPrivate::GraphCanvas, bShowGraph);
-						ResolveTabEnablement(TabNamesPrivate::Interfaces, Config->IsInterfaceEditingEnabled());
-						ResolveTabEnablement(TabNamesPrivate::Members, Config->IsMemberEditingEnabled());
+						ResolveTabEnablement(TabNamesPrivate::Interfaces, EdOptions.bInterfaceEditingEnabled);
+						ResolveTabEnablement(TabNamesPrivate::Members, EdOptions.bMemberEditingEnabled);
+
+						// 'IsPageGraphEditingEnabled' above internally checks if config supports page graph editing,
+						// so need to resolve enablement for the pages tab here.
 
 						// Don't show custom config tab if no config widget was constructed, even if enabled
-						ResolveTabEnablement(TabNamesPrivate::DocConfigWidget, DocConfigWidget.IsValid() && Config->IsCustomTabEnabled());
+						for (const FBPIdWidgetDataPair& Pair : DocConfigWidgetData)
+						{
+							ResolveTabEnablement(Pair.Value.TabInfo.TabName, EdOptions.bCustomTabEnabled);
+						}
 
 						RefreshViews();
-						bVisibilitySet = true;
+						bConfigVisibilitySet = true;
 					}
 				}
 
-				if (!bVisibilitySet)
+				if (!bConfigVisibilitySet)
 				{
 					{
 						constexpr bool bEnableTab = true;
@@ -5398,9 +5704,12 @@ PRAGMA_ENABLE_DEPRECATION_WARNINGS
 					}
 
 					{
-						// User preset widgets are only visible when doc config is set
-						constexpr bool bEnableTab = false;
-						ResolveTabEnablement(TabNamesPrivate::DocConfigWidget, bEnableTab);
+						// Doc config widgets are only visible when visibility is set unless ignore cvar is true.
+						const bool bEnableTab = !AssetEditorPrivate::bIgnoreDocConfigVisibilitySettings;
+						for (const FBPIdWidgetDataPair& Pair : DocConfigWidgetData)
+						{
+							ResolveTabEnablement(Pair.Value.TabInfo.TabName, bEnableTab);
+						}
 					}
 				}
 
@@ -5524,11 +5833,15 @@ PRAGMA_ENABLE_DEPRECATION_WARNINGS
 			OutDelegates.OnDocumentConfigurationChanged.AddSP(this, &FDocumentListener::OnDocumentConfigurationChanged);
 		}
 
-		void FEditor::FDocumentListener::OnDocumentConfigurationChanged(const Frontend::FDocumentConfigChangedArgs&)
+		void FEditor::FDocumentListener::OnDocumentConfigurationChanged(const Frontend::FDocumentConfigChangedArgs& Args)
 		{
 			if (TSharedPtr<FEditor> ParentPtr = Parent.Pin())
 			{
-				ParentPtr->UpdateEnabledTabs();
+				if (Args.Type != Frontend::EDocumentConfigChangeType::Property)
+				{
+					ParentPtr->ReloadAllDocConfigTabs();
+					ParentPtr->UpdateEnabledTabs();
+				}
 			}
 		}
 
@@ -5548,6 +5861,74 @@ PRAGMA_ENABLE_DEPRECATION_WARNINGS
 			}
 		}
 
+		void FEditor::RebuildUserWidgets(
+			UWorld& World,
+			UAudioWidgetSubsystem& WidgetSubsystem,
+			TSubclassOf<UInterface> InterfaceClass,
+			TFunction<bool(UUserWidget*)>Filter,
+			TFunction<FMetaSoundDocumentConfigurationWidgetTabInfo(UUserWidget*)>OnTabConstruct,
+			FGuid TargetBlueprintId)
+		{
+			TFunction<bool(UBlueprint*)> BlueprintFilter;
+			if (TargetBlueprintId.IsValid())
+			{
+				BlueprintFilter = [&TargetBlueprintId](UBlueprint* OtherBlueprint)
+				{
+					return OtherBlueprint && OtherBlueprint->GetBlueprintGuid() == TargetBlueprintId;
+				};
+			}
+
+			TArray<UUserWidget*> NewWidgets = WidgetSubsystem.CreateUserWidgets(World, InterfaceClass, MoveTemp(Filter), BlueprintFilter);
+			for (UUserWidget* UserWidget : NewWidgets)
+			{
+				check(UserWidget);
+				if (UBlueprint* OwningBlueprint = UBlueprint::GetBlueprintFromClass(UserWidget->GetClass()))
+				{
+					const FBlueprintId& OwningBPId = OwningBlueprint->GetBlueprintGuid();
+
+					FDelegateHandle RecompileHandle;
+
+					// If request to rebuild provides a target blueprint ID, reuse the
+					// existing handle if found. Otherwise, assign a new delegate handler.
+					if (TargetBlueprintId.IsValid())
+					{
+						if (FUserWidgetData* WidgetData = DocConfigWidgetData.Find(TargetBlueprintId))
+						{
+							RecompileHandle = MoveTemp(WidgetData->RecompileHandle);
+						}
+					}
+					else
+					{
+						OwningBlueprint->OnCompiled().RemoveAll(this);
+					}
+
+					DocConfigWidgetData.Remove(OwningBPId);
+
+					if (!RecompileHandle.IsValid())
+					{
+						RecompileHandle = OwningBlueprint->OnCompiled().AddSP(this, &FEditor::OnDocConfigWidgetBlueprintCompiled);
+					}
+
+					TStrongObjectPtr<UUserWidget> UserWidgetPtr(UserWidget);
+					FUserWidgetData& NewData = DocConfigWidgetData.Add(OwningBPId, FUserWidgetData
+					{
+						.Class = InterfaceClass,
+						.RecompileHandle = MoveTemp(RecompileHandle),
+						.Widget = MoveTemp(UserWidgetPtr),
+						.TabInfo = OnTabConstruct(UserWidget)
+					});
+
+					FMetaSoundFrontendDocumentBuilder* DocBuilder = GetFrontendBuilder();
+					checkf(DocBuilder, TEXT("Builder must exist by this point in user widget rebuild"));
+
+					const UScriptStruct* ConfigStruct = DocBuilder->GetConstDocumentConfiguration().GetScriptStruct();
+					checkf(ConfigStruct, TEXT("Builder must have a config set by this point in user widget rebuild"));
+
+					AssetEditorPrivate::SanitizeTabInfo(*ConfigStruct, NewData.TabInfo);
+				}
+			}
+		}
+
 		void FEditor::RefreshExecVisibility(const FGuid& InPageID) const
 		{
 			if (PageStatsWidget.IsValid())
diff --git a/Engine/Plugins/Runtime/Metasound/Source/MetasoundEditor/Private/MetasoundEditor.h b/Engine/Plugins/Runtime/Metasound/Source/MetasoundEditor/Private/MetasoundEditor.h
--- a/Engine/Plugins/Runtime/Metasound/Source/MetasoundEditor/Private/MetasoundEditor.h
+++ b/Engine/Plugins/Runtime/Metasound/Source/MetasoundEditor/Private/MetasoundEditor.h
@@ -9,6 +9,7 @@
 #include "IMetasoundEditor.h"
 #include "Logging/TokenizedMessage.h"
 #include "Math/UnrealMathUtility.h"
+#include "MetasoundDocumentConfigurationWidgetInterface.h"
 #include "MetasoundEditorGraphConnectionManager.h"
 #include "MetasoundEditorGraphNode.h"
 #include "MetasoundFrontend.h"
@@ -34,6 +35,7 @@
 
 
 // Forward Declarations
+class UAudioWidgetSubsystem;
 class FSlateRect;
 class FTabManager;
 class IDetailsView;
@@ -49,6 +51,7 @@ class UMetasoundEditorGraphNode;
 class UMetaSoundPatch;
 class UMetaSoundSource;
 class UUserWidget;
+class UWorld;
 
 struct FGraphActionNode;
 struct FMetaSoundFrontendDocumentBuilder;
@@ -173,6 +176,9 @@ namespace Metasound::Editor
 		virtual void Play() override;
 		virtual void Stop() override;
 
+		/** Returns whether or not doc config visibility settings should be ignored (for debugging/development) */
+		static bool IgnoreDocConfigVisibilitySettings();
+
 		/** Returns whether or not live audition is enabled for previewing */
 		static bool IsLiveAuditionEnabled();
 
@@ -261,6 +267,9 @@ namespace Metasound::Editor
 
 		void ClearDroppedObjectsInfo();
 
+		/** Reloads all document configuration tabs */
+		void ReloadAllDocConfigTabs();
+
 	protected:
 		// Callbacks for action tree
 		bool CanRenameOnActionNode(TWeakPtr<FGraphActionNode> InSelectedNode) const;
@@ -438,6 +447,15 @@ namespace Metasound::Editor
 	private:
 		int32 PromotableSelectedNodes();
 
+		// If target BP Id set, only rebuilds the widget associated with that BP.
+		void RebuildUserWidgets(
+			UWorld& World,
+			UAudioWidgetSubsystem& WidgetSubsystem,
+			TSubclassOf<UInterface> InterfaceClass,
+			TFunction<bool(UUserWidget*)>Filter,
+			TFunction<FMetaSoundDocumentConfigurationWidgetTabInfo(UUserWidget*)>OnTabConstruct,
+			FGuid TargetBlueprintId = FGuid());
+
 		void RefreshExecVisibility(const FGuid& InPageID) const;
 
 		/** Forces refresh of pages view. */
@@ -462,14 +480,14 @@ namespace Metasound::Editor
 		/** Creates all internal widgets for the tabs to point at */
 		void CreateInternalWidgets(UObject& MetaSound);
 
-		UUserWidget* CreateDocConfigWidget();
+		void RebuildDocConfigUserWidgets(const FGuid& TargetBlueprintId = FGuid());
 
 		/** For teardown and regeneration of user document configuration widget on world change
 			Based on UEditorUtilityWidgetBlueprint::ChangeTabWorld */
 		void ChangeDocConfigWidgetWorld(UWorld* World, EMapChangeType MapChangeType);
 
-		/** Reload preset widget tab on associated widget blueprint recompile */
-		void OnPresetWidgetBlueprintCompiled(UBlueprint* Blueprint);
+		/** Reload document configuration widget tab on associated widget blueprint recompile */
+		void OnDocConfigWidgetBlueprintCompiled(UBlueprint* Blueprint);
 
 		/** Builds the toolbar widget for the Metasound editor */
 		void ExtendToolbarInternal();
@@ -558,9 +576,6 @@ namespace Metasound::Editor
 		/** Updates the page info widget. */
 		void UpdatePageInfo(bool bIsPlaying);
 
-		/** Updates the play time widget.  Returns true if currently playing, false if not. */
-		bool UpdatePlayTime(float InDeltaTime);
-
 		/** Updates the render info widget. */
 		void UpdateRenderInfo(bool bIsPlaying, float InDeltaTime = 0.0f);
 
@@ -599,8 +614,37 @@ namespace Metasound::Editor
 		/** Widget showing render performance information that overlays the graph */
 		TSharedPtr<SRenderStats> RenderStatsWidget;
 
-		/** User defined widget shown when editing a doc config with user widget enabled.*/
-		TStrongObjectPtr<UUserWidget> DocConfigWidget;
+		struct FUserWidgetData
+		{
+			TSubclassOf<UInterface> Class;
+			FDelegateHandle RecompileHandle;
+			TStrongObjectPtr<UUserWidget> Widget;
+			FMetaSoundDocumentConfigurationWidgetTabInfo TabInfo;
+
+			void SetAuditionState(UAudioComponent* AudioComponent, bool bIsAuditioning) const;
+		};
+		using FBlueprintId = FGuid;
+		using FBPIdWidgetDataPair = TPair<FBlueprintId, FUserWidgetData>;
+
+		/** User defined widget data displayed when editing a doc config with user widgets enabled.*/
+		TMap<FBlueprintId, FUserWidgetData> DocConfigWidgetData;
+
+		FWorkflowAllowedTabSet DocConfigTabSet;
+
+		struct FDocConfigWidgetSummoner : public FWorkflowTabFactory
+		{
+		public:
+			FDocConfigWidgetSummoner(FName FactoryName, FBlueprintId InBlueprintId, TSharedPtr<class FAssetEditorToolkit> InHostingApp);
+
+			virtual TSharedRef<SWidget> CreateTabBody(const FWorkflowTabSpawnInfo& Info) const override;
+			virtual TSharedRef<SDockTab> SpawnTab(const FWorkflowTabSpawnInfo& Info) const override;
+			virtual FText GetTabToolTipText(const FWorkflowTabSpawnInfo& Info) const override;
+
+		private:
+			TSharedPtr<FEditor> GetParentEditor() const;
+			const FUserWidgetData* GetWidgetData() const;
+			FBlueprintId BlueprintId;
+		};
 
 		TUniquePtr<FGraphConnectionManager> GraphConnectionManager;
 
diff --git a/Engine/Plugins/Runtime/Metasound/Source/MetasoundEditor/Public/DocumentConfigurationViewModels/MetasoundDocumentConfigurationViewModelBase.h b/Engine/Plugins/Runtime/Metasound/Source/MetasoundEditor/Public/DocumentConfigurationViewModels/MetasoundDocumentConfigurationViewModelBase.h
new file mode 100644
--- /dev/null
+++ b/Engine/Plugins/Runtime/Metasound/Source/MetasoundEditor/Public/DocumentConfigurationViewModels/MetasoundDocumentConfigurationViewModelBase.h
@@ -0,0 +1,80 @@
+// Copyright Epic Games, Inc. All Rights Reserved.
+
+#pragma once
+
+#include "Blueprint/UserWidget.h"
+#include "Delegates/IDelegateInstance.h"
+#include "DocumentConfigurations/MetasoundFrontendDocumentConfiguration.h"
+#include "MetasoundBuilderBase.h"
+#include "MetasoundFrontendDocumentModifyDelegates.h"
+#include "MVVMViewModelBase.h"
+#include "StructUtils/StructView.h"
+#include "UObject/Class.h"
+
+#include "MetasoundDocumentConfigurationViewModelBase.generated.h"
+
+// Forward Declarations
+class UUserWidget;
+
+#define UE_API METASOUNDEDITOR_API
+
+
+UCLASS(MinimalAPI, Blueprintable, Abstract, DisplayName = "MetaSound Configuration ViewModel Base")
+class UMetaSoundDocumentConfigurationViewModelBase : public UMVVMViewModelBase
+{
+	GENERATED_BODY()
+
+public:
+	// Constructs a view model
+	static void ConstructWidgetViewModel(
+		UUserWidget* UserWidget,
+		UMetaSoundBuilderBase& InBuilder,
+		TConstStructView<FMetaSoundFrontendDocumentConfiguration> Config,
+		TSubclassOf<UMetaSoundDocumentConfigurationViewModelBase> VMClass);
+
+	UE_API virtual bool IsSupportedConfiguration(const UScriptStruct& InStruct) const;
+
+protected:
+	template <typename ConfigType>
+	const ConfigType* GetConstConfiguration() const
+	{
+		if (Builder)
+		{
+			return Builder->GetConstBuilder().GetConstDocumentConfigurationAs<ConfigType>();
+		}
+		return nullptr;
+	}
+
+	// Enables responding to the builder reloading or the model view requiring a full
+	// resynchronization of configuration properties to view (i.e. when user transactions
+	// are applied that require a full resync of the view such as undo, redo, etc.)
+	UE_API virtual void Reload();
+
+	UPROPERTY(Transient)
+	TObjectPtr<UMetaSoundBuilderBase> Builder;
+
+private:
+	void Initialize(UMetaSoundBuilderBase* InBuilder);
+
+	class FBuilderListener : public Metasound::Frontend::IDocumentBuilderTransactionListener
+	{
+		TWeakObjectPtr<UMetaSoundDocumentConfigurationViewModelBase> Parent;
+
+	public:
+		FBuilderListener() = default;
+		FBuilderListener(UMetaSoundDocumentConfigurationViewModelBase* InParent)
+			: Parent(InParent)
+		{
+		}
+
+		virtual ~FBuilderListener() = default;
+
+	private:
+		virtual void OnBuilderReloaded(Metasound::Frontend::FDocumentModifyDelegates& OutDelegates) override;
+		void OnConfigChanged(const Metasound::Frontend::FDocumentConfigChangedArgs& Args);
+
+		FDelegateHandle ConfigChangeDelegateHandle;
+	};
+	TSharedPtr<FBuilderListener> BuilderListener;
+};
+#undef UE_API
\ No newline at end of file
diff --git a/Engine/Plugins/Runtime/Metasound/Source/MetasoundEditor/Public/DocumentConfigurationViewModels/RandomizerConfigurationViewModel.h b/Engine/Plugins/Runtime/Metasound/Source/MetasoundEditor/Public/DocumentConfigurationViewModels/RandomizerConfigurationViewModel.h
new file mode 100644
--- /dev/null
+++ b/Engine/Plugins/Runtime/Metasound/Source/MetasoundEditor/Public/DocumentConfigurationViewModels/RandomizerConfigurationViewModel.h
@@ -0,0 +1,71 @@
+// Copyright Epic Games, Inc. All Rights Reserved.
+
+#pragma once
+
+#include "MetasoundDocumentConfigurationViewModelBase.h"
+#include "MetasoundFrontendDocumentModifyDelegates.h"
+#include "Sound/SoundWave.h"
+
+#include "RandomizerConfigurationViewModel.generated.h"
+
+#define UE_API METASOUNDEDITOR_API
+
+
+// Forward Declarations
+struct FMetaSoundRandomizerDocumentConfiguration;
+
+
+UCLASS(MinimalAPI, DisplayName = "MetaSound Randomizer ViewModel")
+class UMetaSoundRandomizerViewModel : public UMetaSoundDocumentConfigurationViewModelBase
+{
+	GENERATED_BODY()
+
+public:
+	UE_API bool GetIsOneShot() const;
+	UE_API FVector2f GetPitch() const;
+
+	UE_API const TArray<TObjectPtr<USoundWave>>& GetSounds() const;
+
+	UFUNCTION(BlueprintPure, FieldNotify, Category = "MetaSound|Randomizer")
+	UE_API float GetPitchMin() const;
+
+	UFUNCTION(BlueprintPure, FieldNotify, Category = "MetaSound|Randomizer")
+	UE_API float GetPitchMax() const;
+
+	UFUNCTION(BlueprintCallable, Category = "MetaSound|Randomizer")
+	UE_API void SetNumSounds(int32 NumSounds);
+
+	UFUNCTION(BlueprintCallable, Category = "MetaSound|Randomizer")
+	UE_API void SetPitchMin(float PitchMin);
+
+	UFUNCTION(BlueprintCallable, Category = "MetaSound|Randomizer")
+	UE_API void SetPitchMax(float PitchMax);
+
+	UE_API void SetIsOneShot(bool bInIsOneShot);
+	UE_API void SetPitch(const FVector2f& InPitch);
+
+	UE_API void SetSounds(TArray<TObjectPtr<USoundWave>> InSounds);
+
+	UE_API virtual bool IsSupportedConfiguration(const UScriptStruct& InStruct) const override;
+	UE_API virtual void Reload() override;
+
+private:
+	void OnPropertyChanged(const FPropertyChangedEvent& InEvent);
+
+	FDelegateHandle PropertyChangedHandle;
+
+	// If randomizer is one shot (looping if false).
+	UPROPERTY(BlueprintReadWrite, FieldNotify, Getter = "GetIsOneShot", Setter = "SetIsOneShot", Category = "MetaSound|Randomizer", meta = (AllowPrivateAccess))
+	bool bIsOneShot = false;
+
+	UPROPERTY(BlueprintReadWrite, FieldNotify, Getter = "GetSounds", Setter = "SetSounds", Category = "MetaSound|Randomizer", meta = (
+		AllowPrivateAccess,
+		DisallowedClasses = "/Script/MetasoundEngine.MetaSoundSource, /Script/Engine.SoundSourceBus"
+	))
+	TArray<TObjectPtr<USoundWave>> Sounds;
+
+	// Sound's pitch range in semitones.
+	UPROPERTY(BlueprintReadWrite, FieldNotify, Getter = "GetPitch", Setter = "SetPitch", Category = "MetaSound|Randomizer", meta = (AllowPrivateAccess))
+	FVector2f Pitch { 0.f, 0.f };
+};
+#undef UE_API
diff --git a/Engine/Plugins/Runtime/Metasound/Source/MetasoundEditor/Public/MetasoundDocumentConfigurationWidgetInterface.h b/Engine/Plugins/Runtime/Metasound/Source/MetasoundEditor/Public/MetasoundDocumentConfigurationWidgetInterface.h
new file mode 100644
--- /dev/null
+++ b/Engine/Plugins/Runtime/Metasound/Source/MetasoundEditor/Public/MetasoundDocumentConfigurationWidgetInterface.h
@@ -0,0 +1,75 @@
+// Copyright Epic Games, Inc. All Rights Reserved.
+
+#pragma once
+
+#include "Containers/Set.h"
+#include "MetasoundBuilderBase.h"
+#include "MetasoundBuilderSubsystem.h"
+#include "MetasoundDocumentInterface.h"
+#include "UObject/Class.h"
+#include "UObject/Interface.h"
+#include "UObject/ObjectMacros.h"
+#include "UObject/UObjectGlobals.h"
+#include "UObject/ScriptInterface.h"
+
+
+#include "MetasoundDocumentConfigurationWidgetInterface.generated.h"
+
+
+// Forward Declarations
+class UMetaSoundDocumentConfigurationViewModelBase;
+
+
+UINTERFACE(Blueprintable, MinimalAPI, DisplayName = "MetaSound Configuration Widget Interface")
+class UMetaSoundDocumentConfigurationWidgetInterface : public UInterface
+{
+	GENERATED_BODY()
+};
+
+USTRUCT(BlueprintType, MinimalAPI, DisplayName = "MetaSound Configuration Widget Tab Info")
+struct FMetaSoundDocumentConfigurationWidgetTabInfo
+{
+	GENERATED_BODY()
+
+	UPROPERTY(BlueprintReadWrite, Category = TabInfo)
+	FName TabName;
+
+	UPROPERTY(BlueprintReadWrite, Category = TabInfo)
+	FText TabDisplayName;
+
+	UPROPERTY(BlueprintReadWrite, Category = TabInfo)
+	FName StyleName = "MetaSoundStyle";
+
+	UPROPERTY(BlueprintReadWrite, Category = TabInfo)
+	FName TabIcon = "MetasoundEditor.Metasound.Icon";
+
+	UPROPERTY(BlueprintReadWrite, Category = TabInfo)
+	FText ToolTip;
+};
+
+class IMetaSoundDocumentConfigurationWidgetInterface
+{
+	GENERATED_BODY()
+
+public:
+	// Returns whether or not the widget supports the given configuration. Called when initializing
+	// the MetaSound Editor to determine if an implemented configuration supports the given widget.
+	UFUNCTION(BlueprintImplementableEvent, BlueprintCallable, Category = "MetaSound|Configuration", meta = (DisplayName = "Is Supported Configuration"))
+	void IsSupportedConfiguration(FName ConfigName, const UMetaSoundBuilderBase* Builder, UPARAM(DisplayName = "Is Supported") bool& bOutIsSupported);
+
+	// Called when the MetaSound starts and stops auditioning. Provides a reference to the audio component when auditioning starts, and returns nullptr when auditioning stops.
+	UFUNCTION(BlueprintImplementableEvent, Category = "MetaSound|Configuration", meta = (DisplayName = "On MetaSound Audition State Changed"))
+	void OnAuditionStateChanged(UAudioComponent* AudioComponent, bool bIsAuditioning);
+
+	// Called when the builder is initialized (prior to PreConstruct).
+	UFUNCTION(BlueprintImplementableEvent, Category = "MetaSound|Configuration", meta = (DisplayName = "On MetaSound Builder Set"))
+	void OnBuilderSet(UMetaSoundBuilderBase* Builder);
+
+	// (Optional) Called when initializing a tab owning the given widget to allow implementation to provide basic tab display options and information.
+	UFUNCTION(BlueprintImplementableEvent, Category = "MetaSound|Configuration", meta = (DisplayName = "On Initializing Tab"))
+	void OnInitializingTab(FMetaSoundDocumentConfigurationWidgetTabInfo& TabInfo);
+
+	// (Optional) If ViewModel class provided, automatically initializes and sets ViewModel of the given class for the given widget.
+	UFUNCTION(BlueprintImplementableEvent, Category = "MetaSound|Configuration", meta = (DisplayName = "On Initializing ViewModel"))
+	void OnInitializingViewModel(TSubclassOf<UMetaSoundDocumentConfigurationViewModelBase>& Class);
+};
diff --git a/Engine/Plugins/Runtime/Metasound/Source/MetasoundEditor/Public/MetasoundPresetWidgetInterface.h b/Engine/Plugins/Runtime/Metasound/Source/MetasoundEditor/Public/MetasoundPresetWidgetInterface.h
--- a/Engine/Plugins/Runtime/Metasound/Source/MetasoundEditor/Public/MetasoundPresetWidgetInterface.h
+++ b/Engine/Plugins/Runtime/Metasound/Source/MetasoundEditor/Public/MetasoundPresetWidgetInterface.h
@@ -4,6 +4,7 @@
 
 #include "Audio/AudioPanelWidgetInterface.h"
 #include "MetasoundBuilderSubsystem.h"
+#include "MetasoundDocumentConfigurationWidgetInterface.h"
 #include "UObject/Class.h"
 #include "UObject/Interface.h"
 #include "UObject/ObjectMacros.h"
@@ -12,6 +13,8 @@
 
 #include "MetasoundPresetWidgetInterface.generated.h"
 
+
+// IMetaSoundPresetWidgetInterface has been soft deprecated. Use IMetaSoundDocumentConfigurationWidgetInterface instead.
 UINTERFACE(Blueprintable, MinimalAPI)
 class UMetaSoundPresetWidgetInterface : public UAudioPanelWidgetInterface
 {
@@ -36,4 +39,8 @@ public:
 	// Called when the MetaSound starts and stops auditioning. Provides a reference to the audio component when auditioning starts, and returns nullptr when auditioning stops.
 	UFUNCTION(BlueprintCallable, BlueprintNativeEvent, Category = "MetaSound Preset Widget", meta = (DisplayName = "On MetaSound Audition State Changed"))
 	void OnAuditionStateChanged(UAudioComponent* AudioComponent, bool bIsAuditioning);
+
+	// (Optional) Called when initializing a tab owning the given widget to allow implementation to provide basic tab display options and information.
+	UFUNCTION(BlueprintImplementableEvent, Category = "MetaSound Preset Widget", meta = (DisplayName = "On Initializing Tab"))
+	void OnInitializingTab(FMetaSoundDocumentConfigurationWidgetTabInfo& TabInfo);
 };
diff --git a/Engine/Plugins/Runtime/Metasound/Source/MetasoundEngine/Private/DocumentConfigurations/MetasoundRandomizerDocumentConfiguration.cpp b/Engine/Plugins/Runtime/Metasound/Source/MetasoundEngine/Private/DocumentConfigurations/MetasoundRandomizerDocumentConfiguration.cpp
--- a/Engine/Plugins/Runtime/Metasound/Source/MetasoundEngine/Private/DocumentConfigurations/MetasoundRandomizerDocumentConfiguration.cpp
+++ b/Engine/Plugins/Runtime/Metasound/Source/MetasoundEngine/Private/DocumentConfigurations/MetasoundRandomizerDocumentConfiguration.cpp
@@ -279,9 +279,14 @@ bool FMetaSoundRandomizerDocumentConfiguration::ConfigureDocument(FMetaSoundFron
 #endif // WITH_EDITORONLY_DATA
 
 #if WITH_EDITOR
-TArray<const UClass*> FMetaSoundRandomizerDocumentConfiguration::GetSupportedAssetActionClasses() const
+const FMetaSoundFrontendDocumentConfiguration::FEditorOptions& FMetaSoundRandomizerDocumentConfiguration::GetEditorOptions() const
 {
-	return { USoundWave::StaticClass() };
+	static const FEditorOptions RandomizerOptions
+	{
+		.AssetActionClasses = { USoundWave::StaticClass() },
+		.bConfigurationPropertiesVisible = false
+	};
+	return RandomizerOptions;
 }
 
 void FMetaSoundRandomizerDocumentConfiguration::OnAssetInitialized(TArray<UObject*> SelectedObjects, FMetaSoundFrontendDocumentBuilder& OutBuilder)
@@ -294,7 +299,22 @@ void FMetaSoundRandomizerDocumentConfiguration::OnAssetInitialized(TArray<UObjec
 
 void FMetaSoundRandomizerDocumentConfiguration::OnPropertyChanged(const FPropertyChangedEvent& InEvent, FMetaSoundFrontendDocumentBuilder& OutBuilder)
 {
-	ConfigureDocument(OutBuilder, false);
+	if (InEvent.Property && InEvent.Property->GetName() == GET_MEMBER_NAME_CHECKED(FMetaSoundRandomizerDocumentConfiguration, Sounds))
+	{
+		RemoveInvalidSounds(OutBuilder, Sounds);
+	}
+
+	Super::OnPropertyChanged(InEvent, OutBuilder);
+}
+
+bool FMetaSoundRandomizerDocumentConfiguration::RemoveInvalidSounds(const FMetaSoundFrontendDocumentBuilder& ParentBuilder, TArray<TObjectPtr<USoundWave>>& OutSounds)
+{
+	const UObject& ThisObject = ParentBuilder.CastDocumentObjectChecked<UObject>();
+	return OutSounds.RemoveAll([&ThisObject](const TObjectPtr<USoundWave>& Wave)
+	{
+		// Cannot reference this builder's object if it inherits from USoundWave (ex. MetaSoundSource) nor non-explicit USoundWave types.
+		return !Wave || Wave == &ThisObject || Wave->GetClass() != USoundWave::StaticClass();
+	}) > 0;
 }
 #endif // WITH_EDITOR
 #undef LOCTEXT_NAMESPACE
diff --git a/Engine/Plugins/Runtime/Metasound/Source/MetasoundEngine/Private/MetasoundAssetSubsystem.cpp b/Engine/Plugins/Runtime/Metasound/Source/MetasoundEngine/Private/MetasoundAssetSubsystem.cpp
--- a/Engine/Plugins/Runtime/Metasound/Source/MetasoundEngine/Private/MetasoundAssetSubsystem.cpp
+++ b/Engine/Plugins/Runtime/Metasound/Source/MetasoundEngine/Private/MetasoundAssetSubsystem.cpp
@@ -801,11 +801,16 @@ PRAGMA_ENABLE_DEPRECATION_WARNINGS
 			return true;
 		}
 
-		// Presets must have a single asset reference 
-		FMetasoundAssetBase* ReferencedAsset = InAsset.GetReferencedAssets().Last();
-		if (InAsset.GetReferencedAssets().Num() != 1 || !ReferencedAsset)
+		// Presets must have a single asset reference
+		FMetasoundAssetBase* ReferencedAsset = nullptr;
+		if (InAsset.GetReferencedAssets().Num() == 1)
 		{
-			UE_LOG(LogMetaSound, Warning, TEXT("'%s' is not a valid MetaSound preset with a single referenced asset. Ending referenced preset iteration."), *InAsset.GetOwningAssetName());
+			ReferencedAsset = InAsset.GetReferencedAssets().Last();
+		}
+
+		if (!ReferencedAsset)
+		{
+			UE_LOG(LogMetaSound, Display, TEXT("MetaSound preset '%s' does not provide valid parent asset. Ending referenced preset iteration."), *InAsset.GetOwningAssetName());
 			return false;
 		}
 
diff --git a/Engine/Plugins/Runtime/Metasound/Source/MetasoundEngine/Private/MetasoundBuilderBase.cpp b/Engine/Plugins/Runtime/Metasound/Source/MetasoundEngine/Private/MetasoundBuilderBase.cpp
--- a/Engine/Plugins/Runtime/Metasound/Source/MetasoundEngine/Private/MetasoundBuilderBase.cpp
+++ b/Engine/Plugins/Runtime/Metasound/Source/MetasoundEngine/Private/MetasoundBuilderBase.cpp
@@ -229,10 +229,10 @@ void UMetaSoundBuilderBase::AddInterface(FName InterfaceName, EMetaSoundBuilderR
 	OutResult = bInterfaceAdded ? EMetaSoundBuilderResult::Succeeded : EMetaSoundBuilderResult::Failed;
 }
 
-void UMetaSoundBuilderBase::AddTransactionListener(TSharedRef<Metasound::Frontend::IDocumentBuilderTransactionListener> BuilderListener)
+FDelegateHandle UMetaSoundBuilderBase::AddTransactionListener(TSharedRef<Metasound::Frontend::IDocumentBuilderTransactionListener> BuilderListener)
 {
 	BuilderListener->OnBuilderReloaded(GetBuilderDelegates());
-	BuilderReloadDelegate.AddSP(BuilderListener, &Metasound::Frontend::IDocumentBuilderTransactionListener::OnBuilderReloaded);
+	return BuilderReloadDelegate.AddSP(BuilderListener, &Metasound::Frontend::IDocumentBuilderTransactionListener::OnBuilderReloaded);
 }
 
 FMetaSoundNodeHandle UMetaSoundBuilderBase::AddNode(const TScriptInterface<IMetaSoundDocumentInterface>& NodeClass, EMetaSoundBuilderResult& OutResult)
@@ -1333,6 +1333,24 @@ void UMetaSoundBuilderBase::SetMemberMetadata(UMetaSoundFrontendMemberMetadata&
 }
 #endif // WITH_EDITOR
 
+void UMetaSoundBuilderBase::SetConfiguration(TInstancedStruct<FMetaSoundFrontendDocumentConfiguration> Configuration)
+{
+	Builder.SetDocumentConfiguration(MoveTemp(Configuration));
+
+PRAGMA_DISABLE_EXPERIMENTAL_WARNINGS
+	Builder.ConfigureDocument(true /* bLogFailure */);
+PRAGMA_ENABLE_EXPERIMENTAL_WARNINGS
+}
+
+void UMetaSoundBuilderBase::SetConfigurationProperties(TFunctionRef<void(TStructView<FMetaSoundFrontendDocumentConfiguration>)> SetPropertyFunction)
+{
+	Builder.SetDocumentConfigurationProperties(SetPropertyFunction);
+
+	PRAGMA_DISABLE_EXPERIMENTAL_WARNINGS
+	Builder.ConfigureDocument(true /* bLogFailure */);
+	PRAGMA_ENABLE_EXPERIMENTAL_WARNINGS
+}
+
 void UMetaSoundBuilderBase::SetNodeInputDefault(const FMetaSoundBuilderNodeInputHandle& InputHandle, const FMetasoundFrontendLiteral& Literal, EMetaSoundBuilderResult& OutResult)
 {
 	const bool bInputDefaultSet = Builder.SetNodeInputDefault(InputHandle.NodeID, InputHandle.VertexID, Literal);
@@ -1405,11 +1423,35 @@ FMetasoundFrontendVersion UMetaSoundBuilderBase::FindNodeClassVersion(const FMet
 	return FMetasoundFrontendVersion::GetInvalid();
 }
 
+TInstancedStruct<FMetaSoundFrontendDocumentConfiguration> UMetaSoundBuilderBase::GetConfiguration(EMetaSoundBuilderResult& OutResult) const
+{
+	using namespace Metasound::Frontend;
+
+	OutResult = EMetaSoundBuilderResult::Failed;
+	TInstancedStruct<FMetaSoundFrontendDocumentConfiguration> NewInstance;
+
+	if (TConstStructView<FDocumentConfiguration> DocConfig = Builder.GetConstDocumentConfiguration(); DocConfig.IsValid())
+	{
+		const UScriptStruct* ScriptStruct = DocConfig.GetScriptStruct();
+		check(ScriptStruct);
+
+		NewInstance.InitializeAsScriptStruct(ScriptStruct, DocConfig.GetMemory());
+		OutResult = EMetaSoundBuilderResult::Succeeded;
+	}
+
+	return NewInstance;
+}
+
 FMetasoundFrontendClassName UMetaSoundBuilderBase::GetRootGraphClassName() const
 {
 	return Builder.GetConstDocumentChecked().RootGraph.Metadata.GetClassName();
 }
 
+TScriptInterface<IMetaSoundDocumentInterface> UMetaSoundBuilderBase::GetMetaSound() const
+{
+	return TScriptInterface<IMetaSoundDocumentInterface>(&Builder.CastDocumentObjectChecked<UObject>());
+}
+
 void UMetaSoundBuilderBase::GetNodeInputData(const FMetaSoundBuilderNodeInputHandle& InputHandle, FName& Name, FName& DataType, EMetaSoundBuilderResult& OutResult)
 {
 	if (const FMetasoundFrontendVertex* Vertex = Builder.FindNodeInput(InputHandle.NodeID, InputHandle.VertexID))
diff --git a/Engine/Plugins/Runtime/Metasound/Source/MetasoundEngine/Private/MetasoundEngineAsset.h b/Engine/Plugins/Runtime/Metasound/Source/MetasoundEngine/Private/MetasoundEngineAsset.h
--- a/Engine/Plugins/Runtime/Metasound/Source/MetasoundEngine/Private/MetasoundEngineAsset.h
+++ b/Engine/Plugins/Runtime/Metasound/Source/MetasoundEngine/Private/MetasoundEngineAsset.h
@@ -87,6 +87,7 @@ namespace Metasound::Engine
 			FMetaSoundFrontendDocumentBuilder& Builder = BuilderRegistry.FindOrBeginBuilding(InDocInterface);
 
 #if WITH_EDITORONLY_DATA
+			Builder.GetDocumentDelegates().OnDocumentConfigurationChanged.Broadcast({ EDocumentConfigChangeType::UndoRedo });
 			Builder.ConfigureDocument();
 #endif // WITH_EDITORONLY_DATA
 
diff --git a/Engine/Plugins/Runtime/Metasound/Source/MetasoundEngine/Public/DocumentConfigurations/MetasoundRandomizerDocumentConfiguration.h b/Engine/Plugins/Runtime/Metasound/Source/MetasoundEngine/Public/DocumentConfigurations/MetasoundRandomizerDocumentConfiguration.h
--- a/Engine/Plugins/Runtime/Metasound/Source/MetasoundEngine/Public/DocumentConfigurations/MetasoundRandomizerDocumentConfiguration.h
+++ b/Engine/Plugins/Runtime/Metasound/Source/MetasoundEngine/Public/DocumentConfigurations/MetasoundRandomizerDocumentConfiguration.h
@@ -28,13 +28,17 @@ struct FMetaSoundRandomizerDocumentConfiguration : public FMetaSoundFrontendDocu
 
 #if WITH_EDITOR
 	UE_EXPERIMENTAL(5.8, "Document Configuration is in active development and subject to change")
-	UE_API virtual TArray<const UClass*> GetSupportedAssetActionClasses() const override;
+	UE_API virtual const FMetaSoundFrontendDocumentConfiguration::FEditorOptions& GetEditorOptions() const override;
 
 	UE_EXPERIMENTAL(5.8, "Document Configuration is in active development and subject to change")
 	UE_API virtual void OnAssetInitialized(TArray<UObject*> SelectedObjects, FMetaSoundFrontendDocumentBuilder& OutBuilder) override;
 
 	UE_EXPERIMENTAL(5.8, "Document Configuration is in active development and subject to change")
 	UE_API virtual void OnPropertyChanged(const FPropertyChangedEvent& InEvent, FMetaSoundFrontendDocumentBuilder& OutBuilder) override;
+
+	// Helper function to remove sounds that cannot be referenced by randomizer (i.e. the given parent builder's MetaSound
+	// if it is a source or child class types of USoundWave). Returns true if sounds were removed, false if not.
+	UE_API static bool RemoveInvalidSounds(const FMetaSoundFrontendDocumentBuilder& ParentBuilder, TArray<TObjectPtr<USoundWave>>& OutSounds);
 #endif // WITH_EDITOR
 
 public:
diff --git a/Engine/Plugins/Runtime/Metasound/Source/MetasoundEngine/Public/MetasoundBuilderBase.h b/Engine/Plugins/Runtime/Metasound/Source/MetasoundEngine/Public/MetasoundBuilderBase.h
--- a/Engine/Plugins/Runtime/Metasound/Source/MetasoundEngine/Public/MetasoundBuilderBase.h
+++ b/Engine/Plugins/Runtime/Metasound/Source/MetasoundEngine/Public/MetasoundBuilderBase.h
@@ -2,6 +2,7 @@
 #pragma once
 
 #include "Containers/Map.h"
+#include "Delegates/IDelegateInstance.h"
 #include "Interfaces/MetasoundOutputFormatInterfaces.h"
 #include "MetasoundDocumentInterface.h"
 #include "MetasoundDynamicOperatorTransactor.h"
@@ -183,7 +184,7 @@ public:
 	// Adds transaction listener which allows objects to respond to when certain graph operations are applied from anywhere (adding or removing nodes, edges, pages, etc.)
 	// Currently there is no guarantee all transactions will be represented until the Controller API is fully deprecated! (ex. if a node or edge is added or removed via a
 	// controller API call, the transaction will be missed). OnBuilderReloaded is however guaranteed to be called on mutable controller creation.
-	UE_API void AddTransactionListener(TSharedRef<Metasound::Frontend::IDocumentBuilderTransactionListener> BuilderListener);
+	UE_API FDelegateHandle AddTransactionListener(TSharedRef<Metasound::Frontend::IDocumentBuilderTransactionListener> BuilderListener);
 
 #if WITH_EDITOR
 	UE_API bool ClearMemberMetadata(const FGuid& InMemberID);
@@ -323,7 +324,11 @@ public:
 	// Returns output's parent node if the input is valid, otherwise returns invalid node handle.
 	UFUNCTION(BlueprintCallable, Category = "Audio|MetaSound|Builder", meta = (ExpandEnumAsExecs = "OutResult"))
 	UE_API UPARAM(DisplayName = "Node Class Version") FMetasoundFrontendVersion FindNodeClassVersion(const FMetaSoundNodeHandle& NodeHandle, EMetaSoundBuilderResult& OutResult);
-	
+
+	// Returns copy of the builder's configuration. Returns success if the builder has a configuration set, failed if not.
+	UFUNCTION(BlueprintCallable, Category = "Audio|MetaSound|Builder", meta = (ExpandEnumAsExecs = "OutResult"))
+	UE_API TInstancedStruct<FMetaSoundFrontendDocumentConfiguration> GetConfiguration(EMetaSoundBuilderResult& OutResult) const;
+
 	// Gets names of all graph inputs. 
 	UFUNCTION(BlueprintCallable, Category = "Audio|MetaSound|Builder", meta = (ExpandEnumAsExecs = "OutResult"))
 	UE_API UPARAM(DisplayName = "Graph Inputs") TArray<FName> GetGraphInputNames(EMetaSoundBuilderResult& OutResult) const;
@@ -344,6 +349,10 @@ public:
 	UFUNCTION(BlueprintCallable, Category = "Audio|MetaSound|Builder|Advanced", meta = (DisplayName = "Get MetaSound Class Name"))
 	UE_API UPARAM(DisplayName = "Class Name") FMetasoundFrontendClassName GetRootGraphClassName() const;
 
+	// Returns the interface to the MetaSound the builder is acting on.
+	UFUNCTION(BlueprintCallable, Category = "Audio|MetaSound|Builder")
+	UE_API UPARAM(DisplayName = "MetaSound") TScriptInterface<IMetaSoundDocumentInterface> GetMetaSound() const;
+
 	// Returns node input's data if valid (including things like name and datatype).
 	UFUNCTION(BlueprintCallable, Category = "Audio|MetaSound|Builder", meta = (ExpandEnumAsExecs = "OutResult"))
 	UE_API void GetNodeInputData(const FMetaSoundBuilderNodeInputHandle& InputHandle, FName& Name, FName& DataType, EMetaSoundBuilderResult& OutResult);
@@ -457,6 +466,26 @@ public:
 	UE_API void SetAuthor(const FString& InAuthor);
 #endif // WITH_EDITOR
 
+	// Sets the given builder's configuration and immediately configures the MetaSound.
+	UFUNCTION(BlueprintCallable, Category = "Audio|MetaSound|Builder")
+	UE_API void SetConfiguration(TInstancedStruct<FMetaSoundFrontendDocumentConfiguration> Configuration);
+
+	// Passes view of struct along for mutating properties, firing OnDocumentConfigurationChanged on completion and immediately reconfigures the MetaSound.
+	UE_API void SetConfigurationProperties(TFunctionRef<void(TStructView<FMetaSoundFrontendDocumentConfiguration>)> SetPropertiesFunction);
+
+	template <typename TConfig>
+	bool SetConfigurationProperties(TFunctionRef<void(TConfig&)> SetPropertyFunction)
+	{
+		const bool bConfigUpdated = GetBuilder().SetDocumentConfigurationProperties(SetPropertyFunction);
+		if (bConfigUpdated)
+		{
+			PRAGMA_DISABLE_EXPERIMENTAL_WARNINGS
+			GetBuilder().ConfigureDocument(true /* bLogFailure */);
+			PRAGMA_ENABLE_EXPERIMENTAL_WARNINGS
+		}
+		return bConfigUpdated;
+	}
+
 	// Sets the node's input default value (used if no connection to the given node input is present)
 	UFUNCTION(BlueprintCallable, Category = "Audio|MetaSound|Builder", meta = (ExpandEnumAsExecs = "OutResult"))
 	UE_API void SetNodeInputDefault(const FMetaSoundBuilderNodeInputHandle& NodeInputHandle, const FMetasoundFrontendLiteral& Literal, EMetaSoundBuilderResult& OutResult);
diff --git a/Engine/Plugins/Runtime/Metasound/Source/MetasoundFrontend/Private/DocumentConfigurations/MetasoundFrontendDocumentConfiguration.cpp b/Engine/Plugins/Runtime/Metasound/Source/MetasoundFrontend/Private/DocumentConfigurations/MetasoundFrontendDocumentConfiguration.cpp
--- a/Engine/Plugins/Runtime/Metasound/Source/MetasoundFrontend/Private/DocumentConfigurations/MetasoundFrontendDocumentConfiguration.cpp
+++ b/Engine/Plugins/Runtime/Metasound/Source/MetasoundFrontend/Private/DocumentConfigurations/MetasoundFrontendDocumentConfiguration.cpp
@@ -20,23 +20,29 @@ bool FMetaSoundFrontendDocumentConfiguration::ConfigureDocument(FMetaSoundFronte
 }
 
 #if WITH_EDITOR
-TArray<const UClass*> FMetaSoundFrontendDocumentConfiguration::GetSupportedAssetActionClasses() const
+Metasound::Frontend::FOnConfigPropertyChangedDelegate& FMetaSoundFrontendDocumentConfiguration::GetPropertyChangedDelegate() const
 {
-	return { };
+	return PropertyChangedDelegate;
 }
 
-TArray<const UClass*> FMetaSoundFrontendDocumentConfiguration::GetSupportedMetaSoundClasses() const
+const FMetaSoundFrontendDocumentConfiguration::FEditorOptions& FMetaSoundFrontendDocumentConfiguration::FEditorOptions::GetDefaultOptions()
 {
-	return { };
+	static const FEditorOptions Default;
+	return Default;
 }
 
-FText FMetaSoundFrontendDocumentConfiguration::GetToolTip() const
+const FMetaSoundFrontendDocumentConfiguration::FEditorOptions& FMetaSoundFrontendDocumentConfiguration::GetEditorOptions() const
+{
+	return FEditorOptions::GetDefaultOptions();
+}
+
+void FMetaSoundFrontendDocumentConfiguration::OnAssetInitialized(TArray<UObject*> SelectedObjects, FMetaSoundFrontendDocumentBuilder& OutBuilder)
 {
-	return { };
 }
 
 void FMetaSoundFrontendDocumentConfiguration::OnPropertyChanged(const FPropertyChangedEvent& InEvent, FMetaSoundFrontendDocumentBuilder& OutBuilder)
 {
+	PropertyChangedDelegate.Broadcast(InEvent);
 }
 #endif // WITH_EDITOR
 #undef LOCTEXT_NAMESPACE
diff --git a/Engine/Plugins/Runtime/Metasound/Source/MetasoundFrontend/Private/DocumentConfigurations/MetasoundFrontendPresetDocumentConfiguration.cpp b/Engine/Plugins/Runtime/Metasound/Source/MetasoundFrontend/Private/DocumentConfigurations/MetasoundFrontendPresetDocumentConfiguration.cpp
--- a/Engine/Plugins/Runtime/Metasound/Source/MetasoundFrontend/Private/DocumentConfigurations/MetasoundFrontendPresetDocumentConfiguration.cpp
+++ b/Engine/Plugins/Runtime/Metasound/Source/MetasoundFrontend/Private/DocumentConfigurations/MetasoundFrontendPresetDocumentConfiguration.cpp
@@ -115,9 +115,14 @@ TInstancedStruct<FMetaSoundFrontendDocumentConfigurationVertexMetadata> FMetaSou
 }
 
 #if WITH_EDITOR
-TArray<const UClass*> FMetaSoundFrontendPresetDocumentConfiguration::GetSupportedAssetActionClasses() const
+const FMetaSoundFrontendDocumentConfiguration::FEditorOptions& FMetaSoundFrontendPresetDocumentConfiguration::GetEditorOptions() const
 {
-	return { };
+	static const FEditorOptions PresetOptions
+	{
+		.bMemberEditingEnabled = true,
+		.bGraphEditorVisible = true
+	};
+	return PresetOptions;
 }
 
 void FMetaSoundFrontendPresetDocumentConfiguration::OnAssetInitialized(TArray<UObject*> SelectedObjects, FMetaSoundFrontendDocumentBuilder& OutBuilder)
@@ -146,7 +151,6 @@ void FMetaSoundFrontendPresetDocumentConfiguration::OnPropertyChanged(const FPro
 			}
 		}
 	}
-	ConfigureDocument(OutBuilder, false);
 }
 
 bool FMetaSoundFrontendPresetConfigurationVertexMetadata::IsDefaultEditable() const
diff --git a/Engine/Plugins/Runtime/Metasound/Source/MetasoundFrontend/Private/MetasoundFrontendDocumentBuilder.cpp b/Engine/Plugins/Runtime/Metasound/Source/MetasoundFrontend/Private/MetasoundFrontendDocumentBuilder.cpp
--- a/Engine/Plugins/Runtime/Metasound/Source/MetasoundFrontend/Private/MetasoundFrontendDocumentBuilder.cpp
+++ b/Engine/Plugins/Runtime/Metasound/Source/MetasoundFrontend/Private/MetasoundFrontendDocumentBuilder.cpp
@@ -2550,8 +2550,13 @@ bool FMetaSoundFrontendDocumentBuilder::ConfigureDocument(bool bLogFailure)
 		TSharedRef<FDocumentModifyDelegates> ModifyDelegates = MakeShared<FDocumentModifyDelegates>(*DocumentDelegates);
 		ClearDocument(ModifyDelegates);
 
-		// Add required interfaces to ensure document is left in valid state (ex. MetaSound sources
-		// need the source interface to properly initialize a null-rendering MetaSound generator.)
+#if WITH_EDITORONLY_DATA
+		SetIsGraphEditable(false);
+#endif // WITH_EDITORONLY_DATA
+
+		// Add required interfaces, but disable editing to ensure document maintains valid state
+		// (ex. MetaSound sources need  the source interface to properly initialize a null-rendering
+		// MetaSound generator.)
 		TArray<FMetasoundFrontendVersion> InitVersions = ISearchEngine::Get().FindUClassDefaultInterfaceVersions(GetBuilderClassPath());
 		FModifyInterfaceOptions Options({ }, MoveTemp(InitVersions));
 		ModifyInterfaces(MoveTemp(Options));
@@ -2592,6 +2597,8 @@ bool FMetaSoundFrontendDocumentBuilder::ConvertFromPreset()
 
 	if (IsPreset())
 	{
+		const bool bConfigChanged = GetDocumentChecked().Configuration.IsValid();
+
 		// Don't process the null configuration in order to "freeze" the current document state and not clear it.
 		GetDocumentChecked().Configuration.Reset();
 
@@ -2599,7 +2606,11 @@ bool FMetaSoundFrontendDocumentBuilder::ConvertFromPreset()
 		SetIsGraphEditable(true);
 #endif // WITH_EDITORONLY_DATA?
 
-		DocumentDelegates->OnDocumentConfigurationChanged.Broadcast(FDocumentConfigChangedArgs { });
+		if (bConfigChanged)
+		{
+			DocumentDelegates->OnDocumentConfigurationChanged.Broadcast(FDocumentConfigChangedArgs{ .Type = EDocumentConfigChangeType::Struct });
+		}
+
 		return true;
 	}
 
@@ -5906,7 +5917,14 @@ void FMetaSoundFrontendDocumentBuilder::SetDocumentConfiguration(TInstancedStruc
 	using namespace Metasound::Frontend;
 
 	GetDocumentChecked().Configuration = MoveTemp(InConfig);
-	DocumentDelegates->OnDocumentConfigurationChanged.Broadcast({ });
+	DocumentDelegates->OnDocumentConfigurationChanged.Broadcast({ EDocumentConfigChangeType::Struct });
+}
+
+void FMetaSoundFrontendDocumentBuilder::SetDocumentConfigurationProperties(TFunctionRef<void(TStructView<FMetaSoundFrontendDocumentConfiguration>)> SetPropertyFunction)
+{
+	using namespace Metasound::Frontend;
+	SetPropertyFunction(GetDocumentChecked().Configuration);
+	DocumentDelegates->OnDocumentConfigurationChanged.Broadcast({ EDocumentConfigChangeType::Property });
 }
 
 #if WITH_EDITORONLY_DATA
diff --git a/Engine/Plugins/Runtime/Metasound/Source/MetasoundFrontend/Public/DocumentConfigurations/MetasoundFrontendDocumentConfiguration.h b/Engine/Plugins/Runtime/Metasound/Source/MetasoundFrontend/Public/DocumentConfigurations/MetasoundFrontendDocumentConfiguration.h
--- a/Engine/Plugins/Runtime/Metasound/Source/MetasoundFrontend/Public/DocumentConfigurations/MetasoundFrontendDocumentConfiguration.h
+++ b/Engine/Plugins/Runtime/Metasound/Source/MetasoundFrontend/Public/DocumentConfigurations/MetasoundFrontendDocumentConfiguration.h
@@ -2,6 +2,11 @@
 #pragma once
 
 #include "Containers/Set.h"
+
+#if WITH_EDITOR
+#include "Delegates/DelegateCombinations.h"
+#endif // WITH_EDITOR
+
 #include "Internationalization/Text.h"
 #include "Misc/Guid.h"
 #include "StructUtils/InstancedStruct.h"
@@ -20,6 +25,12 @@ class UMetaSoundFrontendMemberMetadata;
 
 enum class EMetasoundFrontendClassType : uint8;
 
+#if WITH_EDITOR
+namespace Metasound::Frontend
+{
+	DECLARE_MULTICAST_DELEGATE_OneParam(FOnConfigPropertyChangedDelegate, const FPropertyChangedEvent& /* InEvent */);
+} // namespace Metasound::Frontend
+#endif // WITH_EDITOR
 
 // Configuration of a MetaSound, used to define procedural behavior when making modifications
 // to a MetaSound prior to serialization. Configurations enable MetaSounds to be procedurally
@@ -37,41 +48,65 @@ struct FMetaSoundFrontendDocumentConfiguration
 	UE_INTERNAL UE_API virtual bool ConfigureDocument(FMetaSoundFrontendDocumentBuilder& OutBuilder, bool bLogFailure);
 
 #if WITH_EDITOR
-	// Asset class types this configuration supports as a selection (passed on initial creation) when user
-	// creates a configured MetaSound from an action (eg. right-click action in the Content Browser right-click
-	// menu). Defaults to all MetaSound asset types (eg. Source & Patch) if unset.
-	UE_EXPERIMENTAL(5.8, "Document Configuration is in active development and subject to change")
-	UE_API virtual TArray<const UClass*> GetSupportedAssetActionClasses() const;
-
-	// MetaSound asset types this configuration can be set to when user creates a configured MetaSound
-	// from an action (eg. right-click action in the Content Browser right-click menu). Defaults to
-	// all MetaSound asset types (eg. Source & Patch) if unset.
-	UE_EXPERIMENTAL(5.8, "Document Configuration is in active development and subject to change")
-	UE_API virtual TArray<const UClass*> GetSupportedMetaSoundClasses() const;
+	struct FEditorOptions
+	{
+		// Asset class types this configuration supports as a selection (passed on initial creation) when user
+		// creates a configured MetaSound from an action (eg. right-click action in the Content Browser right-click
+		// menu). Defaults to all MetaSound asset types (eg. Source & Patch) if unset.
+		UE_EXPERIMENTAL(5.8, "Document Configuration is in active development and subject to change")
+		 TArray<TSoftClassPtr<UObject>> AssetActionClasses;
+
+		// MetaSound asset types this configuration can be set to when user creates a configured MetaSound
+		// from an action (eg. right-click action in the Content Browser right-click menu). Defaults to
+		// all MetaSound asset types (eg. Source & Patch) if unset.
+		UE_EXPERIMENTAL(5.8, "Document Configuration is in active development and subject to change")
+		TArray<TSoftClassPtr<UObject>> MetaSoundClasses;
+
+		// ToolTip shown when selecting document configuration from list of available configurations
+		UE_EXPERIMENTAL(5.8, "Document Configuration is in active development and subject to change")
+		FText ToolTip;
+
+		// Whether or not configuration supports custom tabs (i.e. custom document configuration widgets)
+		UE_EXPERIMENTAL(5.8, "Document Configuration is in active development and subject to change")
+		bool bCustomTabEnabled = true;
+
+		// Whether or not interface editing is enabled in MetaSound asset editors
+		UE_EXPERIMENTAL(5.8, "Document Configuration is in active development and subject to change")
+		bool bInterfaceEditingEnabled = false;
+
+		// Whether or not member editing is enabled in MetaSound asset editors
+		UE_EXPERIMENTAL(5.8, "Document Configuration is in active development and subject to change")
+		bool bMemberEditingEnabled = false;
+
+		// Whether or not graph page editing is enabled in MetaSound asset editors
+		UE_EXPERIMENTAL(5.8, "Document Configuration is in active development and subject to change")
+		bool bPageGraphEditingEnabled = false;
+
+		// Whether or not graph editor is visible by default (user can overwrite if desired)
+		UE_EXPERIMENTAL(5.8, "Document Configuration is in active development and subject to change")
+		bool bGraphEditorVisible = false;
+
+		// Whether or not to display configuration properties in the details panel.
+		UE_EXPERIMENTAL(5.8, "Document Configuration is in active development and subject to change")
+		bool bConfigurationPropertiesVisible = true;
+
+		// Default options if none are specified
+		UE_API static const FEditorOptions& GetDefaultOptions();
+	};
 
 	UE_EXPERIMENTAL(5.8, "Document Configuration is in active development and subject to change")
-	UE_API virtual FText GetToolTip() const;
+	UE_API virtual const FEditorOptions& GetEditorOptions() const;
 
 	UE_EXPERIMENTAL(5.8, "Document Configuration is in active development and subject to change")
-	virtual bool IsCustomTabEnabled() const { return true; }
+	UE_API virtual void OnAssetInitialized(TArray<UObject*> SelectedObjects, FMetaSoundFrontendDocumentBuilder& OutBuilder);
 
 	UE_EXPERIMENTAL(5.8, "Document Configuration is in active development and subject to change")
-	virtual bool IsGraphEditorVisible() const { return false; }
-
-	UE_EXPERIMENTAL(5.8, "Document Configuration is in active development and subject to change")
-	virtual bool IsInterfaceEditingEnabled() const { return false; }
-
-	UE_EXPERIMENTAL(5.8, "Document Configuration is in active development and subject to change")
-	virtual bool IsMemberEditingEnabled() const { return false; }
-
-	UE_EXPERIMENTAL(5.8, "Document Configuration is in active development and subject to change")
-	virtual bool IsPageGraphEditingEnabled() const { return false; }
+	UE_API virtual void OnPropertyChanged(const FPropertyChangedEvent& InEvent, FMetaSoundFrontendDocumentBuilder& OutBuilder);
 
-	UE_EXPERIMENTAL(5.8, "Document Configuration is in active development and subject to change")
-	virtual void OnAssetInitialized(TArray<UObject*> SelectedObjects, FMetaSoundFrontendDocumentBuilder& OutBuilder) { }
+	UE_API Metasound::Frontend::FOnConfigPropertyChangedDelegate& GetPropertyChangedDelegate() const;
 
-	UE_EXPERIMENTAL(5.8, "Document Configuration is in active development and subject to change")
-	UE_API virtual void OnPropertyChanged(const FPropertyChangedEvent& InEvent, FMetaSoundFrontendDocumentBuilder& OutBuilder);
+private:
+	mutable Metasound::Frontend::FOnConfigPropertyChangedDelegate PropertyChangedDelegate;
 #endif // WITH_EDITOR
 };
 
diff --git a/Engine/Plugins/Runtime/Metasound/Source/MetasoundFrontend/Public/DocumentConfigurations/MetasoundFrontendPresetDocumentConfiguration.h b/Engine/Plugins/Runtime/Metasound/Source/MetasoundFrontend/Public/DocumentConfigurations/MetasoundFrontendPresetDocumentConfiguration.h
--- a/Engine/Plugins/Runtime/Metasound/Source/MetasoundFrontend/Public/DocumentConfigurations/MetasoundFrontendPresetDocumentConfiguration.h
+++ b/Engine/Plugins/Runtime/Metasound/Source/MetasoundFrontend/Public/DocumentConfigurations/MetasoundFrontendPresetDocumentConfiguration.h
@@ -28,16 +28,10 @@ struct FMetaSoundFrontendPresetDocumentConfiguration final : public FMetaSoundFr
 
 #if WITH_EDITOR
 	UE_EXPERIMENTAL(5.8, "Document Configuration is in active development and subject to change")
-	virtual TArray<const UClass*> GetSupportedAssetActionClasses() const override;
+	UE_API virtual const FMetaSoundFrontendDocumentConfiguration::FEditorOptions& GetEditorOptions() const override;
 
 	UE_EXPERIMENTAL(5.8, "Document Configuration is in active development and subject to change")
-	virtual bool IsMemberEditingEnabled() const override { return true; }
-
-	UE_EXPERIMENTAL(5.8, "Document Configuration is in active development and subject to change")
-	virtual bool IsGraphEditorVisible() const override { return true; }
-
-	UE_EXPERIMENTAL(5.8, "Document Configuration is in active development and subject to change")
-	virtual void OnAssetInitialized(TArray<UObject*> SelectedObjects, FMetaSoundFrontendDocumentBuilder& OutBuilder) override;
+	UE_API virtual void OnAssetInitialized(TArray<UObject*> SelectedObjects, FMetaSoundFrontendDocumentBuilder& OutBuilder) override;
 
 	UE_EXPERIMENTAL(5.8, "Document Configuration is in active development and subject to change")
 	UE_API virtual void OnPropertyChanged(const FPropertyChangedEvent& InEvent, FMetaSoundFrontendDocumentBuilder& OutBuilder) override;
diff --git a/Engine/Plugins/Runtime/Metasound/Source/MetasoundFrontend/Public/MetasoundFrontendDocumentBuilder.h b/Engine/Plugins/Runtime/Metasound/Source/MetasoundFrontend/Public/MetasoundFrontendDocumentBuilder.h
--- a/Engine/Plugins/Runtime/Metasound/Source/MetasoundFrontend/Public/MetasoundFrontendDocumentBuilder.h
+++ b/Engine/Plugins/Runtime/Metasound/Source/MetasoundFrontend/Public/MetasoundFrontendDocumentBuilder.h
@@ -670,7 +670,35 @@ PRAGMA_ENABLE_EXPERIMENTAL_WARNINGS
 #endif // WITH_EDITORONLY_DATA
 
 PRAGMA_DISABLE_EXPERIMENTAL_WARNINGS
+	// Sets the entire internal configuration with that provided.
 	UE_API void SetDocumentConfiguration(TInstancedStruct<FMetaSoundFrontendDocumentConfiguration> InConfig);
+
+	// Update the given document configuration inline using the internal struct view passed to a given function.
+	UE_API void SetDocumentConfigurationProperties(TFunctionRef<void(TStructView<FMetaSoundFrontendDocumentConfiguration>)> SetPropertyFunction);
+
+	// If struct type of function coincides with that of the given builder's MetaSound configuration, sets
+	// the given document configuration properties inline using the internal struct view passed to a given
+	// function. Returns true if set function successfully executed, false if not.
+	template <typename TConfig>
+	bool SetDocumentConfigurationProperties(TFunctionRef<void(TConfig&)> SetPropertyFunction)
+	{
+		bool bSucceeded = false;
+
+		const UScriptStruct* ScriptStruct = GetConstDocumentConfiguration().GetScriptStruct();
+		if (ScriptStruct && ScriptStruct->IsChildOf(TConfig::StaticStruct()))
+		{
+			SetDocumentConfigurationProperties([&SetPropertyFunction, &bSucceeded](TStructView<FMetaSoundFrontendDocumentConfiguration> Config)
+			{
+				if (TConfig* ConfigPtr = Config.GetPtr<TConfig>())
+				{
+					SetPropertyFunction(*ConfigPtr);
+					bSucceeded = true;
+				}
+			});
+		}
+
+		return bSucceeded;
+	}
 PRAGMA_ENABLE_EXPERIMENTAL_WARNINGS
 
 #if WITH_EDITORONLY_DATA
diff --git a/Engine/Plugins/Runtime/Metasound/Source/MetasoundFrontend/Public/MetasoundFrontendDocumentModifyDelegates.h b/Engine/Plugins/Runtime/Metasound/Source/MetasoundFrontend/Public/MetasoundFrontendDocumentModifyDelegates.h
--- a/Engine/Plugins/Runtime/Metasound/Source/MetasoundFrontend/Public/MetasoundFrontendDocumentModifyDelegates.h
+++ b/Engine/Plugins/Runtime/Metasound/Source/MetasoundFrontend/Public/MetasoundFrontendDocumentModifyDelegates.h
@@ -28,8 +28,27 @@ namespace Metasound::Frontend
 		FGuid PageID;
 	};
 
+	enum class EDocumentConfigChangeType : uint8
+	{
+		// Initial, unset value.
+		Invalid = 0,
+
+		// Config property or properties were changed
+		Property,
+
+		// Config struct has been completely replaced
+		// (Could now be unset/cleared)
+		Struct,
+
+		// Transaction has been re-applied (redo) or rolled back (undo)
+		UndoRedo,
+	};
+
+	struct FDocumentConfigChangedArgs
+	{
+		EDocumentConfigChangeType Type = EDocumentConfigChangeType::Invalid;
+	};
 
-	struct FDocumentConfigChangedArgs {};
 	DECLARE_MULTICAST_DELEGATE_OneParam(FOnDocumentConfigurationChanged, const FDocumentConfigChangedArgs& /* Args */);
 
 	// Soft deprecated in 5.8 in favor of FOnDocumentConfigurationChanged.  Once delegate is removed, it can be officially deprecated.
diff --git a/Engine/Source/Runtime/Engine/Private/Audio/AudioWidgetSubsystem.cpp b/Engine/Source/Runtime/Engine/Private/Audio/AudioWidgetSubsystem.cpp
--- a/Engine/Source/Runtime/Engine/Private/Audio/AudioWidgetSubsystem.cpp
+++ b/Engine/Source/Runtime/Engine/Private/Audio/AudioWidgetSubsystem.cpp
@@ -17,7 +17,11 @@
 #define LOCTEXT_NAMESPACE "AudioWidgetSubsystem"
 
 
-TArray<UUserWidget*> UAudioWidgetSubsystem::CreateUserWidgets(UWorld& InWorld, TSubclassOf<UInterface> InWidgetInterface, TFunction<bool(UUserWidget*)> FilterFunction) const
+TArray<UUserWidget*> UAudioWidgetSubsystem::CreateUserWidgets(
+	UWorld& InWorld,
+	TSubclassOf<UInterface> InWidgetInterface,
+	TFunction<bool(UUserWidget*)> FilterFunction,
+	TFunction<bool(UBlueprint*)> BlueprintFilterFunction) const
 {
 	TArray<UUserWidget*> UserWidgets;
 	UClass* InterfaceClass = InWidgetInterface ? InWidgetInterface.Get() : UInterface::StaticClass();
@@ -42,13 +46,14 @@ TArray<UUserWidget*> UAudioWidgetSubsystem::CreateUserWidgets(UWorld& InWorld, T
 			continue;
 		}
 
+		if (BlueprintFilterFunction.IsSet() && !BlueprintFilterFunction(Blueprint))
+		{
+			continue;
+		}
+
 		if (UUserWidget * UserWidget = CreateWidget<UUserWidget>(&InWorld, GeneratedClass))
 		{
-			if (!FilterFunction)
-			{
-				UserWidgets.Add(UserWidget);
-			}
-			else if (FilterFunction(UserWidget))
+			if (!FilterFunction || FilterFunction(UserWidget))
 			{
 				UserWidgets.Add(UserWidget);
 			}
diff --git a/Engine/Source/Runtime/Engine/Public/Audio/AudioWidgetSubsystem.h b/Engine/Source/Runtime/Engine/Public/Audio/AudioWidgetSubsystem.h
--- a/Engine/Source/Runtime/Engine/Public/Audio/AudioWidgetSubsystem.h
+++ b/Engine/Source/Runtime/Engine/Public/Audio/AudioWidgetSubsystem.h
@@ -32,5 +32,9 @@ private:
 
 public:
 	// Returns user widgets that implement a widget interface of the given subclass and optionally, filters and returns only those which pass the provided filter function.
-	ENGINE_API TArray<UUserWidget*> CreateUserWidgets(UWorld& InWorld, TSubclassOf<UInterface> InWidgetInterface, TFunction<bool(UUserWidget*)> FilterFunction = TFunction<bool(UUserWidget*)>()) const;
+	ENGINE_API TArray<UUserWidget*> CreateUserWidgets(
+		UWorld& InWorld,
+		TSubclassOf<UInterface> InWidgetInterface,
+		TFunction<bool(UUserWidget*)> FilterFunction = TFunction<bool(UUserWidget*)>(),
+		TFunction<bool(UBlueprint*)> BlueprintFilterFunction = TFunction<bool(UBlueprint*)>()) const;
 };
