Remove copy/pasta pragma once
#rnx

#Synced-CL 49790744

diff --git a/Engine/Plugins/Runtime/Metasound/Source/MetasoundEditor/Private/DocumentConfigurationViewModels/MetasoundDocumentConfigurationViewModelBase.cpp b/Engine/Plugins/Runtime/Metasound/Source/MetasoundEditor/Private/DocumentConfigurationViewModels/MetasoundDocumentConfigurationViewModelBase.cpp
--- a/Engine/Plugins/Runtime/Metasound/Source/MetasoundEditor/Private/DocumentConfigurationViewModels/MetasoundDocumentConfigurationViewModelBase.cpp
+++ b/Engine/Plugins/Runtime/Metasound/Source/MetasoundEditor/Private/DocumentConfigurationViewModels/MetasoundDocumentConfigurationViewModelBase.cpp
@@ -1,7 +1,5 @@
 // Copyright Epic Games, Inc. All Rights Reserved.
 
-#pragma once
-
 #include "DocumentConfigurationViewModels/MetasoundDocumentConfigurationViewModelBase.h"
 #include "MetasoundFrontendDocumentModifyDelegates.h"
 #include "View/MVVMView.h"
