Fix randomizer property view
#rnx

diff --git a/Engine/Plugins/Runtime/Metasound/Source/MetasoundEngine/Private/DocumentConfigurations/MetasoundRandomizerDocumentConfiguration.cpp b/Engine/Plugins/Runtime/Metasound/Source/MetasoundEngine/Private/DocumentConfigurations/MetasoundRandomizerDocumentConfiguration.cpp
--- a/Engine/Plugins/Runtime/Metasound/Source/MetasoundEngine/Private/DocumentConfigurations/MetasoundRandomizerDocumentConfiguration.cpp
+++ b/Engine/Plugins/Runtime/Metasound/Source/MetasoundEngine/Private/DocumentConfigurations/MetasoundRandomizerDocumentConfiguration.cpp
@@ -284,7 +284,7 @@ const FMetaSoundFrontendDocumentConfiguration::FEditorOptions& FMetaSoundRandomi
 	static const FEditorOptions RandomizerOptions
 	{
 		.AssetActionClasses = { USoundWave::StaticClass() },
-		.bConfigurationPropertiesVisible = false
+		.bConfigurationPropertiesVisible = true
 	};
 	return RandomizerOptions;
 }
@@ -301,19 +301,23 @@ void FMetaSoundRandomizerDocumentConfiguration::OnPropertyChanged(const FPropert
 {
 	if (InEvent.Property && InEvent.Property->GetName() == GET_MEMBER_NAME_CHECKED(FMetaSoundRandomizerDocumentConfiguration, Sounds))
 	{
-		RemoveInvalidSounds(OutBuilder, Sounds);
+		RemoveInvalidSounds(OutBuilder, false /* bRemoveNullEntries */, Sounds);
 	}
 
 	Super::OnPropertyChanged(InEvent, OutBuilder);
 }
 
-bool FMetaSoundRandomizerDocumentConfiguration::RemoveInvalidSounds(const FMetaSoundFrontendDocumentBuilder& ParentBuilder, TArray<TObjectPtr<USoundWave>>& OutSounds)
+bool FMetaSoundRandomizerDocumentConfiguration::RemoveInvalidSounds(const FMetaSoundFrontendDocumentBuilder& ParentBuilder, bool bRemoveNullEntries, TArray<TObjectPtr<USoundWave>>& OutSounds)
 {
-	const UObject& ThisObject = ParentBuilder.CastDocumentObjectChecked<UObject>();
-	return OutSounds.RemoveAll([&ThisObject](const TObjectPtr<USoundWave>& Wave)
+	return OutSounds.RemoveAll([&bRemoveNullEntries](const TObjectPtr<USoundWave>& Wave)
 	{
-		// Cannot reference this builder's object if it inherits from USoundWave (ex. MetaSoundSource) nor non-explicit USoundWave types.
-		return !Wave || Wave == &ThisObject || Wave->GetClass() != USoundWave::StaticClass();
+		if (bRemoveNullEntries && Wave == nullptr)
+		{
+			return true;
+		}
+
+		// Remove if not removing null entries set or class is not exactly USoundWave.
+		return Wave && Wave->GetClass() != USoundWave::StaticClass();
 	}) > 0;
 }
 #endif // WITH_EDITOR
