Fixup missing file from last CL

#rnx

diff --git a/Engine/Plugins/Runtime/Metasound/Source/MetasoundEditor/Private/DocumentConfigurationViewModels/RandomizerConfigurationViewModel.cpp b/Engine/Plugins/Runtime/Metasound/Source/MetasoundEditor/Private/DocumentConfigurationViewModels/RandomizerConfigurationViewModel.cpp
--- a/Engine/Plugins/Runtime/Metasound/Source/MetasoundEditor/Private/DocumentConfigurationViewModels/RandomizerConfigurationViewModel.cpp
+++ b/Engine/Plugins/Runtime/Metasound/Source/MetasoundEditor/Private/DocumentConfigurationViewModels/RandomizerConfigurationViewModel.cpp
@@ -147,7 +147,7 @@ void UMetaSoundRandomizerViewModel::SetSounds(TArray<TObjectPtr<USoundWave>> InS
 {
 	if (Builder)
 	{
-		FMetaSoundRandomizerDocumentConfiguration::RemoveInvalidSounds(Builder->GetConstBuilder(), InSounds);
+		FMetaSoundRandomizerDocumentConfiguration::RemoveInvalidSounds(Builder->GetConstBuilder(), true /* bRemoveNullEntries */, InSounds);
 		if (UE_MVVM_SET_PROPERTY_VALUE(Sounds, InSounds))
 		{
 			Builder->SetConfigurationProperties<FMetaSoundRandomizerDocumentConfiguration>([this](FMetaSoundRandomizerDocumentConfiguration& RandomConfig)
diff --git a/Engine/Plugins/Runtime/Metasound/Source/MetasoundEngine/Private/DocumentConfigurations/MetasoundRandomizerDocumentConfiguration.cpp b/Engine/Plugins/Runtime/Metasound/Source/MetasoundEngine/Private/DocumentConfigurations/MetasoundRandomizerDocumentConfiguration.cpp
--- a/Engine/Plugins/Runtime/Metasound/Source/MetasoundEngine/Private/DocumentConfigurations/MetasoundRandomizerDocumentConfiguration.cpp
+++ b/Engine/Plugins/Runtime/Metasound/Source/MetasoundEngine/Private/DocumentConfigurations/MetasoundRandomizerDocumentConfiguration.cpp
@@ -284,7 +284,7 @@ const FMetaSoundFrontendDocumentConfiguration::FEditorOptions& FMetaSoundRandomi
 	static const FEditorOptions RandomizerOptions
 	{
 		.AssetActionClasses = { USoundWave::StaticClass() },
-		.bConfigurationPropertiesVisible = false
+		.bConfigurationPropertiesVisible = true
 	};
 	return RandomizerOptions;
 }
@@ -301,19 +301,22 @@ void FMetaSoundRandomizerDocumentConfiguration::OnPropertyChanged(const FPropert
 {
 	if (InEvent.Property && InEvent.Property->GetName() == GET_MEMBER_NAME_CHECKED(FMetaSoundRandomizerDocumentConfiguration, Sounds))
 	{
-		RemoveInvalidSounds(OutBuilder, Sounds);
+		RemoveInvalidSounds(OutBuilder, false /* bRemoveNullEntries */, Sounds);
 	}
 
 	Super::OnPropertyChanged(InEvent, OutBuilder);
 }
 
-bool FMetaSoundRandomizerDocumentConfiguration::RemoveInvalidSounds(const FMetaSoundFrontendDocumentBuilder& ParentBuilder, TArray<TObjectPtr<USoundWave>>& OutSounds)
+bool FMetaSoundRandomizerDocumentConfiguration::RemoveInvalidSounds(const FMetaSoundFrontendDocumentBuilder& ParentBuilder, bool bRemoveNullEntries, TArray<TObjectPtr<USoundWave>>& OutSounds)
 {
-	const UObject& ThisObject = ParentBuilder.CastDocumentObjectChecked<UObject>();
-	return OutSounds.RemoveAll([&ThisObject](const TObjectPtr<USoundWave>& Wave)
+	return OutSounds.RemoveAll([&bRemoveNullEntries](const TObjectPtr<USoundWave>& Wave)
 	{
-		// Cannot reference this builder's object if it inherits from USoundWave (ex. MetaSoundSource) nor non-explicit USoundWave types.
-		return !Wave || Wave == &ThisObject || Wave->GetClass() != USoundWave::StaticClass();
+		if (bRemoveNullEntries && Wave == nullptr)
+		{
+			return true;
+		}
+
+		return Wave && Wave->GetClass() != USoundWave::StaticClass();
 	}) > 0;
 }
 #endif // WITH_EDITOR
diff --git a/Engine/Plugins/Runtime/Metasound/Source/MetasoundEngine/Public/DocumentConfigurations/MetasoundRandomizerDocumentConfiguration.h b/Engine/Plugins/Runtime/Metasound/Source/MetasoundEngine/Public/DocumentConfigurations/MetasoundRandomizerDocumentConfiguration.h
--- a/Engine/Plugins/Runtime/Metasound/Source/MetasoundEngine/Public/DocumentConfigurations/MetasoundRandomizerDocumentConfiguration.h
+++ b/Engine/Plugins/Runtime/Metasound/Source/MetasoundEngine/Public/DocumentConfigurations/MetasoundRandomizerDocumentConfiguration.h
@@ -36,9 +36,9 @@ struct FMetaSoundRandomizerDocumentConfiguration : public FMetaSoundFrontendDocu
 	UE_EXPERIMENTAL(5.8, "Document Configuration is in active development and subject to change")
 	UE_API virtual void OnPropertyChanged(const FPropertyChangedEvent& InEvent, FMetaSoundFrontendDocumentBuilder& OutBuilder) override;
 
-	// Helper function to remove sounds that cannot be referenced by randomizer (i.e. the given parent builder's MetaSound
-	// if it is a source or child class types of USoundWave). Returns true if sounds were removed, false if not.
-	UE_API static bool RemoveInvalidSounds(const FMetaSoundFrontendDocumentBuilder& ParentBuilder, TArray<TObjectPtr<USoundWave>>& OutSounds);
+	// Helper function to remove sounds that cannot be referenced by randomizer (i.e. if the sounds are not exactly of type USoundWave
+	// or, optionally, if the underlying pointer is null). Returns true if sounds were removed, false if not.
+	UE_API static bool RemoveInvalidSounds(const FMetaSoundFrontendDocumentBuilder& ParentBuilder, bool bRemoveNullEntries, TArray<TObjectPtr<USoundWave>>& OutSounds);
 #endif // WITH_EDITOR
 
 public:
