Make trigger repeat node high precision so it doesn't drift compared to Quartz / DAWs.

#Synced-CL 49958372

diff --git a/Engine/Plugins/Runtime/Metasound/Source/MetasoundStandardNodes/Private/MetasoundTriggerRepeatNode.cpp b/Engine/Plugins/Runtime/Metasound/Source/MetasoundStandardNodes/Private/MetasoundTriggerRepeatNode.cpp
--- a/Engine/Plugins/Runtime/Metasound/Source/MetasoundStandardNodes/Private/MetasoundTriggerRepeatNode.cpp
+++ b/Engine/Plugins/Runtime/Metasound/Source/MetasoundStandardNodes/Private/MetasoundTriggerRepeatNode.cpp
@@ -63,7 +63,8 @@ namespace Metasound
 			FInt32ReadRef NumRepeats;
 
 			FSampleCount BlockSize;
-			FSampleCounter SampleCounter;
+			double SampleRate;
+			double SampleCounter;
 
 			//The number of repeats the node has already performed
 			int CurrentRepeats = 0;
@@ -77,7 +78,8 @@ namespace Metasound
 		, Period(InPeriod)
 		, NumRepeats(InNumRepeats)
 		, BlockSize(InSettings.GetNumFramesPerBlock())
-		, SampleCounter(0, InSettings.GetSampleRate())
+		, SampleRate(InSettings.GetSampleRate())
+		, SampleCounter(0.0)
 	{
 	}
 
@@ -104,7 +106,7 @@ namespace Metasound
 			[this](int32 StartFrame, int32 EndFrame)
 			{
 				bEnabled = true;
-				SampleCounter.SetNumSamples(0);
+				SampleCounter = 0.0;
 
 				//Reset current number of repeats if we get another play command
 				CurrentRepeats = 0;
@@ -123,12 +125,12 @@ namespace Metasound
 
 		if (bEnabled)
 		{
-			FSampleCount PeriodInSamples = FSampleCounter::FromTime(*Period, SampleCounter.GetSampleRate()).GetNumSamples();
+			double PeriodInSamples = (FTime::ToMilliseconds(*Period) / 1000.0) * SampleRate;
+			double BlockSizeDouble = static_cast<double>(BlockSize);
 
 			// Time must march on, can't stay in the now forever.
-			PeriodInSamples = FMath::Max(static_cast<FSampleCount>(1), PeriodInSamples);
-
-			while ((SampleCounter - BlockSize).GetNumSamples() <= 0)
+			PeriodInSamples = FMath::Max(1.0, PeriodInSamples);
+			while (SampleCounter <= BlockSizeDouble)
 			{
 				//Whether or not the node has any repeats it should perform
 				bool bShouldRepeat = (*NumRepeats == 0) || (CurrentRepeats < *NumRepeats);
@@ -138,14 +140,14 @@ namespace Metasound
 					break;
 				}
 
-				const int32 StartOffset = static_cast<int32>(SampleCounter.GetNumSamples());
+				const int32 StartOffset = FMath::RoundToInt32(SampleCounter);
 				TriggerOut->TriggerFrame(StartOffset);
 				SampleCounter += PeriodInSamples;
 
 				++CurrentRepeats;
 			}
 
-			SampleCounter -= BlockSize;
+			SampleCounter -= BlockSizeDouble;
 		}
 	}
 
@@ -154,7 +156,7 @@ namespace Metasound
 		bEnabled = false;
 
 		TriggerOut->Reset();
-		SampleCounter.SetNumSamples(0);
+		SampleCounter = 0;
 		CurrentRepeats = 0;
 	}
 
