Improved page support for MetaSound builder BP API and editor builder listener
- Add BP paged input support, including input change type
- Clean up changing build page
- Exposed page change delegates and remove/reset graph input defaults

#Synced-CL 50003345

diff --git a/Engine/Plugins/Experimental/TechAudioTools/Source/TechAudioToolsMetaSound/Private/ViewModels/MetaSoundViewModel.cpp b/Engine/Plugins/Experimental/TechAudioTools/Source/TechAudioToolsMetaSound/Private/ViewModels/MetaSoundViewModel.cpp
--- a/Engine/Plugins/Experimental/TechAudioTools/Source/TechAudioToolsMetaSound/Private/ViewModels/MetaSoundViewModel.cpp
+++ b/Engine/Plugins/Experimental/TechAudioTools/Source/TechAudioToolsMetaSound/Private/ViewModels/MetaSoundViewModel.cpp
@@ -282,11 +282,11 @@ void UMetaSoundViewModel::OnInputDataTypeChanged(const FName VertexName, const F
 	}
 }
 
-void UMetaSoundViewModel::OnInputDefaultChanged(const FName VertexName, FMetasoundFrontendLiteral LiteralValue, const FName PageName)
+void UMetaSoundViewModel::OnInputDefaultChanged(const FName VertexName, FMetasoundFrontendLiteral LiteralValue, const FName PageName, EMetaSoundLiteralChangeType ChangeType)
 {
 	if (UMetaSoundInputViewModel* InputViewModel = InputViewModels.FindRef(VertexName))
 	{
-		InputViewModel->OnInputDefaultChanged(LiteralValue, PageName);
+		InputViewModel->OnInputDefaultChanged(LiteralValue, PageName, ChangeType);
 	}
 }
 
@@ -540,7 +540,7 @@ void UMetaSoundInputViewModel::OnInputDataTypeChanged(const FName& NewDataType)
 	}
 }
 
-void UMetaSoundInputViewModel::OnInputDefaultChanged(const FMetasoundFrontendLiteral& LiteralValue, const FName& PageName)
+void UMetaSoundInputViewModel::OnInputDefaultChanged(const FMetasoundFrontendLiteral& LiteralValue, const FName& PageName, EMetaSoundLiteralChangeType ChangeType)
 {
 	UE_MVVM_SET_PROPERTY_VALUE(Literal, LiteralValue);
 }
diff --git a/Engine/Plugins/Experimental/TechAudioTools/Source/TechAudioToolsMetaSound/Public/ViewModels/MetaSoundViewModel.h b/Engine/Plugins/Experimental/TechAudioTools/Source/TechAudioToolsMetaSound/Public/ViewModels/MetaSoundViewModel.h
--- a/Engine/Plugins/Experimental/TechAudioTools/Source/TechAudioToolsMetaSound/Public/ViewModels/MetaSoundViewModel.h
+++ b/Engine/Plugins/Experimental/TechAudioTools/Source/TechAudioToolsMetaSound/Public/ViewModels/MetaSoundViewModel.h
@@ -89,7 +89,7 @@ public:
 	UFUNCTION()
 	UE_API void OnInputDataTypeChanged(FName VertexName, FName DataType);
 	UFUNCTION()
-	UE_API void OnInputDefaultChanged(FName VertexName, FMetasoundFrontendLiteral LiteralValue, FName PageName);
+	UE_API void OnInputDefaultChanged(FName VertexName, FMetasoundFrontendLiteral LiteralValue, FName PageName, EMetaSoundLiteralChangeType ChangeType = EMetaSoundLiteralChangeType::Invalid);
 	UFUNCTION()
 	UE_API void OnInputInheritsDefaultChanged(FName VertexName, bool bInheritsDefault);
 	UFUNCTION()
@@ -186,7 +186,7 @@ public:
 
 	UE_API void OnInputNameChanged(const FName& NewName);
 	UE_API void OnInputDataTypeChanged(const FName& NewDataType);
-	UE_API void OnInputDefaultChanged(const FMetasoundFrontendLiteral& LiteralValue, const FName& PageName);
+	UE_API void OnInputDefaultChanged(const FMetasoundFrontendLiteral& LiteralValue, const FName& PageName, EMetaSoundLiteralChangeType ChangeType = EMetaSoundLiteralChangeType::Invalid);
 	UE_API void OnInputInheritsDefaultChanged(const bool bInheritsDefault);
 	UE_API void OnInputIsConstructorPinChanged(const bool bNewIsConstructorPin);
 };
diff --git a/Engine/Plugins/Experimental/TechAudioTools/Source/TechAudioToolsMetaSoundEditor/Private/ViewModels/MetaSoundEditorViewModel.cpp b/Engine/Plugins/Experimental/TechAudioTools/Source/TechAudioToolsMetaSoundEditor/Private/ViewModels/MetaSoundEditorViewModel.cpp
--- a/Engine/Plugins/Experimental/TechAudioTools/Source/TechAudioToolsMetaSoundEditor/Private/ViewModels/MetaSoundEditorViewModel.cpp
+++ b/Engine/Plugins/Experimental/TechAudioTools/Source/TechAudioToolsMetaSoundEditor/Private/ViewModels/MetaSoundEditorViewModel.cpp
@@ -198,6 +198,14 @@ void UMetaSoundEditorViewModel::SetOutputIsAdvancedDisplay(const FName& OutputNa
 	}
 }
 
+void UMetaSoundEditorViewModel::OnGraphInputDefaultChanged(const FGraphInputDefaultChangedArgs& Args)
+{
+	if (UMetaSoundInputViewModel* InputViewModel = InputViewModels.FindRef(Args.InputName))
+	{
+		InputViewModel->OnInputDefaultChanged(Args.Literal, Args.PageName, Args.ChangeType);
+	}
+}
+
 void UMetaSoundEditorViewModel::OnInputDisplayNameChanged(const FName InputName, const FText InDisplayName)
 {
 	if (UMetaSoundInputEditorViewModel* InputEditorViewModel = Cast<UMetaSoundInputEditorViewModel>(*InputViewModels.Find(InputName)))
@@ -335,7 +343,7 @@ void UMetaSoundEditorViewModel::BindDelegates()
 		BuilderListener->OnRemovingGraphInputDelegate.AddDynamic(this, &UMetaSoundViewModel::OnInputRemoved);
 		BuilderListener->OnGraphInputNameChangedDelegate.AddDynamic(this, &UMetaSoundViewModel::OnInputNameChanged);
 		BuilderListener->OnGraphInputDataTypeChangedDelegate.AddDynamic(this, &UMetaSoundViewModel::OnInputDataTypeChanged);
-		BuilderListener->OnGraphInputDefaultChangedDelegate.AddDynamic(this, &UMetaSoundViewModel::OnInputDefaultChanged);
+		BuilderListener->OnGraphInputDefaultChangedWithArgsDelegate.AddDynamic(this, &ThisClass::OnGraphInputDefaultChanged);
 		BuilderListener->OnGraphInputInheritsDefaultChangedDelegate.AddDynamic(this, &UMetaSoundViewModel::OnInputInheritsDefaultChanged);
 		BuilderListener->OnGraphInputIsConstructorPinChangedDelegate.AddDynamic(this, &UMetaSoundViewModel::OnInputIsConstructorPinChanged);
 
diff --git a/Engine/Plugins/Experimental/TechAudioTools/Source/TechAudioToolsMetaSoundEditor/Public/ViewModels/MetaSoundEditorViewModel.h b/Engine/Plugins/Experimental/TechAudioTools/Source/TechAudioToolsMetaSoundEditor/Public/ViewModels/MetaSoundEditorViewModel.h
--- a/Engine/Plugins/Experimental/TechAudioTools/Source/TechAudioToolsMetaSoundEditor/Public/ViewModels/MetaSoundEditorViewModel.h
+++ b/Engine/Plugins/Experimental/TechAudioTools/Source/TechAudioToolsMetaSoundEditor/Public/ViewModels/MetaSoundEditorViewModel.h
@@ -8,6 +8,7 @@
 
 #define UE_API TECHAUDIOTOOLSMETASOUNDEDITOR_API
 
+struct FGraphInputDefaultChangedArgs;
 class UMetaSoundEditorBuilderListener;
 
 /**
@@ -75,6 +76,9 @@ public:
 	UE_API void SetOutputSortOrderIndex(const FName& OutputName, int32 InSortOrderIndex) const;
 	UE_API void SetOutputIsAdvancedDisplay(const FName& OutputName, bool bInIsAdvancedDisplay) const;
 
+	UFUNCTION()
+	void OnGraphInputDefaultChanged(const FGraphInputDefaultChangedArgs& Args);
+	
 	UFUNCTION()
 	void OnDisplayNameChanged(const FText NewDisplayName) { UE_MVVM_SET_PROPERTY_VALUE(DisplayName, NewDisplayName); }
 	UFUNCTION()
diff --git a/Engine/Plugins/Runtime/Metasound/Source/MetasoundEditor/Private/MetasoundDefaultLiteralCustomization.cpp b/Engine/Plugins/Runtime/Metasound/Source/MetasoundEditor/Private/MetasoundDefaultLiteralCustomization.cpp
--- a/Engine/Plugins/Runtime/Metasound/Source/MetasoundEditor/Private/MetasoundDefaultLiteralCustomization.cpp
+++ b/Engine/Plugins/Runtime/Metasound/Source/MetasoundEditor/Private/MetasoundDefaultLiteralCustomization.cpp
@@ -293,7 +293,7 @@ namespace Metasound::Editor
 				LiteralPtr->RemoveDefault(PageID);
 
 				constexpr bool bPostTransaction = false;
-				Member->UpdateFrontendDefaultLiteral(bPostTransaction);
+				Member->UpdateFrontendDefaultLiteral(bPostTransaction, /*InPageID=*/nullptr, EMetaSoundLiteralChangeType::Remove);
 
 				FMetasoundAssetBase& MetasoundAsset = Editor::FGraphBuilder::GetOutermostMetaSoundChecked(*LiteralPtr);
 				MetasoundAsset.GetModifyContext().AddMemberIDsModified({ Member->GetMemberID() });
diff --git a/Engine/Plugins/Runtime/Metasound/Source/MetasoundEditor/Private/MetasoundEditor.cpp b/Engine/Plugins/Runtime/Metasound/Source/MetasoundEditor/Private/MetasoundEditor.cpp
--- a/Engine/Plugins/Runtime/Metasound/Source/MetasoundEditor/Private/MetasoundEditor.cpp
+++ b/Engine/Plugins/Runtime/Metasound/Source/MetasoundEditor/Private/MetasoundEditor.cpp
@@ -5949,7 +5949,7 @@ PRAGMA_ENABLE_DEPRECATION_WARNINGS
 			}
 		}
 
-		void FEditor::FDocumentListener::OnInputDefaultChanged(int32 Index)
+		void FEditor::FDocumentListener::OnInputDefaultChanged(const Frontend::FDocumentArrayPagedInputArgs& Args)
 		{
 			if (TSharedPtr<FEditor> ParentPtr = Parent.Pin())
 			{
@@ -5964,9 +5964,9 @@ PRAGMA_ENABLE_DEPRECATION_WARNINGS
 				const FMetasoundFrontendGraph& Graph = GraphClass.GetConstDefaultGraph();
 
 				const TArray<FMetasoundFrontendClassInput>& GraphInputs = GraphClass.GetDefaultInterface().Inputs;
-				if (GraphInputs.IsValidIndex(Index))
+				if (GraphInputs.IsValidIndex(Args.Index))
 				{
-					const FMetasoundFrontendClassInput& GraphInput = GraphClass.GetDefaultInterface().Inputs[Index];
+					const FMetasoundFrontendClassInput& GraphInput = GraphClass.GetDefaultInterface().Inputs[Args.Index];
 					const FName InputName = GraphInput.Name;
 					UMetasoundEditorGraphInput* EditorGraphInput = MetaSoundGraph->FindInput(InputName);
 
diff --git a/Engine/Plugins/Runtime/Metasound/Source/MetasoundEditor/Private/MetasoundEditor.h b/Engine/Plugins/Runtime/Metasound/Source/MetasoundEditor/Private/MetasoundEditor.h
--- a/Engine/Plugins/Runtime/Metasound/Source/MetasoundEditor/Private/MetasoundEditor.h
+++ b/Engine/Plugins/Runtime/Metasound/Source/MetasoundEditor/Private/MetasoundEditor.h
@@ -693,7 +693,7 @@ namespace Metasound::Editor
 			virtual void OnBuilderReloaded(Frontend::FDocumentModifyDelegates& OutDelegates) override;
 			void OnDocumentConfigurationChanged(const Frontend::FDocumentConfigChangedArgs& Args);
 			void OnPageSet(const Frontend::FDocumentMutatePageArgs& Args);
-			void OnInputDefaultChanged(int32 Index);
+			void OnInputDefaultChanged(const Frontend::FDocumentArrayPagedInputArgs& Args);
 			void OnMemberMetadataSet(FGuid MemberID);
 		};
 
diff --git a/Engine/Plugins/Runtime/Metasound/Source/MetasoundEditor/Private/MetasoundEditorBuilderListener.cpp b/Engine/Plugins/Runtime/Metasound/Source/MetasoundEditor/Private/MetasoundEditorBuilderListener.cpp
--- a/Engine/Plugins/Runtime/Metasound/Source/MetasoundEditor/Private/MetasoundEditorBuilderListener.cpp
+++ b/Engine/Plugins/Runtime/Metasound/Source/MetasoundEditor/Private/MetasoundEditorBuilderListener.cpp
@@ -2,6 +2,8 @@
 
 #include "MetasoundEditorBuilderListener.h"
 
+#include "DocumentConfigurations/MetasoundFrontendPresetDocumentConfiguration.h"
+
 #include UE_INLINE_GENERATED_CPP_BY_NAME(MetasoundEditorBuilderListener)
 
 void UMetaSoundEditorBuilderListener::Init(const TWeakObjectPtr<UMetaSoundBuilderBase> InBuilder)
@@ -86,48 +88,94 @@ void UMetaSoundEditorBuilderListener::OnDocumentIsDeprecatedChanged()
 	}
 }
 
+void UMetaSoundEditorBuilderListener::OnPageAdded(const Metasound::Frontend::FDocumentMutatePageArgs& PageArgs)
+{
+	if (Builder.IsValid())
+	{
+		const FName PageName = UMetaSoundSettings::GetPageName(PageArgs.PageID);
+		OnPageAddedDelegate.Broadcast(PageName);
+	}
+}
+
+void UMetaSoundEditorBuilderListener::OnRemovingPage(const Metasound::Frontend::FDocumentMutatePageArgs& PageArgs)
+{
+	if (Builder.IsValid())
+	{
+		const FName PageName = UMetaSoundSettings::GetPageName(PageArgs.PageID);
+		OnRemovingPageDelegate.Broadcast(PageName);
+	}
+}
+
+void UMetaSoundEditorBuilderListener::OnBuildPageChanged(const Metasound::Frontend::FDocumentMutatePageArgs& PageArgs)
+{
+	if (Builder.IsValid())
+	{
+		const FName PageName = UMetaSoundSettings::GetPageName(PageArgs.PageID);
+		OnBuildPageChangedDelegate.Broadcast(PageName);
+	}
+}
+
 void UMetaSoundEditorBuilderListener::OnGraphInputAdded(int32 Index)
 {
 	if (Builder.IsValid())
 	{
 		const FMetaSoundFrontendDocumentBuilder& ConstBuilder = Builder->GetConstBuilder();
 		const FMetasoundFrontendGraphClass& GraphClass = ConstBuilder.GetConstDocumentChecked().RootGraph;
-		const FMetasoundFrontendClassInput& GraphInput = GraphClass.GetDefaultInterface().Inputs[Index];
-
-		OnGraphInputAddedDelegate.Broadcast(GraphInput.Name, GraphInput.TypeName);
+		const TArray<FMetasoundFrontendClassInput>& Inputs = GraphClass.GetDefaultInterface().Inputs;
+		if (ensure(Inputs.IsValidIndex(Index)))
+		{
+			const FMetasoundFrontendClassInput& GraphInput = Inputs[Index];
+			OnGraphInputAddedDelegate.Broadcast(GraphInput.Name, GraphInput.TypeName);
+		}
 	}
 }
 
 void UMetaSoundEditorBuilderListener::OnGraphInputDefaultChanged(int32 Index)
+{
+}
+
+void UMetaSoundEditorBuilderListener::OnGraphInputDefaultChanged(const Metasound::Frontend::FDocumentArrayPagedInputArgs& Args)
 {
 	if (Builder.IsValid())
 	{
 		const FMetaSoundFrontendDocumentBuilder& ConstBuilder = Builder->GetConstBuilder();
 		const FMetasoundFrontendGraphClass& GraphClass = ConstBuilder.GetConstDocumentChecked().RootGraph;
-
-		// todo: support paged graph once frontend delegates support paged literals
-		const FMetasoundFrontendGraph& Graph = GraphClass.GetConstDefaultGraph();
-		FName PageName = Metasound::Frontend::DefaultPageName;
-
-		const FMetasoundFrontendClassInput& GraphInput = GraphClass.GetDefaultInterface().Inputs[Index];
-		const FMetasoundFrontendLiteral* DefaultLiteral = GraphInput.FindConstDefault(Metasound::Frontend::DefaultPageID);
-		if (ensure(DefaultLiteral))
+		const TArray<FMetasoundFrontendClassInput>& Inputs = GraphClass.GetDefaultInterface().Inputs;
+		if (ensure(Inputs.IsValidIndex(Args.Index)))
 		{
-			OnGraphInputDefaultChangedDelegate.Broadcast(GraphInput.Name, *DefaultLiteral, PageName);
+			const FMetasoundFrontendClassInput& GraphInput = Inputs[Args.Index];
+			FMetasoundFrontendLiteral DefaultLiteral;
+			// If default literal was removed, then it will be null
+			if (const FMetasoundFrontendLiteral* CurrentLiteral = GraphInput.FindConstDefault(Args.PageID))
+			{
+				DefaultLiteral = *CurrentLiteral;
+			}
+			const FName PageName = UMetaSoundSettings::GetPageName(Args.PageID);
+			OnGraphInputDefaultChangedWithArgsDelegate.Broadcast(FGraphInputDefaultChangedArgs(
+					GraphInput.Name, PageName, DefaultLiteral, Args.ChangeType));
 		}
 	}
 }
 
 void UMetaSoundEditorBuilderListener::OnGraphInputInheritsDefaultChanged(int32 Index)
 {
+	using namespace Metasound::Frontend;
+	
 	if (Builder.IsValid())
 	{
 		const FMetaSoundFrontendDocumentBuilder& ConstBuilder = Builder->GetConstBuilder();
 		const FMetasoundFrontendGraphClass& GraphClass = ConstBuilder.GetConstDocumentChecked().RootGraph;
-		const FMetasoundFrontendClassInput& GraphInput = GraphClass.GetDefaultInterface().Inputs[Index];
-		const bool bInheritsDefault = GraphClass.PresetOptions.InputsInheritingDefault.Contains(GraphInput.Name);
-
-		OnGraphInputInheritsDefaultChangedDelegate.Broadcast(GraphInput.Name, bInheritsDefault);
+		const TArray<FMetasoundFrontendClassInput>& Inputs = GraphClass.GetDefaultInterface().Inputs;
+		if (ensure(Inputs.IsValidIndex(Index)))
+		{
+			const FMetasoundFrontendClassInput& GraphInput = Inputs[Index];
+			bool bInheritsDefault = false;
+			if (const FPresetVertexMetadata* VertexMetadata = ConstBuilder.FindConstVertexMetadata<FPresetVertexMetadata>(GraphInput.NodeID))
+			{
+				bInheritsDefault = !VertexMetadata->bOverrideInheritedDefault;
+			}
+			OnGraphInputInheritsDefaultChangedDelegate.Broadcast(GraphInput.Name, bInheritsDefault);
+		}
 	}
 }
 
@@ -137,9 +185,12 @@ void UMetaSoundEditorBuilderListener::OnRemovingGraphInput(int32 Index)
 	{
 		const FMetaSoundFrontendDocumentBuilder& ConstBuilder = Builder->GetConstBuilder();
 		const FMetasoundFrontendGraphClass& GraphClass = ConstBuilder.GetConstDocumentChecked().RootGraph;
-		const FMetasoundFrontendClassInput& GraphInput = GraphClass.GetDefaultInterface().Inputs[Index];
-
-		OnRemovingGraphInputDelegate.Broadcast(GraphInput.Name, GraphInput.TypeName);
+		const TArray<FMetasoundFrontendClassInput>& Inputs = GraphClass.GetDefaultInterface().Inputs;
+		if (ensure(Inputs.IsValidIndex(Index)))
+		{
+			const FMetasoundFrontendClassInput& GraphInput = Inputs[Index];
+			OnRemovingGraphInputDelegate.Broadcast(GraphInput.Name, GraphInput.TypeName);
+		}
 	}
 }
 
@@ -173,9 +224,12 @@ void UMetaSoundEditorBuilderListener::OnGraphInputDataTypeChanged(int32 Index)
 	{
 		const FMetaSoundFrontendDocumentBuilder& ConstBuilder = Builder->GetConstBuilder();
 		const FMetasoundFrontendGraphClass& GraphClass = ConstBuilder.GetConstDocumentChecked().RootGraph;
-		const FMetasoundFrontendClassInput& GraphInput = GraphClass.GetDefaultInterface().Inputs[Index];
-
-		OnGraphInputDataTypeChangedDelegate.Broadcast(GraphInput.Name, GraphInput.TypeName);
+		const TArray<FMetasoundFrontendClassInput>& Inputs = GraphClass.GetDefaultInterface().Inputs;
+		if (ensure(Inputs.IsValidIndex(Index)))
+		{
+			const FMetasoundFrontendClassInput& GraphInput = Inputs[Index];
+			OnGraphInputDataTypeChangedDelegate.Broadcast(GraphInput.Name, GraphInput.TypeName);
+		}
 	}
 }
 
@@ -185,9 +239,12 @@ void UMetaSoundEditorBuilderListener::OnGraphInputDisplayNameChanged(int32 Index
 	{
 		const FMetaSoundFrontendDocumentBuilder& ConstBuilder = Builder->GetConstBuilder();
 		const FMetasoundFrontendGraphClass& GraphClass = ConstBuilder.GetConstDocumentChecked().RootGraph;
-		const FMetasoundFrontendClassInput& GraphInput = GraphClass.GetDefaultInterface().Inputs[Index];
-
-		OnGraphInputDisplayNameChangedDelegate.Broadcast(GraphInput.Name, GraphInput.Metadata.GetDisplayName());
+		const TArray<FMetasoundFrontendClassInput>& Inputs = GraphClass.GetDefaultInterface().Inputs;
+		if (ensure(Inputs.IsValidIndex(Index)))
+		{
+			const FMetasoundFrontendClassInput& GraphInput = Inputs[Index];
+			OnGraphInputDisplayNameChangedDelegate.Broadcast(GraphInput.Name, GraphInput.Metadata.GetDisplayName());
+		}
 	}
 }
 
@@ -197,9 +254,12 @@ void UMetaSoundEditorBuilderListener::OnGraphInputDescriptionChanged(int32 Index
 	{
 		const FMetaSoundFrontendDocumentBuilder& ConstBuilder = Builder->GetConstBuilder();
 		const FMetasoundFrontendGraphClass& GraphClass = ConstBuilder.GetConstDocumentChecked().RootGraph;
-		const FMetasoundFrontendClassInput& GraphInput = GraphClass.GetDefaultInterface().Inputs[Index];
-
-		OnGraphInputDescriptionChangedDelegate.Broadcast(GraphInput.Name, GraphInput.Metadata.GetDescription());
+		const TArray<FMetasoundFrontendClassInput>& Inputs = GraphClass.GetDefaultInterface().Inputs;
+		if (ensure(Inputs.IsValidIndex(Index)))
+		{
+			const FMetasoundFrontendClassInput& GraphInput = Inputs[Index];
+			OnGraphInputDescriptionChangedDelegate.Broadcast(GraphInput.Name, GraphInput.Metadata.GetDescription());
+		}
 	}
 }
 
@@ -209,9 +269,12 @@ void UMetaSoundEditorBuilderListener::OnGraphInputSortOrderIndexChanged(int32 In
 	{
 		const FMetaSoundFrontendDocumentBuilder& ConstBuilder = Builder->GetConstBuilder();
 		const FMetasoundFrontendGraphClass& GraphClass = ConstBuilder.GetConstDocumentChecked().RootGraph;
-		const FMetasoundFrontendClassInput& GraphInput = GraphClass.GetDefaultInterface().Inputs[Index];
-
-		OnGraphInputSortOrderIndexChangedDelegate.Broadcast(GraphInput.Name, GraphInput.Metadata.SortOrderIndex);
+		const TArray<FMetasoundFrontendClassInput>& Inputs = GraphClass.GetDefaultInterface().Inputs;
+		if (ensure(Inputs.IsValidIndex(Index)))
+		{
+			const FMetasoundFrontendClassInput& GraphInput = Inputs[Index];
+			OnGraphInputSortOrderIndexChangedDelegate.Broadcast(GraphInput.Name, GraphInput.Metadata.SortOrderIndex);
+		}
 	}
 }
 
@@ -221,9 +284,12 @@ void UMetaSoundEditorBuilderListener::OnGraphInputIsConstructorPinChanged(int32
 	{
 		const FMetaSoundFrontendDocumentBuilder& ConstBuilder = Builder->GetConstBuilder();
 		const FMetasoundFrontendGraphClass& GraphClass = ConstBuilder.GetConstDocumentChecked().RootGraph;
-		const FMetasoundFrontendClassInput& GraphInput = GraphClass.GetDefaultInterface().Inputs[Index];
-
-		OnGraphInputIsConstructorPinChangedDelegate.Broadcast(GraphInput.Name, GraphInput.AccessType == EMetasoundFrontendVertexAccessType::Value);
+		const TArray<FMetasoundFrontendClassInput>& Inputs = GraphClass.GetDefaultInterface().Inputs;
+		if (ensure(Inputs.IsValidIndex(Index)))
+		{
+			const FMetasoundFrontendClassInput& GraphInput = Inputs[Index];
+			OnGraphInputIsConstructorPinChangedDelegate.Broadcast(GraphInput.Name, GraphInput.AccessType == EMetasoundFrontendVertexAccessType::Value);
+		}
 	}
 }
 
@@ -233,9 +299,12 @@ void UMetaSoundEditorBuilderListener::OnGraphInputIsAdvancedDisplayChanged(int32
 	{
 		const FMetaSoundFrontendDocumentBuilder& ConstBuilder = Builder->GetConstBuilder();
 		const FMetasoundFrontendGraphClass& GraphClass = ConstBuilder.GetConstDocumentChecked().RootGraph;
-		const FMetasoundFrontendClassInput& GraphInput = GraphClass.GetDefaultInterface().Inputs[Index];
-
-		OnGraphInputIsAdvancedDisplayChangedDelegate.Broadcast(GraphInput.Name, GraphInput.Metadata.bIsAdvancedDisplay);
+		const TArray<FMetasoundFrontendClassInput>& Inputs = GraphClass.GetDefaultInterface().Inputs;
+		if (ensure(Inputs.IsValidIndex(Index)))
+		{
+			const FMetasoundFrontendClassInput& GraphInput = Inputs[Index];
+			OnGraphInputIsAdvancedDisplayChangedDelegate.Broadcast(GraphInput.Name, GraphInput.Metadata.bIsAdvancedDisplay);
+		}
 	}
 }
 
@@ -338,7 +407,13 @@ void UMetaSoundEditorBuilderListener::RemoveAllDelegates()
 	OnGraphInputSortOrderIndexChangedDelegate.RemoveAll(this);
 	OnGraphInputIsAdvancedDisplayChangedDelegate.RemoveAll(this);
 	OnGraphInputIsConstructorPinChangedDelegate.RemoveAll(this);
+	OnGraphInputDefaultChangedWithArgsDelegate.RemoveAll(this);
+PRAGMA_DISABLE_DEPRECATION_WARNINGS
 	OnGraphInputDefaultChangedDelegate.RemoveAll(this);
+PRAGMA_ENABLE_DEPRECATION_WARNINGS
+	OnPageAddedDelegate.RemoveAll(this);
+	OnRemovingPageDelegate.RemoveAll(this);
+	OnBuildPageChangedDelegate.RemoveAll(this);
 
 	OnGraphOutputAddedDelegate.RemoveAll(this);
 	OnRemovingGraphOutputDelegate.RemoveAll(this);
@@ -362,6 +437,10 @@ void UMetaSoundEditorBuilderListener::RemoveAllDelegates()
 		BuilderDelegates.OnDocumentMetadataChanged.Remove(OnDocumentKeywordsChangedHandle);
 		BuilderDelegates.OnDocumentMetadataChanged.Remove(OnDocumentCategoryHierarchyChangedHandle);
 		BuilderDelegates.OnDocumentMetadataChanged.Remove(OnDocumentIsDeprecatedChangedHandle);
+		
+		BuilderDelegates.PageDelegates.OnPageAdded.Remove(OnPageAddedHandle);
+		BuilderDelegates.PageDelegates.OnRemovingPage.Remove(OnRemovingPageHandle);
+		BuilderDelegates.PageDelegates.OnPageSet.Remove(OnBuildPageChangedHandle);
 
 		InterfaceDelegates.OnInputAdded.Remove(OnInputAddedHandle);
 		InterfaceDelegates.OnRemovingInput.Remove(OnRemovingInputHandle);
@@ -400,6 +479,10 @@ void UMetaSoundEditorBuilderListener::FEditorBuilderListener::OnBuilderReloaded(
 		Parent->OnDocumentCategoryHierarchyChangedHandle = OutDelegates.OnDocumentMetadataChanged.AddUObject(Parent, &UMetaSoundEditorBuilderListener::OnDocumentCategoryHierarchyChanged);
 		Parent->OnDocumentIsDeprecatedChangedHandle = OutDelegates.OnDocumentMetadataChanged.AddUObject(Parent, &UMetaSoundEditorBuilderListener::OnDocumentIsDeprecatedChanged);
 
+		Parent->OnPageAddedHandle = OutDelegates.PageDelegates.OnPageAdded.AddUObject(Parent, &UMetaSoundEditorBuilderListener::OnPageAdded);
+		Parent->OnRemovingPageHandle = OutDelegates.PageDelegates.OnRemovingPage.AddUObject(Parent, &UMetaSoundEditorBuilderListener::OnRemovingPage);
+		Parent->OnBuildPageChangedHandle = OutDelegates.PageDelegates.OnPageSet.AddUObject(Parent, &UMetaSoundEditorBuilderListener::OnBuildPageChanged);
+
 		Parent->OnInputAddedHandle = InterfaceDelegates.OnInputAdded.AddUObject(Parent, &UMetaSoundEditorBuilderListener::OnGraphInputAdded);
 		Parent->OnRemovingInputHandle = InterfaceDelegates.OnRemovingInput.AddUObject(Parent, &UMetaSoundEditorBuilderListener::OnRemovingGraphInput);
 		Parent->OnInputNameChangedHandle = InterfaceDelegates.OnInputNameChanged.AddUObject(Parent, &UMetaSoundEditorBuilderListener::OnGraphInputNameChanged);
diff --git a/Engine/Plugins/Runtime/Metasound/Source/MetasoundEditor/Private/MetasoundEditorGraph.cpp b/Engine/Plugins/Runtime/Metasound/Source/MetasoundEditor/Private/MetasoundEditorGraph.cpp
--- a/Engine/Plugins/Runtime/Metasound/Source/MetasoundEditor/Private/MetasoundEditorGraph.cpp
+++ b/Engine/Plugins/Runtime/Metasound/Source/MetasoundEditor/Private/MetasoundEditorGraph.cpp
@@ -64,7 +64,12 @@ namespace Metasound::Editor
 		void OnLiteralChanged(UMetasoundEditorGraphMember& InMember, const FGuid* InPageID, EPropertyChangeType::Type InChangeType)
 		{
 			constexpr bool bPostTransaction = false;
-			InMember.UpdateFrontendDefaultLiteral(bPostTransaction, InPageID);
+			EMetaSoundLiteralChangeType LiteralChangeType = EMetaSoundLiteralChangeType::Set;
+			if (InChangeType & EPropertyChangeType::ResetToDefault)
+			{
+				LiteralChangeType = EMetaSoundLiteralChangeType::Reset;
+			}
+			InMember.UpdateFrontendDefaultLiteral(bPostTransaction, InPageID, LiteralChangeType);
 
 			const bool bCommitChange = InChangeType != EPropertyChangeType::Interactive;
 			if (bCommitChange)
@@ -1098,7 +1103,7 @@ void UMetasoundEditorGraphInput::CacheBreadcrumb()
 	}
 }
 
-void UMetasoundEditorGraphInput::UpdateFrontendDefaultLiteral(bool bPostTransaction, const FGuid* InPageID)
+void UMetasoundEditorGraphInput::UpdateFrontendDefaultLiteral(bool bPostTransaction, const FGuid* InPageID, EMetaSoundLiteralChangeType InChangeType)
 {
 	using namespace Metasound::Editor;
 
@@ -1119,7 +1124,7 @@ void UMetasoundEditorGraphInput::UpdateFrontendDefaultLiteral(bool bPostTransact
 				const FMetasoundFrontendLiteral* ExistingDefault = Builder.GetGraphInputDefault(MemberName, InPageID);
 				if (!ExistingDefault || *ExistingDefault != Default)
 				{
-					Builder.SetGraphInputDefault(MemberName, Default, InPageID);
+					Builder.SetGraphInputDefault(MemberName, Default, InPageID, InChangeType);
 				}
 			}
 		}
@@ -1140,7 +1145,7 @@ void UMetasoundEditorGraphInput::UpdateFrontendDefaultLiteral(bool bPostTransact
 			const FMetasoundFrontendClassInput* GraphInput = Builder.FindGraphInput(MemberName);
 			if (!GraphInput || bDefaultsModified || NewDefaults.Num() != GraphInput->GetDefaults().Num())
 			{
-				Builder.SetGraphInputDefaults(MemberName, MoveTemp(NewDefaults));
+				Builder.SetGraphInputDefaults(MemberName, MoveTemp(NewDefaults), InChangeType);
 			}
 		}
 	}
@@ -1400,7 +1405,7 @@ void UMetasoundEditorGraphOutput::SetVertexAccessType(EMetasoundFrontendVertexAc
 	}
 }
 
-void UMetasoundEditorGraphOutput::UpdateFrontendDefaultLiteral(bool bPostTransaction, const FGuid* InPageID)
+void UMetasoundEditorGraphOutput::UpdateFrontendDefaultLiteral(bool bPostTransaction, const FGuid* InPageID, EMetaSoundLiteralChangeType InChangeType)
 {
 	using namespace Metasound;
 	using namespace Metasound::Editor;
@@ -2057,7 +2062,7 @@ bool UMetasoundEditorGraphVariable::Synchronize()
 	return bModified;
 }
 
-void UMetasoundEditorGraphVariable::UpdateFrontendDefaultLiteral(bool bPostTransaction, const FGuid* InPageID)
+void UMetasoundEditorGraphVariable::UpdateFrontendDefaultLiteral(bool bPostTransaction, const FGuid* InPageID, EMetaSoundLiteralChangeType InChangeType)
 {
 	using namespace Metasound;
 	using namespace Metasound::Editor;
diff --git a/Engine/Plugins/Runtime/Metasound/Source/MetasoundEditor/Private/MetasoundEditorSubsystem.cpp b/Engine/Plugins/Runtime/Metasound/Source/MetasoundEditor/Private/MetasoundEditorSubsystem.cpp
--- a/Engine/Plugins/Runtime/Metasound/Source/MetasoundEditor/Private/MetasoundEditorSubsystem.cpp
+++ b/Engine/Plugins/Runtime/Metasound/Source/MetasoundEditor/Private/MetasoundEditorSubsystem.cpp
@@ -257,6 +257,17 @@ UMetaSoundFrontendMemberMetadata* UMetaSoundEditorSubsystem::FindOrCreateGraphIn
 	return nullptr;
 }
 
+FName UMetaSoundEditorSubsystem::GetBuildPage(UMetaSoundBuilderBase* Builder) const
+{
+	if (Builder)
+	{
+		const FMetaSoundFrontendDocumentBuilder& DocBuilder = Builder->GetConstBuilder();
+		const FGuid& BuildPageID = DocBuilder.GetBuildPageID();
+		return UMetaSoundSettings::GetPageName(BuildPageID);
+	}
+	return FName();
+}
+
 TSubclassOf<UMetasoundEditorGraphMemberDefaultLiteral> UMetaSoundEditorSubsystem::GetLiteralClassForType(FName TypeName) const
 {
 	using namespace Metasound::Editor;
@@ -440,6 +451,36 @@ UMetaSoundEditorBuilderListener* UMetaSoundEditorSubsystem::AddBuilderDelegateLi
 	return nullptr;
 }
 
+void UMetaSoundEditorSubsystem::SetBuildPage(UMetaSoundBuilderBase* Builder, FName PageName, EMetaSoundBuilderResult& OutResult) const
+{
+	if (!Builder)
+	{
+		OutResult = EMetaSoundBuilderResult::Failed;
+		return;
+	}
+	
+	// For now, the focused page and build page are the same,
+	// but in the future they may be different (ex. if the editor supports multiple tabs for editing different pages at once)
+	// Only set the page, don't open the editor or change the audition page
+	const FGuid* PageID = UMetaSoundSettings::GetPageID(PageName);
+	if (!PageID)
+	{
+		OutResult = EMetaSoundBuilderResult::Failed;
+		return;
+	}
+	
+	constexpr bool bOpenEditor = false;
+	constexpr bool bPostTransaction = true;
+	constexpr bool bSetAuditionPage = false;
+	if (SetFocusedPageInternal(PageName, *PageID, *Builder, bOpenEditor, bPostTransaction, bSetAuditionPage))
+	{
+		OutResult = EMetaSoundBuilderResult::Succeeded;
+		return;
+	}
+
+	OutResult = EMetaSoundBuilderResult::Failed;
+}
+
 void UMetaSoundEditorSubsystem::SetFocusedPage(UMetaSoundBuilderBase* Builder, FName PageName, bool bOpenEditor, EMetaSoundBuilderResult& OutResult) const
 {
 	using namespace Metasound::Frontend;
@@ -454,7 +495,8 @@ void UMetaSoundEditorSubsystem::SetFocusedPage(UMetaSoundBuilderBase* Builder, F
 	if (const FMetaSoundPageSettings* PageSettings = Settings->FindPageSettings(PageName))
 	{
 		constexpr bool bPostTransaction = true;
-		const bool bFocusedPage = SetFocusedPageInternal(PageSettings->Name, PageSettings->UniqueId, *Builder, bOpenEditor, bPostTransaction);
+		constexpr bool bSetAuditionPage = true;
+		const bool bFocusedPage = SetFocusedPageInternal(PageSettings->Name, PageSettings->UniqueId, *Builder, bOpenEditor, bPostTransaction, bSetAuditionPage);
 		if (bFocusedPage)
 		{
 			OutResult = EMetaSoundBuilderResult::Succeeded;
@@ -481,38 +523,45 @@ bool UMetaSoundEditorSubsystem::SetFocusedPage(UMetaSoundBuilderBase& Builder, c
 	return SetFocusedPageInternal(PageName, InPageID, Builder, bOpenEditor, bPostTransaction);
 }
 
-bool UMetaSoundEditorSubsystem::SetFocusedPageInternal(FName PageName, const FGuid& InPageID, UMetaSoundBuilderBase& Builder, bool bOpenEditor, bool bPostTransaction) const
+bool UMetaSoundEditorSubsystem::SetFocusedPageInternal(FName PageName, const FGuid& InPageID, UMetaSoundBuilderBase& Builder, bool bOpenEditor, bool bPostTransaction, bool bSetAuditionPage) const
 {
 	using namespace Metasound::Frontend;
 
 	const FScopedTransaction Transaction(FText::Format(LOCTEXT("SetFocusedPageTransactionFormat", "Set Focused Page '{0}'"), FText::FromName(PageName)), bPostTransaction);
-	bool bAuditionPageSet = false;
+	bool bSuccess = false;
 	// Must set audition target page before setting build page ID as listeners
 	// to build page ID changes need to reliably be able to adjust to newly assigned
 	// audition target page.
-	UMetasoundEditorSettings* EditorSettings = GetMutableDefault<UMetasoundEditorSettings>();
-	check(EditorSettings);
-	if (EditorSettings->GetAuditionPageMode() == EAuditionPageMode::Focused)
+	if (bSetAuditionPage)
 	{
-		if (EditorSettings->GetAuditionPage() != PageName)
+		UMetasoundEditorSettings* EditorSettings = GetMutableDefault<UMetasoundEditorSettings>();
+		check(EditorSettings);
+		if (EditorSettings->GetAuditionPageMode() == EAuditionPageMode::Focused)
 		{
-			EditorSettings->Modify();
-			EditorSettings->SetAuditionPage(PageName);
-			bAuditionPageSet = true;
+			if (EditorSettings->GetAuditionPage() != PageName)
+			{
+				EditorSettings->Modify();
+				EditorSettings->SetAuditionPage(PageName);
+			}
 		}
 	}
-
+	
 	const FMetaSoundFrontendDocumentBuilder& DocBuilder = Builder.GetConstBuilder();
 	if (DocBuilder.GetBuildPageID() != InPageID)
 	{
 		Builder.Modify();
 		UObject& MetaSound = DocBuilder.CastDocumentObjectChecked<UObject>();
 		MetaSound.Modify();
-		DocBuilder.GetMetasoundAsset().GetGraphChecked().Modify();
+		// Ed graph may be invalid if MetaSound has not been opened yet
+		if (UEdGraph* EdGraph = DocBuilder.GetMetasoundAsset().GetGraph())
+		{
+			EdGraph->Modify();
+		}
 		if (Builder.GetBuilder().SetBuildPageID(InPageID))
 		{
 			// Reregister to ensure all future audible instances are using the new page implementation.
 			RegisterGraphWithFrontend(MetaSound);
+			bSuccess = true;
 		}
 
 		if (GEditor && bOpenEditor)
@@ -524,7 +573,7 @@ bool UMetaSoundEditorSubsystem::SetFocusedPageInternal(FName PageName, const FGu
 		}
 	}
 
-	return bAuditionPageSet;
+	return bSuccess;
 }
 
 bool UMetaSoundEditorSubsystem::UnregisterToolbarExtender(TSharedRef<FExtender> InExtender)
diff --git a/Engine/Plugins/Runtime/Metasound/Source/MetasoundEditor/Public/MetasoundEditorBuilderListener.h b/Engine/Plugins/Runtime/Metasound/Source/MetasoundEditor/Public/MetasoundEditorBuilderListener.h
--- a/Engine/Plugins/Runtime/Metasound/Source/MetasoundEditor/Public/MetasoundEditorBuilderListener.h
+++ b/Engine/Plugins/Runtime/Metasound/Source/MetasoundEditor/Public/MetasoundEditorBuilderListener.h
@@ -9,15 +9,34 @@
 
 // Forward Declarations
 struct FMetaSoundNodeHandle;
-struct FMetasoundFrontendLiteral;
+
+USTRUCT(BlueprintType)
+struct FGraphInputDefaultChangedArgs
+{
+	GENERATED_BODY()
+
+	UPROPERTY(BlueprintReadWrite, Category = "Audio|MetaSound|BuilderListener")
+	FName InputName;
+
+	UPROPERTY(BlueprintReadWrite, Category = "Audio|MetaSound|BuilderListener")
+	FName PageName;
+
+	UPROPERTY(BlueprintReadWrite, Category = "Audio|MetaSound|BuilderListener")
+	FMetasoundFrontendLiteral Literal;
+
+	UPROPERTY(BlueprintReadWrite, Category = "Audio|MetaSound|BuilderListener")
+	EMetaSoundLiteralChangeType ChangeType = EMetaSoundLiteralChangeType::Invalid;
+};
 
 // BP Delegates for builder changes
 DECLARE_DYNAMIC_MULTICAST_DELEGATE_OneParam(FOnMetaSoundBuilderDocumentMetadataTextChanged, FText, NewText);
 DECLARE_DYNAMIC_MULTICAST_DELEGATE_OneParam(FOnMetaSoundBuilderDocumentMetadataTextArrayChanged, TArray<FText>, NewTextArray);
 DECLARE_DYNAMIC_MULTICAST_DELEGATE_OneParam(FOnMetaSoundBuilderDocumentMetadataStringChanged, FString, NewString);
 DECLARE_DYNAMIC_MULTICAST_DELEGATE_OneParam(FOnMetaSoundBuilderDocumentMetadataBoolChanged, bool, NewBool);
+DECLARE_DYNAMIC_MULTICAST_DELEGATE_OneParam(FOnMetaSoundBuilderPageMutate, FName, PageName);
 DECLARE_DYNAMIC_MULTICAST_DELEGATE_TwoParams(FOnMetaSoundBuilderGraphInterfaceMutate, FName, VertexName, FName, DataType);
 DECLARE_DYNAMIC_MULTICAST_DELEGATE_ThreeParams(FOnMetaSoundBuilderGraphLiteralMutate, FName, VertexName, FMetasoundFrontendLiteral, LiteralValue, FName, PageName);
+DECLARE_DYNAMIC_MULTICAST_DELEGATE_OneParam(FOnMetaSoundBuilderGraphInputDefaultChanged, const FGraphInputDefaultChangedArgs&, GraphInputDefaultChangedArgs);
 DECLARE_DYNAMIC_MULTICAST_DELEGATE_TwoParams(FOnMetaSoundBuilderGraphVertexRename, FName, OldName, FName, NewName);
 DECLARE_DYNAMIC_MULTICAST_DELEGATE_TwoParams(FOnMetaSoundBuilderGraphVertexBoolChanged, FName, VertexName, bool, bNewBool);
 DECLARE_DYNAMIC_MULTICAST_DELEGATE_TwoParams(FOnMetaSoundBuilderGraphVertexIntChanged, FName, VertexName, int32, NewInt);
@@ -28,8 +47,9 @@ class UMetaSoundEditorBuilderListener : public UObject
 {
 public:
 	GENERATED_BODY()
-
+PRAGMA_DISABLE_DEPRECATION_WARNINGS
 	virtual ~UMetaSoundEditorBuilderListener() = default;
+PRAGMA_ENABLE_DEPRECATION_WARNINGS
 	
 	void Init(const TWeakObjectPtr<UMetaSoundBuilderBase> InBuilder);
 
@@ -39,6 +59,10 @@ public:
 	void OnDocumentKeywordsChanged();
 	void OnDocumentCategoryHierarchyChanged();
 	void OnDocumentIsDeprecatedChanged();
+	
+	void OnPageAdded(const Metasound::Frontend::FDocumentMutatePageArgs& PageArgs);
+	void OnRemovingPage(const Metasound::Frontend::FDocumentMutatePageArgs& PageArgs);
+	void OnBuildPageChanged(const Metasound::Frontend::FDocumentMutatePageArgs& PageArgs);
 
 	void OnGraphInputAdded(int32 Index);
 	void OnRemovingGraphInput(int32 Index);
@@ -49,7 +73,9 @@ public:
 	void OnGraphInputSortOrderIndexChanged(int32 Index);
 	void OnGraphInputIsConstructorPinChanged(int32 Index);
 	void OnGraphInputIsAdvancedDisplayChanged(int32 Index);
+	UE_DEPRECATED(5.8, "Use the input default changed delegate that includes additional args")
 	void OnGraphInputDefaultChanged(int32 Index);
+	void OnGraphInputDefaultChanged(const Metasound::Frontend::FDocumentArrayPagedInputArgs& Args);
 	void OnGraphInputInheritsDefaultChanged(int32 Index);
 
 	void OnGraphOutputAdded(int32 Index);
@@ -65,6 +91,7 @@ public:
 	UFUNCTION(BlueprintCallable, Category = "Audio|MetaSound|BuilderListener")
 	METASOUNDEDITOR_API void RemoveAllDelegates();
 
+	// Metadata
 	UPROPERTY(BlueprintAssignable, BlueprintCallable, Category = "Audio|MetaSound|BuilderListener", DisplayName = "On Document Display Name Changed")
 	FOnMetaSoundBuilderDocumentMetadataTextChanged OnDocumentDisplayNameChangedDelegate;
 
@@ -83,6 +110,17 @@ public:
 	UPROPERTY(BlueprintAssignable, BlueprintCallable, Category = "Audio|MetaSound|BuilderListener", DisplayName = "On Document Is Deprecated Changed")
 	FOnMetaSoundBuilderDocumentMetadataBoolChanged OnDocumentIsDeprecatedChangedDelegate;
 
+	// Pages
+	UPROPERTY(BlueprintAssignable, BlueprintCallable, Category = "Audio|MetaSound|BuilderListener", DisplayName = "On Page Added")
+	FOnMetaSoundBuilderPageMutate OnPageAddedDelegate;
+
+	UPROPERTY(BlueprintAssignable, BlueprintCallable, Category = "Audio|MetaSound|BuilderListener", DisplayName = "On Page Removed")
+	FOnMetaSoundBuilderPageMutate OnRemovingPageDelegate;
+
+	UPROPERTY(BlueprintAssignable, BlueprintCallable, Category = "Audio|MetaSound|BuilderListener", DisplayName = "On Build Page Changed")
+	FOnMetaSoundBuilderPageMutate OnBuildPageChangedDelegate;
+
+	// Graph Interface
 	UPROPERTY(BlueprintAssignable, BlueprintCallable, Category = "Audio|MetaSound|BuilderListener|InputDelegates", DisplayName = "On Graph Input Added")
 	FOnMetaSoundBuilderGraphInterfaceMutate OnGraphInputAddedDelegate;
 
@@ -110,9 +148,13 @@ public:
 	UPROPERTY(BlueprintAssignable, BlueprintCallable, Category = "Audio|MetaSound|BuilderListener|InputDelegates", DisplayName = "On Graph Input Is Advanced Display Changed")
 	FOnMetaSoundBuilderGraphVertexBoolChanged OnGraphInputIsAdvancedDisplayChangedDelegate;
 
-	UPROPERTY(BlueprintAssignable, BlueprintCallable, Category = "Audio|MetaSound|BuilderListener|InputDelegates", DisplayName = "On Graph Input Default Changed")
+	UE_DEPRECATED(5.8, "Use the GraphInputDefaultChanged delegate that includes additional args")
+	UPROPERTY(BlueprintAssignable, BlueprintCallable, Category = "Audio|MetaSound|BuilderListener|InputDelegates", DisplayName = "[Deprecated] On Graph Input Default Changed")
 	FOnMetaSoundBuilderGraphLiteralMutate OnGraphInputDefaultChangedDelegate;
 
+	UPROPERTY(BlueprintAssignable, BlueprintCallable, Category = "Audio|MetaSound|BuilderListener|InputDelegates", DisplayName = "On Graph Input Default Changed")
+	FOnMetaSoundBuilderGraphInputDefaultChanged OnGraphInputDefaultChangedWithArgsDelegate;
+	
 	UPROPERTY(BlueprintAssignable, BlueprintCallable, Category = "Audio|MetaSound|BuilderListener|InputDelegates", DisplayName = "On Graph Input Inherits Default Changed")
 	FOnMetaSoundBuilderGraphVertexBoolChanged OnGraphInputInheritsDefaultChangedDelegate;
 
@@ -151,6 +193,10 @@ private:
 	FDelegateHandle OnDocumentKeywordsChangedHandle;
 	FDelegateHandle OnDocumentCategoryHierarchyChangedHandle;
 	FDelegateHandle OnDocumentIsDeprecatedChangedHandle;
+	
+	FDelegateHandle OnPageAddedHandle;
+	FDelegateHandle OnRemovingPageHandle;
+	FDelegateHandle OnBuildPageChangedHandle;
 
 	FDelegateHandle OnInputAddedHandle;
 	FDelegateHandle OnRemovingInputHandle;
diff --git a/Engine/Plugins/Runtime/Metasound/Source/MetasoundEditor/Public/MetasoundEditorGraph.h b/Engine/Plugins/Runtime/Metasound/Source/MetasoundEditor/Public/MetasoundEditorGraph.h
--- a/Engine/Plugins/Runtime/Metasound/Source/MetasoundEditor/Public/MetasoundEditorGraph.h
+++ b/Engine/Plugins/Runtime/Metasound/Source/MetasoundEditor/Public/MetasoundEditorGraph.h
@@ -202,7 +202,7 @@ public:
 	/** Update the frontend with the given member's default UObject value.
 	 * @param bPostTransaction - Post as editor transaction if true
 	 */
-	virtual void UpdateFrontendDefaultLiteral(bool bPostTransaction, const FGuid* InPageID = nullptr) PURE_VIRTUAL(UMetasoundEditorGraphMember::UpdateFrontendDefaultLiteral, );
+	virtual void UpdateFrontendDefaultLiteral(bool bPostTransaction, const FGuid* InPageID = nullptr, EMetaSoundLiteralChangeType InChangeType = EMetaSoundLiteralChangeType::Set) PURE_VIRTUAL(UMetasoundEditorGraphMember::UpdateFrontendDefaultLiteral, );
 
 	// Synchronizes cached data with the frontend representation on the represented document.
 	UE_API virtual bool Synchronize();
@@ -371,7 +371,7 @@ public:
 	UE_API virtual void SetVertexAccessType(EMetasoundFrontendVertexAccessType InNewAccessType, bool bPostTransaction) override;
 	UE_API virtual bool Synchronize() override;
 	UE_API virtual void CacheBreadcrumb() override;
-	UE_API virtual void UpdateFrontendDefaultLiteral(bool bPostTransaction, const FGuid* InPageID = nullptr) override;
+	UE_API virtual void UpdateFrontendDefaultLiteral(bool bPostTransaction, const FGuid* InPageID = nullptr, EMetaSoundLiteralChangeType InChangeType = EMetaSoundLiteralChangeType::Set) override;
 	UE_API virtual EMetasoundFrontendVertexAccessType GetVertexAccessType() const override;
 
 protected:
@@ -400,7 +400,7 @@ public:
 	UE_API virtual void SetDisplayName(const FText& InNewName, bool bPostTransaction) override;
 	UE_API virtual void SetVertexAccessType(EMetasoundFrontendVertexAccessType InNewAccessType, bool bPostTransaction) override;
 	UE_API virtual bool Synchronize() override;
-	UE_API virtual void UpdateFrontendDefaultLiteral(bool bPostTransaction, const FGuid* InPageID = nullptr) override;
+	UE_API virtual void UpdateFrontendDefaultLiteral(bool bPostTransaction, const FGuid* InPageID = nullptr, EMetaSoundLiteralChangeType InChangeType = EMetaSoundLiteralChangeType::Set) override;
 	UE_API virtual EMetasoundFrontendVertexAccessType GetVertexAccessType() const override;
 
 protected:
@@ -439,7 +439,7 @@ public:
 	UE_API virtual void ResetToClassDefault() override;
 	UE_API virtual bool Synchronize() override;
 	UE_API virtual void CacheBreadcrumb() override;
-	UE_API virtual void UpdateFrontendDefaultLiteral(bool bPostTransaction, const FGuid* InPageID = nullptr) override;
+	UE_API virtual void UpdateFrontendDefaultLiteral(bool bPostTransaction, const FGuid* InPageID = nullptr, EMetaSoundLiteralChangeType InChangeType = EMetaSoundLiteralChangeType::Set) override;
 
 	UE_API virtual TArray<UMetasoundEditorGraphMemberNode*> GetNodes() const override;
 
diff --git a/Engine/Plugins/Runtime/Metasound/Source/MetasoundEditor/Public/MetasoundEditorSubsystem.h b/Engine/Plugins/Runtime/Metasound/Source/MetasoundEditor/Public/MetasoundEditorSubsystem.h
--- a/Engine/Plugins/Runtime/Metasound/Source/MetasoundEditor/Public/MetasoundEditorSubsystem.h
+++ b/Engine/Plugins/Runtime/Metasound/Source/MetasoundEditor/Public/MetasoundEditorSubsystem.h
@@ -72,9 +72,14 @@ public:
 	UFUNCTION(BlueprintCallable, Category = "Audio|MetaSound|Builder", meta = (ExpandEnumAsExecs = "OutResult"))
 	UE_API UPARAM(DisplayName = "Input Metadata") UMetaSoundFrontendMemberMetadata* FindOrCreateGraphInputMetadata(UPARAM(DisplayName = "Builder") UMetaSoundBuilderBase* InBuilder, FName InputName, EMetaSoundBuilderResult& OutResult);
 
+	
 	// Returns the corresponding literal class for a given type
 	UE_API TSubclassOf<UMetasoundEditorGraphMemberDefaultLiteral> GetLiteralClassForType(FName TypeName) const;
 
+	// Get the build page 
+	UFUNCTION(BlueprintCallable, Category = "Audio|MetaSound|Builder|Editor")
+	UE_API FName GetBuildPage(UMetaSoundBuilderBase* Builder) const;
+
 	// Sets the visual location to InLocation of a given node InNode of a given builder's document.
 	UFUNCTION(BlueprintCallable, Category = "Audio|MetaSound|Builder|Editor", meta = (ExpandEnumAsExecs = "OutResult"))
 	UE_API void SetNodeLocation(
@@ -108,13 +113,19 @@ public:
 	// Register toolbar extender that will be displayed in the MetaSound Asset Editor.
 	UE_API void RegisterToolbarExtender(TSharedRef<FExtender> InExtender);
 
-	// If the given page name is implemented on the provided builder, sets the focused page of
+	// If the given page name is implemented on the provided builder, sets the focused (and build) page of
 	// the provided builder to the associated page and sets the audition page to
 	// the provided name. If the given builder has an asset editor open, optionally opens or brings
 	// that editor's associated PageID into user focus.
-	UFUNCTION(BlueprintCallable, Category = "Audio|MetaSound|Builder|Editor", meta = (ExpandEnumAsExecs = "OutResult"))
+	UFUNCTION(BlueprintCallable, BlueprintPure=false, Category = "Audio|MetaSound|Builder|Editor", meta = (ExpandEnumAsExecs = "OutResult"))
 	UE_API void SetFocusedPage(UMetaSoundBuilderBase* Builder, FName PageName, bool bOpenEditor, EMetaSoundBuilderResult& OutResult) const;
 
+	// If the given page name is implemented on the provided builder, sets the build page.
+	// This page will be the page that graph operations (ex. node add)
+	// will be applied to by default if a page name is not specified.
+	UFUNCTION(BlueprintCallable, BlueprintPure=false, Category = "Audio|MetaSound|Builder|Editor", meta = (ExpandEnumAsExecs = "OutResult"))
+	UE_API void SetBuildPage(UMetaSoundBuilderBase* Builder, FName PageName, EMetaSoundBuilderResult& OutResult) const;
+	
 	// If the given PageID is implemented on the provided builder, sets the focused page of
 	// the provided builder to the associated page and sets the audition target page to
 	// the provided ID. If the given builder has an asset editor open, optionally opens or brings
@@ -140,7 +151,7 @@ public:
 private:
 	UE_API UMetasoundEditorGraphMemberDefaultLiteral* CreateMemberMetadata(FMetaSoundFrontendDocumentBuilder& Builder, TSubclassOf<UMetasoundEditorGraphMemberDefaultLiteral> LiteralClass) const;
 
-	UE_API bool SetFocusedPageInternal(FName PageName, const FGuid& InPageID, UMetaSoundBuilderBase& Builder, bool bOpenEditor, bool bPostTransaction) const;
+	UE_API bool SetFocusedPageInternal(FName PageName, const FGuid& InPageID, UMetaSoundBuilderBase& Builder, bool bOpenEditor, bool bPostTransaction, bool bSetAuditionPage = true) const;
 
 
 	// Editor Toolbar Extenders
diff --git a/Engine/Plugins/Runtime/Metasound/Source/MetasoundEngine/Private/MetasoundBuilderBase.cpp b/Engine/Plugins/Runtime/Metasound/Source/MetasoundEngine/Private/MetasoundBuilderBase.cpp
--- a/Engine/Plugins/Runtime/Metasound/Source/MetasoundEngine/Private/MetasoundBuilderBase.cpp
+++ b/Engine/Plugins/Runtime/Metasound/Source/MetasoundEngine/Private/MetasoundBuilderBase.cpp
@@ -956,9 +956,10 @@ const FMetaSoundFrontendDocumentBuilder& UMetaSoundBuilderBase::GetConstBuilder(
 	return Builder;
 }
 
-FMetasoundFrontendLiteral UMetaSoundBuilderBase::GetGraphInputDefault(FName InputName, EMetaSoundBuilderResult& OutResult) const
+FMetasoundFrontendLiteral UMetaSoundBuilderBase::GetGraphInputDefault(FName InputName, EMetaSoundBuilderResult& OutResult, FName PageName) const
 {
-	if (const FMetasoundFrontendLiteral* Default = Builder.GetGraphInputDefault(InputName))
+	const FGuid* PageID = UMetaSoundSettings::GetPageID(PageName);
+	if (const FMetasoundFrontendLiteral* Default = Builder.GetGraphInputDefault(InputName, PageID))
 	{
 		OutResult = EMetaSoundBuilderResult::Succeeded;
 		return *Default;
@@ -1180,6 +1181,12 @@ void UMetaSoundBuilderBase::ReloadCache(bool bPrimeCache)
 	Reload({ }, bPrimeCache);
 }
 
+void UMetaSoundBuilderBase::ResetGraphInputDefaults(FName Name, EMetaSoundBuilderResult& OutResult)
+{
+	const bool bReset = Builder.ResetGraphInputDefault(Name);
+	OutResult = bReset ? EMetaSoundBuilderResult::Succeeded : EMetaSoundBuilderResult::Failed;
+}
+
 #if WITH_EDITORONLY_DATA
 void UMetaSoundBuilderBase::ResetGraphPages(bool bClearDefaultGraph)
 {
@@ -1201,6 +1208,13 @@ void UMetaSoundBuilderBase::RemoveGraphInput(FName Name, EMetaSoundBuilderResult
 	OutResult = bRemoved ? EMetaSoundBuilderResult::Succeeded : EMetaSoundBuilderResult::Failed;
 }
 
+void UMetaSoundBuilderBase::RemoveGraphInputDefault(FName Name, EMetaSoundBuilderResult& OutResult, FName PageName)
+{
+	const FGuid* PageID = UMetaSoundSettings::GetPageID(PageName);
+	const bool bRemoved = Builder.RemoveGraphInputDefault(Name, PageID ? *PageID : Builder.GetBuildPageID());
+	OutResult = bRemoved ? EMetaSoundBuilderResult::Succeeded : EMetaSoundBuilderResult::Failed;
+}
+
 void UMetaSoundBuilderBase::RemoveGraphOutput(FName Name, EMetaSoundBuilderResult& OutResult)
 {
 	const bool bRemoved = Builder.RemoveGraphOutput(Name);
@@ -1296,9 +1310,11 @@ void UMetaSoundBuilderBase::SetGraphInputDataType(FName InputName, FName DataTyp
 	OutResult = bSet ? EMetaSoundBuilderResult::Succeeded : EMetaSoundBuilderResult::Failed;
 }
 
-void UMetaSoundBuilderBase::SetGraphInputDefault(FName InputName, const FMetasoundFrontendLiteral& Literal, EMetaSoundBuilderResult& OutResult)
+void UMetaSoundBuilderBase::SetGraphInputDefault(FName InputName, const FMetasoundFrontendLiteral& Literal, EMetaSoundBuilderResult& OutResult, FName PageName)
 {
-	const bool bSet = Builder.SetGraphInputDefault(InputName, Literal);
+	using namespace Metasound::Frontend;
+	const FGuid* PageID = UMetaSoundSettings::GetPageID(PageName);
+	const bool bSet = Builder.SetGraphInputDefault(InputName, Literal, PageID);
 	OutResult = bSet ? EMetaSoundBuilderResult::Succeeded : EMetaSoundBuilderResult::Failed;
 }
 
diff --git a/Engine/Plugins/Runtime/Metasound/Source/MetasoundEngine/Private/MetasoundBuilderSubsystem.cpp b/Engine/Plugins/Runtime/Metasound/Source/MetasoundEngine/Private/MetasoundBuilderSubsystem.cpp
--- a/Engine/Plugins/Runtime/Metasound/Source/MetasoundEngine/Private/MetasoundBuilderSubsystem.cpp
+++ b/Engine/Plugins/Runtime/Metasound/Source/MetasoundEngine/Private/MetasoundBuilderSubsystem.cpp
@@ -504,14 +504,14 @@ void UMetaSoundSourceBuilder::OnInputDataTypeChanged(int32 InputIndex)
 	OnInputAdded(InputIndex);
 }
 
-void UMetaSoundSourceBuilder::OnInputDefaultChanged(int32 InputIndex) const
+void UMetaSoundSourceBuilder::OnInputDefaultChanged(const Metasound::Frontend::FDocumentArrayPagedInputArgs& Args) const
 {
 	using namespace Metasound::DynamicGraph;
 	using namespace Metasound::Frontend;
 
 	const FMetasoundFrontendDocument& Doc = Builder.GetConstDocumentChecked();
 	const FMetasoundFrontendGraphClass& GraphClass = Doc.RootGraph;
-	const FMetasoundFrontendClassInput& InputWithNewDefault = GraphClass.GetDefaultInterface().Inputs[InputIndex];
+	const FMetasoundFrontendClassInput& InputWithNewDefault = GraphClass.GetDefaultInterface().Inputs[Args.Index];
 
 	UMetaSoundSource& Source = GetMetaSoundSource();
 	Source.RuntimeInputData.InputMap.Remove(InputWithNewDefault.Name);
diff --git a/Engine/Plugins/Runtime/Metasound/Source/MetasoundEngine/Private/MetasoundSettings.cpp b/Engine/Plugins/Runtime/Metasound/Source/MetasoundEngine/Private/MetasoundSettings.cpp
--- a/Engine/Plugins/Runtime/Metasound/Source/MetasoundEngine/Private/MetasoundSettings.cpp
+++ b/Engine/Plugins/Runtime/Metasound/Source/MetasoundEngine/Private/MetasoundSettings.cpp
@@ -430,6 +430,33 @@ TArray<FGuid> UMetaSoundSettings::GeneratePageOrder() const
 	return GeneratePageOrderInternal(PlatformName, GetTargetPageSettings().UniqueId);
 };
 
+const FGuid* UMetaSoundSettings::GetPageID(FName PageName)
+{
+	if (!PageName.IsNone())
+	{
+		const UMetaSoundSettings* Settings = GetDefault<UMetaSoundSettings>();
+		check(Settings);
+
+		if (const FMetaSoundPageSettings* PageSettings = Settings->FindPageSettings(PageName))
+		{
+			return &PageSettings->UniqueId;
+		}
+	}
+	return nullptr;
+}
+
+FName UMetaSoundSettings::GetPageName(FGuid PageID)
+{
+	FName PageName = Metasound::Frontend::DefaultPageName;
+	const UMetaSoundSettings* Settings = GetDefault<UMetaSoundSettings>();
+	check(Settings);
+	if (const FMetaSoundPageSettings* PageSettings = Settings->FindPageSettings(PageID))
+	{
+		PageName = PageSettings->Name;
+	}
+	return PageName;
+}
+
 TArrayView<const FGuid> UMetaSoundSettings::GetPageOrder()
 {
 	// Thread access for the CachedPageOrder is a bit hairy due to the existence
diff --git a/Engine/Plugins/Runtime/Metasound/Source/MetasoundEngine/Public/MetasoundBuilderBase.h b/Engine/Plugins/Runtime/Metasound/Source/MetasoundEngine/Public/MetasoundBuilderBase.h
--- a/Engine/Plugins/Runtime/Metasound/Source/MetasoundEngine/Public/MetasoundBuilderBase.h
+++ b/Engine/Plugins/Runtime/Metasound/Source/MetasoundEngine/Public/MetasoundBuilderBase.h
@@ -338,8 +338,9 @@ public:
 	UE_API UPARAM(DisplayName = "Graph Outputs") TArray<FName> GetGraphOutputNames(EMetaSoundBuilderResult& OutResult) const;
 	
 	// Gets the graph input's default literal value.
-	UFUNCTION(BlueprintCallable, Category = "Audio|MetaSound|Builder", meta = (ExpandEnumAsExecs = "OutResult"))
-	UE_API UPARAM(DisplayName = "Literal Value") FMetasoundFrontendLiteral GetGraphInputDefault(FName InputName, EMetaSoundBuilderResult& OutResult) const;
+	//  @param PageName	If the page name is not specified or not found, the builder's current build page will be used (which can be set in the MetaSound editor).
+	UFUNCTION(BlueprintCallable, Category = "Audio|MetaSound|Builder", meta = (ExpandEnumAsExecs = "OutResult", AdvancedDisplay = "PageName"))
+	UE_API UPARAM(DisplayName = "Literal Value") FMetasoundFrontendLiteral GetGraphInputDefault(FName InputName, EMetaSoundBuilderResult& OutResult, FName PageName = FName()) const;
 
 	// Gets the graph variable's default literal value.
 	UFUNCTION(BlueprintCallable, Category = "Audio|MetaSound|Builder", meta = (ExpandEnumAsExecs = "OutResult"))
@@ -409,6 +410,11 @@ public:
 	UFUNCTION(BlueprintCallable, Category = "Audio|MetaSound|Builder", meta = (ExpandEnumAsExecs = "OutResult"))
 	UE_API void ConvertToPreset(UPARAM(DisplayName = "Parent Node Class") const TScriptInterface<IMetaSoundDocumentInterface>& ReferencedNodeClass, EMetaSoundBuilderResult& OutResult);
 
+	// Removes input literal defaults for all pages if set, reverting the value to the class defaults
+	// Returns success if values reset, false if not (i.e. wasn't set to begin with).
+	UFUNCTION(BlueprintCallable, Category = "Audio|MetaSound|Builder", meta = (ExpandEnumAsExecs = "OutResult"))
+	UE_API void ResetGraphInputDefaults(FName Name, EMetaSoundBuilderResult& OutResult);
+	
 #if WITH_EDITORONLY_DATA
 	// Removes all graph pages except the default.  If bClearDefaultPage is true, clears the default graph page implementation.
 	UE_API void ResetGraphPages(bool bClearDefaultPage);
@@ -418,6 +424,12 @@ public:
 	UFUNCTION(BlueprintCallable, Category = "Audio|MetaSound|Builder", meta = (ExpandEnumAsExecs = "OutResult"))
 	UE_API void RemoveGraphInput(FName Name, EMetaSoundBuilderResult& OutResult);
 
+	// Removes graph input literal default if set, reverting the value to the class default
+	// @param PageName If the page name is not specified or not found, the builder's current build page will be used (which can be set in the MetaSound editor).
+	// Returns success if value was removed, false if not removed (i.e. wasn't set to begin with).
+	UFUNCTION(BlueprintCallable, Category = "Audio|MetaSound|Builder", meta = (ExpandEnumAsExecs = "OutResult", AdvancedDisplay = "PageName"))
+	UE_API void RemoveGraphInputDefault(FName Name, EMetaSoundBuilderResult& OutResult, FName PageName = FName());
+	
 	// Removes output from all paged graphs if it exists; sets result to succeeded if it was removed and failed if it was not.
 	UFUNCTION(BlueprintCallable, Category = "Audio|MetaSound|Builder", meta = (ExpandEnumAsExecs = "OutResult"))
 	UE_API void RemoveGraphOutput(FName Name, EMetaSoundBuilderResult& OutResult);
@@ -514,10 +526,12 @@ public:
 	UFUNCTION(BlueprintCallable, Category = "Audio|MetaSound|Builder", meta = (ExpandEnumAsExecs = "OutResult"))
 	UE_API void SetGraphInputDataType(FName InputName, FName DataType, EMetaSoundBuilderResult& OutResult);
 
-	// Sets the input node's default value, overriding the default provided by the referenced graph if the graph is a preset.
-	UFUNCTION(BlueprintCallable, Category = "Audio|MetaSound|Builder", meta = (ExpandEnumAsExecs = "OutResult"))
-	UE_API void SetGraphInputDefault(FName InputName, const FMetasoundFrontendLiteral& Literal, EMetaSoundBuilderResult& OutResult);
-	
+	/* Sets the input node's default value, overriding the default provided by the referenced graph if the graph is a preset.
+	*  @param PageName	If the page name is not specified or not found, the builder's current build page will be used (which can be set in the MetaSound editor).
+	*/
+	UFUNCTION(BlueprintCallable, Category = "Audio|MetaSound|Builder", meta = (ExpandEnumAsExecs = "OutResult", AdvancedDisplay = "PageName"))
+	UE_API void SetGraphInputDefault(FName InputName, const FMetasoundFrontendLiteral& Literal, EMetaSoundBuilderResult& OutResult, FName PageName = FName());	
+
 	// Sets the given graph input's name to the new name. 
 	// Result succeeds if the name was successfully changed or the new name is the same as the old name, and fails if the given input name doesn't exist.
 	UFUNCTION(BlueprintCallable, Category = "Audio|MetaSound|Builder", meta = (ExpandEnumAsExecs = "OutResult"))
@@ -603,8 +617,7 @@ public:
 
 	// Initializes builder given an existing document interface 
 	UE_API void Initialize(TScriptInterface<IMetaSoundDocumentInterface>& DocInterface);
-
-
+	
 protected:
 	virtual void BuildAndOverwriteMetaSoundInternal(TScriptInterface<IMetaSoundDocumentInterface> ExistingMetaSound, bool bForceUniqueClassName) const PURE_VIRTUAL(UMetaSoundBuilderBase::BuildAndOverwriteMetaSoundInternal, ;);
 	UE_API virtual void InitDelegates(Metasound::Frontend::FDocumentModifyDelegates& OutDocumentDelegates);
diff --git a/Engine/Plugins/Runtime/Metasound/Source/MetasoundEngine/Public/MetasoundBuilderSubsystem.h b/Engine/Plugins/Runtime/Metasound/Source/MetasoundEngine/Public/MetasoundBuilderSubsystem.h
--- a/Engine/Plugins/Runtime/Metasound/Source/MetasoundEngine/Public/MetasoundBuilderSubsystem.h
+++ b/Engine/Plugins/Runtime/Metasound/Source/MetasoundEngine/Public/MetasoundBuilderSubsystem.h
@@ -126,7 +126,8 @@ private:
 	UE_API void OnEdgeAdded(int32 EdgeIndex) const;
 	UE_API void OnInputAdded(int32 InputIndex);
 	UE_API void OnInputDataTypeChanged(int32 InputIndex);
-	UE_API void OnInputDefaultChanged(int32 InputIndex) const;
+	UE_API void OnInputDefaultChanged(const Metasound::Frontend::FDocumentArrayPagedInputArgs& Args) const;
+	
 	UE_API void OnInputNameChanged(FName OldName, FName NewName);
 	UE_API void OnLiveComponentFinished(UAudioComponent* AudioComponent);
 	UE_API void OnNodeAdded(int32 NodeIndex) const;
diff --git a/Engine/Plugins/Runtime/Metasound/Source/MetasoundEngine/Public/MetasoundSettings.h b/Engine/Plugins/Runtime/Metasound/Source/MetasoundEngine/Public/MetasoundSettings.h
--- a/Engine/Plugins/Runtime/Metasound/Source/MetasoundEngine/Public/MetasoundSettings.h
+++ b/Engine/Plugins/Runtime/Metasound/Source/MetasoundEngine/Public/MetasoundSettings.h
@@ -237,9 +237,15 @@ private:
 #endif
 
 public:
+	// Helper function to convert page name to id. Returns nullptr if page name not found.
+	static UE_API const FGuid* GetPageID(FName PageName);
+
+	// Helper function to convert page ID into page name. If page ID is not found, returns default page.
+	static UE_API FName GetPageName(FGuid PageID);
+	   
 	// Returns a cached global version of the PageOrder from the UMetaSoundSettings Default Object.
 	static UE_API TArrayView<const FGuid> GetPageOrder();
-	   
+
 	// Generates the page order for a current configured platform and target ID
 	TArray<FGuid> GeneratePageOrder() const;
 
diff --git a/Engine/Plugins/Runtime/Metasound/Source/MetasoundFrontend/Private/MetasoundFrontendDocumentBuilder.cpp b/Engine/Plugins/Runtime/Metasound/Source/MetasoundFrontend/Private/MetasoundFrontendDocumentBuilder.cpp
--- a/Engine/Plugins/Runtime/Metasound/Source/MetasoundFrontend/Private/MetasoundFrontendDocumentBuilder.cpp
+++ b/Engine/Plugins/Runtime/Metasound/Source/MetasoundFrontend/Private/MetasoundFrontendDocumentBuilder.cpp
@@ -1431,7 +1431,8 @@ const FMetasoundFrontendVariable* FMetaSoundFrontendDocumentBuilder::AddGraphVar
 {
 	using namespace Metasound::Frontend;
 
-	if (const FMetasoundFrontendVariable* ExistingVariable = FindGraphVariable(VariableName))
+	const FGuid& PageID = InPageID ? *InPageID : BuildPageID;
+	if (const FMetasoundFrontendVariable* ExistingVariable = FindGraphVariable(VariableName, &PageID))
 	{
 		UE_LOG(LogMetaSound, Warning, TEXT("AddGraphVariable Failed: Variable already exists with name '%s' (existing DataType '%s', requested DataType '%s')"),
 			*VariableName.ToString(), *ExistingVariable->TypeName.ToString(), *DataType.ToString());
@@ -1512,7 +1513,6 @@ const FMetasoundFrontendVariable* FMetaSoundFrontendDocumentBuilder::AddGraphVar
 #endif // WITH_EDITOR
 	};
 
-	const FGuid& PageID = InPageID ? *InPageID : BuildPageID;
 	if (FMetasoundFrontendNode* VariableNode = AddNodeInternal(VariableNodeClassKey, FinalizeNode, PageID))
 	{
 		Variable.VariableNodeID = VariableNode->GetID();
@@ -4964,7 +4964,6 @@ void FMetaSoundFrontendDocumentBuilder::Reload(TSharedPtr<Metasound::Frontend::F
 	}
 }
 
-#if WITH_EDITORONLY_DATA
 bool FMetaSoundFrontendDocumentBuilder::RemoveGraphInputDefault(FName InputName, const FGuid& InPageID, bool bClearInheritsDefault)
 {
 	return RemoveGraphInputDefault(InputName, InPageID);
@@ -4973,6 +4972,7 @@ bool FMetaSoundFrontendDocumentBuilder::RemoveGraphInputDefault(FName InputName,
 bool FMetaSoundFrontendDocumentBuilder::RemoveGraphInputDefault(FName InputName, const FGuid& InPageID)
 {
 	using namespace Metasound;
+	using namespace Metasound::Frontend;
 
 	auto NameMatchesInput = [&InputName](const FMetasoundFrontendClassInput& Input) { return Input.Name == InputName; };
 	FMetasoundFrontendDocument& Document = GetDocumentChecked();
@@ -4985,14 +4985,19 @@ bool FMetaSoundFrontendDocumentBuilder::RemoveGraphInputDefault(FName InputName,
 		const bool bRemovedDefault = Input.RemoveDefault(InPageID);
 		if (bRemovedDefault)
 		{
-			DocumentDelegates->InterfaceDelegates.OnInputDefaultChanged.Broadcast(Index);
+			DocumentDelegates->InterfaceDelegates.OnInputDefaultChanged.Broadcast(
+				FDocumentArrayPagedInputArgs
+				{
+					.Index = Index,
+					.PageID = InPageID,
+					.ChangeType = EMetaSoundLiteralChangeType::Remove
+				});
 			return true;
 		}
 	}
 
 	return false;
 }
-#endif // WITH_EDITORONLY_DATA
 
 bool FMetaSoundFrontendDocumentBuilder::RemoveNodeInputDefault(const FGuid& InNodeID, const FGuid& InVertexID, const FGuid* InPageID)
 {
@@ -5813,7 +5818,6 @@ FMetaSoundFrontendDocumentBuilder::FNodeInstanceReplacementData FMetaSoundFronte
 	return ReplacementData;
 }
 
-#if WITH_EDITORONLY_DATA
 bool FMetaSoundFrontendDocumentBuilder::ResetGraphInputDefault(FName InputName)
 {
 	using namespace Metasound;
@@ -5827,10 +5831,24 @@ bool FMetaSoundFrontendDocumentBuilder::ResetGraphInputDefault(FName InputName)
 	if (Index != INDEX_NONE)
 	{
 		FMetasoundFrontendClassInput& Input = Inputs[Index];
+		TArray<const FGuid> PageIDs;
+		Input.IterateDefaults([&](const FGuid& PageID, FMetasoundFrontendLiteral&)
+		{
+			PageIDs.Add(PageID);
+		});
 		Input.ResetDefaults();
 
-		DocumentDelegates->InterfaceDelegates.OnInputDefaultChanged.Broadcast(Index);
-
+		for (const FGuid& PageID : PageIDs)
+		{
+			DocumentDelegates->InterfaceDelegates.OnInputDefaultChanged.Broadcast(
+			FDocumentArrayPagedInputArgs
+			{
+				.Index = Index,
+				.PageID = PageID,
+				.ChangeType = EMetaSoundLiteralChangeType::Reset
+			});
+		}
+		
 		if (FDocumentVertexConfiguration* Config = Document.Configuration.GetMutablePtr<FDocumentVertexConfiguration>())
 		{
 			if (FMetaSoundFrontendDocumentConfigurationVertexMetadata* VertexMetadata = Config->FindVertexMetadata(Input.NodeID))
@@ -5839,13 +5857,16 @@ bool FMetaSoundFrontendDocumentBuilder::ResetGraphInputDefault(FName InputName)
 			}
 		}
 
+#if WITH_EDITORONLY_DATA
 		Document.Metadata.ModifyContext.AddMemberIDModified(Input.NodeID);
+#endif // WITH_EDITORONLY_DATA
 		return true;
 	}
 
 	return false;
 }
 
+#if WITH_EDITORONLY_DATA
 void FMetaSoundFrontendDocumentBuilder::ResetGraphPages(bool bClearDefaultGraph)
 {
 	using namespace Metasound;
@@ -6086,6 +6107,11 @@ bool FMetaSoundFrontendDocumentBuilder::SetGraphInputDataType(FName InputName, F
 				}
 			}
 			GraphInput.TypeName = DataType;
+			TArray<const FGuid> PageIDsModified;
+			GraphInput.IterateDefaults([&](const FGuid& PageID, FMetasoundFrontendLiteral&)
+			{
+				PageIDsModified.Add(PageID);
+			});
 			GraphInput.ResetDefaults();
 
 			RootGraph.IterateGraphPages([this, &DataType, &GraphInput](FMetasoundFrontendGraph& Graph)
@@ -6124,7 +6150,16 @@ bool FMetaSoundFrontendDocumentBuilder::SetGraphInputDataType(FName InputName, F
 				return false;
 			}
 
-			DocumentDelegates->InterfaceDelegates.OnInputDefaultChanged.Broadcast(*Index);
+			for (const FGuid& PageID : PageIDsModified)
+			{
+				DocumentDelegates->InterfaceDelegates.OnInputDefaultChanged.Broadcast(
+				FDocumentArrayPagedInputArgs
+				{
+					.Index = *Index,
+					.PageID = PageID,
+					.ChangeType = EMetaSoundLiteralChangeType::Reset
+				});
+			}
 			DocumentDelegates->InterfaceDelegates.OnInputDataTypeChanged.Broadcast(*Index);
 
 			RemoveUnusedDependencies();
@@ -6141,7 +6176,7 @@ bool FMetaSoundFrontendDocumentBuilder::SetGraphInputDataType(FName InputName, F
 	return true;
 }
 
-bool FMetaSoundFrontendDocumentBuilder::SetGraphInputDefault(FName InputName, FMetasoundFrontendLiteral InDefaultLiteral, const FGuid* InPageID)
+bool FMetaSoundFrontendDocumentBuilder::SetGraphInputDefault(FName InputName, FMetasoundFrontendLiteral InDefaultLiteral, const FGuid* InPageID, EMetaSoundLiteralChangeType InChangeType)
 {
 	using namespace Metasound;
 	using namespace Metasound::Frontend;
@@ -6170,7 +6205,13 @@ bool FMetaSoundFrontendDocumentBuilder::SetGraphInputDefault(FName InputName, FM
 			{
 				Input.AddDefault(PageID) = MoveTemp(InDefaultLiteral);
 			}
-			DocumentDelegates->InterfaceDelegates.OnInputDefaultChanged.Broadcast(Index);
+			DocumentDelegates->InterfaceDelegates.OnInputDefaultChanged.Broadcast(
+				FDocumentArrayPagedInputArgs
+				{
+					.Index = Index,
+					.PageID = PageID,
+					.ChangeType = InChangeType
+				});
 
 			if (const FDocumentConfigurationVertexMetadata* ConfigVertexMetadata = FindConstVertexMetadata<FDocumentConfigurationVertexMetadata>(Input.NodeID))
 			{
@@ -6185,7 +6226,7 @@ bool FMetaSoundFrontendDocumentBuilder::SetGraphInputDefault(FName InputName, FM
 	return false;
 }
 
-bool FMetaSoundFrontendDocumentBuilder::SetGraphInputDefaults(FName InputName, TArray<FMetasoundFrontendClassInputDefault> Defaults)
+bool FMetaSoundFrontendDocumentBuilder::SetGraphInputDefaults(FName InputName, TArray<FMetasoundFrontendClassInputDefault> Defaults, EMetaSoundLiteralChangeType InChangeType)
 {
 	using namespace Metasound;
 	using namespace Metasound::Frontend;
@@ -6198,7 +6239,42 @@ bool FMetaSoundFrontendDocumentBuilder::SetGraphInputDefaults(FName InputName, T
 	if (Index != INDEX_NONE)
 	{
 		FMetasoundFrontendClassInput& Input = Inputs[Index];
-		TSet<FGuid> ValidPageIDs;
+		TArray<FGuid> PageIDsModified;
+
+		// If defaults are being removed or reset, compare with existing ones to report the correct page id being modified
+		if (InChangeType == EMetaSoundLiteralChangeType::Remove)
+		{
+			Input.IterateDefaults([&](const FGuid& PageID, const FMetasoundFrontendLiteral& Literal)
+			{
+				PageIDsModified.Emplace(PageID);
+			});
+			for (const FMetasoundFrontendClassInputDefault& Default : Defaults)
+			{
+				PageIDsModified.Remove(Default.PageID);
+			}
+		}
+		else if (InChangeType == EMetaSoundLiteralChangeType::Set || InChangeType == EMetaSoundLiteralChangeType::Reset)
+		{
+			Input.IterateDefaults([&](const FGuid& PageID, const FMetasoundFrontendLiteral& Literal)
+			{
+				PageIDsModified.Emplace(PageID);
+			});
+			for (const FMetasoundFrontendClassInputDefault& Default : Defaults)
+			{
+				if (const FMetasoundFrontendLiteral* ExistingLiteral = Input.FindConstDefault(Default.PageID))
+				{
+					if (Default.Literal.IsEqual(*ExistingLiteral))
+					{
+						PageIDsModified.Remove(Default.PageID);
+					}
+				}
+			}
+		}
+		else
+		{
+			UE_LOG(LogMetaSound, Warning, TEXT("Possible unhandled EMetaSoundLiteralChangeType in SetGraphInputDefaults"));
+		}
+		
 		bool bAllSupported = Algo::AllOf(Defaults, [&Input](const FMetasoundFrontendClassInputDefault& Default)
 		{
 			return Frontend::IDataTypeRegistry::Get().IsLiteralTypeSupported(Input.TypeName, Default.Literal.GetType());
@@ -6206,7 +6282,16 @@ bool FMetaSoundFrontendDocumentBuilder::SetGraphInputDefaults(FName InputName, T
 		if (bAllSupported)
 		{
 			Input.SetDefaults(MoveTemp(Defaults));
-			DocumentDelegates->InterfaceDelegates.OnInputDefaultChanged.Broadcast(Index);
+			for (const FGuid& PageID : PageIDsModified)
+			{
+				DocumentDelegates->InterfaceDelegates.OnInputDefaultChanged.Broadcast(
+				FDocumentArrayPagedInputArgs
+				{
+					.Index = Index,
+					.PageID = PageID,
+					.ChangeType = InChangeType
+				});
+			}
 
 			if (FDocumentVertexConfiguration* DocConfig = Document.Configuration.GetMutablePtr<FDocumentVertexConfiguration>())
 			{
diff --git a/Engine/Plugins/Runtime/Metasound/Source/MetasoundFrontend/Private/MetasoundFrontendDocumentCache.cpp b/Engine/Plugins/Runtime/Metasound/Source/MetasoundFrontend/Private/MetasoundFrontendDocumentCache.cpp
--- a/Engine/Plugins/Runtime/Metasound/Source/MetasoundFrontend/Private/MetasoundFrontendDocumentCache.cpp
+++ b/Engine/Plugins/Runtime/Metasound/Source/MetasoundFrontend/Private/MetasoundFrontendDocumentCache.cpp
@@ -476,7 +476,7 @@ namespace Metasound::Frontend
 		++TransactionCount;
 	}
 
-	void FDocumentGraphInterfaceCache::OnInputDefaultChanged(int32 NewIndex)
+	void FDocumentGraphInterfaceCache::OnInputDefaultChanged(const FDocumentArrayPagedInputArgs& Args)
 	{
 		++TransactionCount;
 	}
diff --git a/Engine/Plugins/Runtime/Metasound/Source/MetasoundFrontend/Private/MetasoundFrontendDocumentCache.h b/Engine/Plugins/Runtime/Metasound/Source/MetasoundFrontend/Private/MetasoundFrontendDocumentCache.h
--- a/Engine/Plugins/Runtime/Metasound/Source/MetasoundFrontend/Private/MetasoundFrontendDocumentCache.h
+++ b/Engine/Plugins/Runtime/Metasound/Source/MetasoundFrontend/Private/MetasoundFrontendDocumentCache.h
@@ -139,7 +139,7 @@ namespace Metasound::Frontend
 		void OnRemovingInput(int32 IndexBeingRemoved);
 		void OnInputNameChanged(FName OldName, FName NewName);
 		void OnInputDataTypeChanged(int32 Index);
-		void OnInputDefaultChanged(int32 NewIndex);
+		void OnInputDefaultChanged(const FDocumentArrayPagedInputArgs& Args);
 		void OnInputInheritsDefaultChanged(int32 NewIndex);
 #if WITH_EDITOR
 		void OnInputDisplayNameChanged(int32 Index);
diff --git a/Engine/Plugins/Runtime/Metasound/Source/MetasoundFrontend/Public/MetasoundFrontendDocumentBuilder.h b/Engine/Plugins/Runtime/Metasound/Source/MetasoundFrontend/Public/MetasoundFrontendDocumentBuilder.h
--- a/Engine/Plugins/Runtime/Metasound/Source/MetasoundFrontend/Public/MetasoundFrontendDocumentBuilder.h
+++ b/Engine/Plugins/Runtime/Metasound/Source/MetasoundFrontend/Public/MetasoundFrontendDocumentBuilder.h
@@ -631,12 +631,9 @@ PRAGMA_ENABLE_EXPERIMENTAL_WARNINGS
 
 	UE_API void Reload(TSharedPtr<Metasound::Frontend::FDocumentModifyDelegates> Delegates = {}, bool bPrimeCache = false);
 
-#if WITH_EDITORONLY_DATA
 	UE_DEPRECATED(5.8, "Input inheritance data is now cached in doc config.  No need to call directly here. Use overload without bool parameter.")
 	UE_API bool RemoveGraphInputDefault(FName InputName, const FGuid& InPageID, bool bClearInheritsDefault);
-
 	UE_API bool RemoveGraphInputDefault(FName InputName, const FGuid& InPageID);
-#endif // WITH_EDITORONLY_DATA
 
 	UE_API bool RemoveNodeInputDefault(const FGuid& InNodeID, const FGuid& InVertexID, const FGuid* InPageID = nullptr);
 	UE_API bool RemoveUnusedDependencies();
@@ -647,9 +644,8 @@ PRAGMA_ENABLE_EXPERIMENTAL_WARNINGS
 	// Optionally report disconnected inputs and outputs
 	UE_API bool ReplaceDependency(const Metasound::Frontend::FNodeClassRegistryKey& OldClass, const Metasound::Frontend::FNodeClassRegistryKey& NewClass, TArray<FVertexNameAndType>* OutDisconnectedInputs = nullptr, TArray<FVertexNameAndType>* OutDisconnectedOutputs = nullptr);
 
-#if WITH_EDITORONLY_DATA
 	UE_API bool ResetGraphInputDefault(FName InputName);
-
+#if WITH_EDITORONLY_DATA
 	// Removes all graph pages except the default.  If bClearDefaultPage is true, clears the default graph page implementation.
 	UE_API void ResetGraphPages(bool bClearDefaultGraph);
 #endif // WITH_EDITORONLY_DATA
@@ -721,8 +717,10 @@ PRAGMA_ENABLE_EXPERIMENTAL_WARNINGS
 	// already the given DataType.
 	UE_API bool SetGraphInputDataType(FName InputName, FName DataType);
 
-	UE_API bool SetGraphInputDefault(FName InputName, FMetasoundFrontendLiteral InDefaultLiteral, const FGuid* InPageID = nullptr);
-	UE_API bool SetGraphInputDefaults(FName InputName, TArray<FMetasoundFrontendClassInputDefault> Defaults);
+	// Set the given graph input's default(s).
+	// Change type can be optionally passed in so the on input changed delegate can report the context, since a set can be used to perform remove/reset operations.
+	UE_API bool SetGraphInputDefault(FName InputName, FMetasoundFrontendLiteral InDefaultLiteral, const FGuid* InPageID = nullptr, EMetaSoundLiteralChangeType InChangeType = EMetaSoundLiteralChangeType::Set);
+	UE_API bool SetGraphInputDefaults(FName InputName, TArray<FMetasoundFrontendClassInputDefault> Defaults, EMetaSoundLiteralChangeType InChangeType = EMetaSoundLiteralChangeType::Set);
 
 #if WITH_EDITORONLY_DATA
 	// Sets the graph input's description. Returns true if found and set, false if not.
diff --git a/Engine/Plugins/Runtime/Metasound/Source/MetasoundFrontend/Public/MetasoundFrontendDocumentModifyDelegates.h b/Engine/Plugins/Runtime/Metasound/Source/MetasoundFrontend/Public/MetasoundFrontendDocumentModifyDelegates.h
--- a/Engine/Plugins/Runtime/Metasound/Source/MetasoundFrontend/Public/MetasoundFrontendDocumentModifyDelegates.h
+++ b/Engine/Plugins/Runtime/Metasound/Source/MetasoundFrontend/Public/MetasoundFrontendDocumentModifyDelegates.h
@@ -5,6 +5,8 @@
 #include "MetasoundFrontendDocument.h"
 #include "Templates/SharedPointer.h"
 
+#include "MetasoundFrontendDocumentModifyDelegates.generated.h"
+
 #define UE_API METASOUNDFRONTEND_API
 
 
@@ -21,6 +23,22 @@ DECLARE_MULTICAST_DELEGATE_TwoParams(FOnMetaSoundFrontendDocumentRenameVertex, F
 DECLARE_MULTICAST_DELEGATE_OneParam(FOnMetaSoundFrontendDocumentMutateMember, FGuid /* Member ID*/);
 #endif // WITH_EDITOR
 
+UENUM(BlueprintType)
+enum class EMetaSoundLiteralChangeType : uint8
+{
+	// Initial, unset type
+	Invalid = 0,
+		
+	// Default is set 
+	Set,
+		
+	// Default is removed for this page
+	Remove,
+		
+	// Default is reset to default value 
+	Reset
+};
+
 namespace Metasound::Frontend
 {
 	struct FDocumentMutatePageArgs
@@ -43,6 +61,14 @@ namespace Metasound::Frontend
 		// Transaction has been re-applied (redo) or rolled back (undo)
 		UndoRedo,
 	};
+	
+	struct FDocumentArrayPagedInputArgs
+	{
+		int32 Index = INDEX_NONE;
+		FGuid PageID;
+		EMetaSoundLiteralChangeType ChangeType = EMetaSoundLiteralChangeType::Invalid;
+	};
+	DECLARE_MULTICAST_DELEGATE_OneParam(FOnMetaSoundFrontendDocumentMutateInput, const FDocumentArrayPagedInputArgs& /* Args */);
 
 	struct FDocumentConfigChangedArgs
 	{
@@ -77,7 +103,7 @@ namespace Metasound::Frontend
 		FOnMetaSoundFrontendDocumentMutateArray OnRemovingInput;
 		FOnMetaSoundFrontendDocumentRenameVertex OnInputNameChanged;
 		FOnMetaSoundFrontendDocumentMutateArray OnInputDataTypeChanged;
-		FOnMetaSoundFrontendDocumentMutateArray OnInputDefaultChanged;
+		FOnMetaSoundFrontendDocumentMutateInput OnInputDefaultChanged;
 		FOnMetaSoundFrontendDocumentMutateArray OnInputIsConstructorPinChanged;
 		FOnMetaSoundFrontendDocumentMutateArray OnInputInheritsDefaultChanged;
 
