[MetaSound / Verse]:  Avoid hitting cached metasound input parameters during verse asset digest generation

#Synced-CL 49958372

diff --git a/Engine/Plugins/Runtime/Metasound/Source/MetasoundEngine/Private/MetasoundSource.cpp b/Engine/Plugins/Runtime/Metasound/Source/MetasoundEngine/Private/MetasoundSource.cpp
--- a/Engine/Plugins/Runtime/Metasound/Source/MetasoundEngine/Private/MetasoundSource.cpp
+++ b/Engine/Plugins/Runtime/Metasound/Source/MetasoundEngine/Private/MetasoundSource.cpp
@@ -1184,6 +1184,25 @@ void UMetaSoundSource::OnEndGenerate(ISoundGeneratorPtr Generator)
 	ForgetGenerator(Generator);
 }
 
+TArray<UMetaSoundSource::FRuntimeInput> UMetaSoundSource::GatherVerseParameters() const
+{
+	using namespace Metasound;
+	using namespace Metasound::Frontend;
+	using namespace Metasound::Engine;
+
+	constexpr bool bCreateUObjectProxies = false;
+	TSortedVertexNameMap<FRuntimeInput> InputMap = CreateRuntimeInputMap(bCreateUObjectProxies, UMetaSoundSettings::GetPageOrder());
+
+	TArray<FRuntimeInput> OutParameters;
+	OutParameters.Reserve(InputMap.Num());
+	for(const TPair<FVertexName, FRuntimeInput>& Pair : InputMap)
+	{
+		OutParameters.Add(Pair.Value);
+	}
+
+	return OutParameters;
+}
+
 bool UMetaSoundSource::GetAllDefaultParameters(TArray<FAudioParameter>& OutParameters) const
 {
 	using namespace Metasound;
@@ -1828,30 +1847,37 @@ UMetaSoundSource::FRuntimeInput UMetaSoundSource::CreateRuntimeInput(const Metas
 	using namespace Metasound;
 	using namespace Metasound::Frontend;
 
-	bool bIsTransmittable = false;
-	if (const IDataTypeRegistryEntry* RegistryEntry = Registry.FindDataTypeRegistryEntry(Input.TypeName))
-	{
-		bIsTransmittable = RegistryEntry->GetDataTypeInfo().bIsTransmittable;
-	}
-	else
-	{
-		UE_LOG(LogMetaSound, Warning, TEXT("Failed to find data type '%s' in registry. Assuming data type is not transmittable"), *Input.TypeName.ToString());
-	}
+	const TArray<FMetasoundFrontendClassInputDefault>& InputDefualts = Input.GetDefaults();
+	const bool bIsPaged = InputDefualts.Num() != 1 || InputDefualts.Last().PageID != Frontend::DefaultPageID;
 
-	if (const FMetasoundFrontendClassInputDefault* Default = FindPreferredPage(Input.GetDefaults(), InPageOrder))
-	{
-		FAudioParameter DefaultParameter = SourcePrivate::MakeAudioParameter(Registry, Input.Name, Input.TypeName, Default->Literal, bCreateUObjectProxies);
+	const IDataTypeRegistryEntry* RegistryEntry = Registry.FindDataTypeRegistryEntry(Input.TypeName);
+	const FMetasoundFrontendClassInputDefault* Default = FindPreferredPage(InputDefualts, InPageOrder);
 
-		return FRuntimeInput { Input.Name, Input.TypeName, Input.AccessType, DefaultParameter, bIsTransmittable };
-	}
-	else
+	UE_CLOG(!RegistryEntry, LogMetaSound, Warning, TEXT("Failed to find data type '%s' in registry. Assuming data type is not transmittable"), *Input.TypeName.ToString());
+	UE_CLOG(!Default, LogMetaSound, Error, TEXT("Failed to find default literal for parameter %s"), *Input.Name.ToString());
+
+	const bool bIsTransmittable = RegistryEntry? RegistryEntry->GetDataTypeInfo().bIsTransmittable : false;
+
+	const FAudioParameter DefaultParameter = SourcePrivate::MakeAudioParameter
+	(
+		  Registry
+		, Input.Name
+		, Input.TypeName
+		, (Default)? Default->Literal : FMetasoundFrontendLiteral(Registry.CreateDefaultLiteral(Input.TypeName))
+		, bCreateUObjectProxies
+	);
+
+	return FRuntimeInput
 	{
-		UE_LOG(LogMetaSound, Error, TEXT("Failed to find default literal for parameter %s"), *Input.Name.ToString());
-		FMetasoundFrontendLiteral FrontendLiteral;
-		FrontendLiteral.SetFromLiteral(Registry.CreateDefaultLiteral(Input.TypeName));
-		FAudioParameter DefaultParameter = SourcePrivate::MakeAudioParameter(Registry, Input.Name, Input.TypeName, FrontendLiteral, bCreateUObjectProxies);
-		return FRuntimeInput { Input.Name, Input.TypeName, Input.AccessType, DefaultParameter, bIsTransmittable };
-	}
+		  Input.Name
+		, Input.TypeName
+		, Input.AccessType
+		, DefaultParameter
+		, bIsTransmittable
+#if WITH_EDITORONLY_DATA
+		, bIsPaged
+#endif // #if WITH_EDITORONLY_DATA
+	};
 }
 
 Metasound::TSortedVertexNameMap<UMetaSoundSource::FRuntimeInput> UMetaSoundSource::CreateRuntimeInputMap(bool bCreateUObjectProxies, TArrayView<const FGuid> InPageOrder) const
diff --git a/Engine/Plugins/Runtime/Metasound/Source/MetasoundEngine/Public/MetasoundSource.h b/Engine/Plugins/Runtime/Metasound/Source/MetasoundEngine/Public/MetasoundSource.h
--- a/Engine/Plugins/Runtime/Metasound/Source/MetasoundEngine/Public/MetasoundSource.h
+++ b/Engine/Plugins/Runtime/Metasound/Source/MetasoundEngine/Public/MetasoundSource.h
@@ -95,8 +95,9 @@ class UMetaSoundSource : public USoundWaveProcedural, public FMetasoundAssetBase
 	//Forward declare
 	class FAudioParameterCollector;
 
+public:
 	// FRuntimeInput represents an input to a MetaSound which can be manipulated.
-	struct FRuntimeInput
+	struct UE_INTERNAL FRuntimeInput
 	{
 		// Name of input vertex
 		FName Name;
@@ -108,12 +109,19 @@ class UMetaSoundSource : public USoundWaveProcedural, public FMetasoundAssetBase
 		FAudioParameter DefaultParameter;
 		// True if the data type is transmittable. False otherwise.
 		bool bIsTransmittable;
+#if WITH_EDITORONLY_DATA
+		// True if the data / default value deviates from the default page (graph)
+		bool bIsPaged = false;
+#endif // #if WITH_EDITORONLY_DATA
 	};
+private:
 
 	struct FRuntimeInputData
 	{
 		std::atomic<bool> bIsValid = false;
+		PRAGMA_DISABLE_INTERNAL_WARNINGS
 		Metasound::TSortedVertexNameMap<FRuntimeInput> InputMap;
+		PRAGMA_ENABLE_INTERNAL_WARNINGS
 	};
 
 protected:
@@ -324,6 +332,10 @@ public:
 	UE_API virtual const UClass& GetBuilderUClass() const final override;
 	UE_API virtual EMetasoundFrontendClassAccessFlags GetDefaultAccessFlags() const final override;
 
+	PRAGMA_DISABLE_INTERNAL_WARNINGS
+	UE_INTERNAL UE_API TArray<FRuntimeInput> GatherVerseParameters() const;
+	PRAGMA_ENABLE_INTERNAL_WARNINGS
+
 protected:
 	UE_DEPRECATED(5.6, "AccessPtrs are actively being deprecated. Writable access outside of the builder API "
 		"is particularly problematic as in so accessing, the builder's caches are reset which can cause major "
@@ -347,7 +359,9 @@ private:
 	UE_API virtual void OnBeginActiveBuilder() override;
 	UE_API virtual void OnFinishActiveBuilder() override;
 
+	PRAGMA_DISABLE_INTERNAL_WARNINGS
 	UE_API void InitParametersInternal(const Metasound::TSortedVertexNameMap<FRuntimeInput>& InputMap, TArray<FAudioParameter>& ParametersToInit, FName InFeatureName) const;
+	PRAGMA_ENABLE_INTERNAL_WARNINGS
 	UE_API bool IsParameterValidInternal(const FAudioParameter& InParameter, const FName& InTypeName, Metasound::Frontend::IDataTypeRegistry& InDataTypeRegistry) const;
 
 	static UE_API Metasound::SourcePrivate::FParameterRouter& GetParameterRouter();
@@ -379,8 +393,10 @@ private:
 	UE_API void TrackGenerator(Metasound::FGeneratorInstanceInfo&& GeneratorInfo);
 	UE_API void ForgetGenerator(ISoundGeneratorPtr Generator);
 
+	PRAGMA_DISABLE_INTERNAL_WARNINGS
 	static UE_API FRuntimeInput CreateRuntimeInput(const Metasound::Frontend::IDataTypeRegistry& Registry, const FMetasoundFrontendClassInput& Input, bool bCreateUObjectProxies, TArrayView<const FGuid> InPageOrder);
 	UE_API Metasound::TSortedVertexNameMap<FRuntimeInput> CreateRuntimeInputMap(bool bCreateUObjectProxies, TArrayView<const FGuid> InPageOrder) const;
+	PRAGMA_ENABLE_INTERNAL_WARNINGS
 	UE_API void CacheRuntimeInputData(TArrayView<const FGuid> InPageOrder);
 	UE_API void InvalidateCachedRuntimeInputData();
 
diff --git a/Engine/Plugins/Runtime/Metasound/Source/MetasoundFrontend/Private/MetasoundFrontendLiteral.cpp b/Engine/Plugins/Runtime/Metasound/Source/MetasoundFrontend/Private/MetasoundFrontendLiteral.cpp
--- a/Engine/Plugins/Runtime/Metasound/Source/MetasoundFrontend/Private/MetasoundFrontendLiteral.cpp
+++ b/Engine/Plugins/Runtime/Metasound/Source/MetasoundFrontend/Private/MetasoundFrontendLiteral.cpp
@@ -82,6 +82,11 @@ namespace MetasoundFrontendLiteralPrivate
 	}
 }
 
+FMetasoundFrontendLiteral::FMetasoundFrontendLiteral(const Metasound::FLiteral& InLiteral)
+{
+	this->SetFromLiteral(InLiteral);
+}
+
 FMetasoundFrontendLiteral::FMetasoundFrontendLiteral(const FAudioParameter& InParameter)
 {
 	switch (InParameter.ParamType)
diff --git a/Engine/Plugins/Runtime/Metasound/Source/MetasoundFrontend/Public/MetasoundFrontendLiteral.h b/Engine/Plugins/Runtime/Metasound/Source/MetasoundFrontend/Public/MetasoundFrontendLiteral.h
--- a/Engine/Plugins/Runtime/Metasound/Source/MetasoundFrontend/Public/MetasoundFrontendLiteral.h
+++ b/Engine/Plugins/Runtime/Metasound/Source/MetasoundFrontend/Public/MetasoundFrontendLiteral.h
@@ -77,6 +77,8 @@ struct FMetasoundFrontendLiteral
 
 	FMetasoundFrontendLiteral() = default;
 
+	UE_API FMetasoundFrontendLiteral(const Metasound::FLiteral& InLiteral);
+
 	UE_API FMetasoundFrontendLiteral(const FAudioParameter& InParameter);
 
 private:
