Add MetaSound asset action to reassign class name and fixup references

#Synced-CL 50225261

diff --git a/Engine/Plugins/Runtime/Metasound/Source/MetasoundEditor/Private/MetasoundAssetDefinitions.cpp b/Engine/Plugins/Runtime/Metasound/Source/MetasoundEditor/Private/MetasoundAssetDefinitions.cpp
--- a/Engine/Plugins/Runtime/Metasound/Source/MetasoundEditor/Private/MetasoundAssetDefinitions.cpp
+++ b/Engine/Plugins/Runtime/Metasound/Source/MetasoundEditor/Private/MetasoundAssetDefinitions.cpp
@@ -165,7 +165,112 @@ namespace Metasound::Editor
 				}
 			}
 		}
+		
+		void ExecuteReassignClassName(const FToolMenuContext& MenuContext)
+		{
+			using namespace Frontend;
+			using namespace Engine;
+			if (const UContentBrowserAssetContextMenuContext* Context = UContentBrowserAssetContextMenuContext::FindContextWithAssets(MenuContext))
+			{
+				FMetaSoundAssetManager& AssetManager = FMetaSoundAssetManager::GetChecked();
+				UAssetEditorSubsystem* AssetEditorSubsystem = GEditor->GetEditorSubsystem<UAssetEditorSubsystem>();
+				check(AssetEditorSubsystem);
+				const IAssetRegistry& AssetRegistry = IAssetRegistry::GetChecked();
+
+				TArray<UPackage*> PackagesToCheckout;
+				
+				// MetaSound asset filter
+				TArray<FTopLevelAssetPath> ClassNames;
+				IMetasoundUObjectRegistry::Get().IterateRegisteredUClasses([&ClassNames](UClass& InClass)
+				{
+					ClassNames.Add(InClass.GetClassPathName());
+				});
+
+				FARFilter MetaSoundFilter;
+				MetaSoundFilter.ClassPaths = ClassNames;
+				MetaSoundFilter.bRecursiveClasses = true;
+
+				// Get selected objects
+				const TArray<UObject*> Objects = Context->LoadSelectedObjects<UObject>();
+				for (UObject* SelectedObject : Objects)
+				{
+					if (!SelectedObject)
+					{
+						continue;
+					}
+					
+					// Get all assets that reference this one, filtered to MetaSounds
+					const FTopLevelAssetPath MetaSoundPath(SelectedObject);
+					TArray<FName> OnDiskReferencers;
+					AssetRegistry.GetReferencers(MetaSoundPath.GetPackageName(), OnDiskReferencers);
+					
+					TArray<FAssetData> MetaSoundReferencers;
+					if (!OnDiskReferencers.IsEmpty())
+					{
+						MetaSoundFilter.PackageNames = OnDiskReferencers;
+						AssetRegistry.GetAssets(MetaSoundFilter, MetaSoundReferencers);
+					}
+
+					TArray<UObject*> ReferencingObjects;
+					Algo::TransformIf(MetaSoundReferencers, ReferencingObjects,
+					[](const FAssetData& InAssetData)
+					{
+						return InAssetData.GetAsset() != nullptr;
+					},	
+					[](const FAssetData& InAssetData)
+					{
+						return InAssetData.GetAsset();
+					});
+
+					// Close open editors to prevent editor graph errors before reference fixup 
+					for (UObject* ReferencingObject : ReferencingObjects)
+					{
+						AssetEditorSubsystem->CloseAllEditorsForAsset(ReferencingObject);
+					}
+
+					// Reassign class name
+					const TScriptInterface<IMetaSoundDocumentInterface> DocInterface(SelectedObject);
+					if (const IMetaSoundDocumentInterface* Interface = DocInterface.GetInterface())
+					{
+						const FMetasoundFrontendClassMetadata& RootGraphMetadata = Interface->GetConstDocument().RootGraph.Metadata;
+						const FMetasoundFrontendVersionNumber& Version = RootGraphMetadata.GetVersion();
+						const FNodeRegistryKey OldRegistryKey(EMetasoundFrontendClassType::External, RootGraphMetadata.GetClassName(), Version);
+						
+						if (AssetManager.ReassignClassName(DocInterface))
+						{
+							SelectedObject->MarkPackageDirty();
+							PackagesToCheckout.Emplace(SelectedObject->GetPackage());
 
+							// Replace references in objects referencing the asset
+							const FNodeRegistryKey NewRegistryKey(EMetasoundFrontendClassType::External, RootGraphMetadata.GetClassName(), Version);
+							for (UObject* ReferencingObject : ReferencingObjects)
+							{
+								if (AssetManager.ReplaceReferencesInAsset(ReferencingObject, OldRegistryKey, NewRegistryKey))
+								{
+									ReferencingObject->MarkPackageDirty();
+									PackagesToCheckout.Emplace(ReferencingObject->GetPackage());
+								}
+								else
+								{
+									UE_LOG(LogMetaSound, Warning, TEXT("Could not replace references in MetaSound '%s' after reassigning class name of MetaSound '%s'"), *ReferencingObject->GetFullName(), *SelectedObject->GetFullName());								
+								}
+							}
+						}
+						else
+						{
+							UE_LOG(LogMetaSound, Warning, TEXT("Could not reassign class name of MetaSound '%s'"), *SelectedObject->GetFullName());								
+						}
+					}
+				}
+
+				// Save/checkout
+				if (!PackagesToCheckout.IsEmpty())
+				{
+					FEditorFileUtils::PromptToCheckoutPackages(false, PackagesToCheckout);
+				}
+			}
+		}
+		
 		void RegisterAndSetObjectSettings(IMetasoundEditorModule& MetaSoundEditorModule, UMetaSoundEditorSubsystem& EditorSubsystem, UMetaSoundPatch& NewMetaSound, UMetaSoundPatch& ParentMetaSound)
 		{
 			if (MetaSoundEditorModule.IsRestrictedMode())
@@ -260,9 +365,9 @@ namespace Metasound::Editor
 						{
 							constexpr bool bCheckDirty = false;
 							constexpr bool bPromptToSave = false;
-							TArray<UPackage*> OutermostPackagesToSave;
-							OutermostPackagesToSave.Add(InObject->GetOutermost());
-							FEditorFileUtils::PromptForCheckoutAndSave(OutermostPackagesToSave, bCheckDirty, bPromptToSave);
+							TArray<UPackage*> OutermostPackagesToCheckout;
+							OutermostPackagesToCheckout.Add(InObject->GetOutermost());
+							FEditorFileUtils::PromptForCheckoutAndSave(OutermostPackagesToCheckout, bCheckDirty, bPromptToSave);
 						}
 					};
 				}
@@ -1037,6 +1142,31 @@ namespace MenuExtension_MetaSoundTemplate
 				check(MetaSoundClass);
 				AssetDefinitionsPrivate::AddPresetActions(AssetIcon, *MetaSoundClass, *MetaSoundClass, MoveTemp(CreateAction), FPresetDocumentConfiguration::StaticStruct()->GetDisplayNameText());
 			}
+
+			// Asset actions
+			{
+				const FNewToolMenuSectionDelegate AssetActionsSectionDelegate = FNewToolMenuSectionDelegate::CreateLambda([](FToolMenuSection& InSection)
+				{
+					InSection.AddMenuEntry(
+						"ReassignClassName",
+						LOCTEXT("ReassignClassNameTitle", "Reassign Asset Class Guid"),
+						LOCTEXT("ReassignClassNameTooltipText", "Reassigns class name (asset class guid) on selected assets and fixes up asset references.\n" \
+									"Useful for fixing up asset key errors caused by MetaSound assets with duplicate class names, \n " \
+									"which can occur when assets are duplicated outside of the editor."),
+						FSlateIcon(),
+						FToolMenuExecuteAction::CreateStatic(&AssetDefinitionsPrivate::ExecuteReassignClassName)
+					);
+				});
+
+				Metasound::IMetasoundUObjectRegistry::Get().IterateRegisteredUClasses([&AssetActionsSectionDelegate](UClass& InClass)
+				{
+					const FString MenuName = TEXT("ContentBrowser.AssetContextMenu.") + InClass.GetName() + TEXT(".AssetActionsSubMenu");
+					UToolMenu* AssetActionsMenu = UToolMenus::Get()->ExtendMenu(FName(MenuName));
+					FToolMenuSection& MetaSoundAssetActionsMenu = AssetActionsMenu->AddSection("MetaSoundAssetContextActions", LOCTEXT("MetaSoundAssetContextActionsMenuHeading", "MetaSound"));
+
+					MetaSoundAssetActionsMenu.AddDynamicEntry(NAME_None, AssetActionsSectionDelegate);
+				});
+			}
  		}));
 	});
 } // namespace MenuExtension_MetaSoundSourceTemplate
diff --git a/Engine/Plugins/Runtime/Metasound/Source/MetasoundEngine/Private/MetasoundAssetSubsystem.cpp b/Engine/Plugins/Runtime/Metasound/Source/MetasoundEngine/Private/MetasoundAssetSubsystem.cpp
--- a/Engine/Plugins/Runtime/Metasound/Source/MetasoundEngine/Private/MetasoundAssetSubsystem.cpp
+++ b/Engine/Plugins/Runtime/Metasound/Source/MetasoundEngine/Private/MetasoundAssetSubsystem.cpp
@@ -110,7 +110,7 @@ namespace Metasound::Engine
 				TArray<FString> PathStrings;
 				Algo::Transform(TagDatas, PathStrings, [](const FMetaSoundAssetClassInfo& WarnInfo) { return WarnInfo.AssetPath.ToString(); });
 				UE_LOG(LogMetaSound, Warning,
-					TEXT("MetaSoundAssetManager has registered multiple assets with key '%s':\n%s\n"),
+					TEXT("MetaSoundAssetManager has registered multiple assets with key '%s':\n%s\n. Reassign the class name of colliding assets by using the Reassign Asset Class Guid asset action in the right click -> Asset Actions menu of the MetaSound."),
 					*AssetKey.ToString(),
 					*FString::Join(PathStrings, TEXT("\n")));
 			}
@@ -1056,6 +1056,54 @@ PRAGMA_ENABLE_DEPRECATION_WARNINGS
 	}
 
 #if WITH_EDITOR
+	bool FMetaSoundAssetManager::ReplaceReferencesInAssetInternal(UObject* MetaSoundObject, const TMap<FNodeRegistryKey, FNodeRegistryKey>& InOldToNewReferenceKeys) const
+	{
+		MetaSoundObject->Modify();
+		FMetaSoundFrontendDocumentBuilder& Builder = FDocumentBuilderRegistry::GetChecked().FindOrBeginBuilding(MetaSoundObject);
+		const bool bDependencyUpdated = Builder.UpdateDependencyRegistryData(InOldToNewReferenceKeys);
+		if (bDependencyUpdated)
+		{
+			Builder.RemoveUnusedDependencies();
+			if (FMetasoundAssetBase* AssetBase = GetAsAsset(*MetaSoundObject); ensure(AssetBase))
+			{
+				AssetBase->RebuildReferencedAssetClasses();
+			}
+			return true;
+		}
+		return false;
+	}
+	
+	bool FMetaSoundAssetManager::ReplaceReferencesInAsset(TScriptInterface<IMetaSoundDocumentInterface> InMetaSound, const Frontend::FNodeRegistryKey& OldClassKey, const Frontend::FNodeRegistryKey& NewClassKey) const
+	{
+		using namespace Frontend;
+
+		bool bReferencesReplaced = false;
+
+#if WITH_EDITORONLY_DATA
+		if (!NewClassKey.IsValid())
+		{
+			return bReferencesReplaced;
+		}
+		
+		if (UObject* MetaSoundObject = InMetaSound.GetObject())
+		{
+			FMetaSoundClassInfo ClassInfo;
+			const bool bNewClassExists = ISearchEngine::Get().FindClassWithHighestVersion(NewClassKey.ClassName, ClassInfo);
+			if (bNewClassExists)
+			{
+				TMap<FNodeRegistryKey, FNodeRegistryKey> OldToNewReferenceKeys = { { OldClassKey, NewClassKey } };
+				bReferencesReplaced = ReplaceReferencesInAssetInternal(MetaSoundObject, OldToNewReferenceKeys);
+			}
+			else
+			{
+				UE_LOG(LogMetaSound, Display, TEXT("Cannot replace references in MetaSound asset '%s': NewClass '%s' does not exist"), *InMetaSound.GetObject()->GetFullName(), *NewClassKey.ToString());
+			}
+		}
+		
+#endif // WITH_EDITORONLY_DATA
+		return bReferencesReplaced;
+	}
+	
 	bool FMetaSoundAssetManager::ReplaceReferencesInDirectory(const TArray<FMetaSoundAssetDirectory>& InDirectories, const Frontend::FNodeRegistryKey& OldClassKey, const Frontend::FNodeRegistryKey& NewClassKey) const
 	{
 		using namespace Frontend;
@@ -1080,18 +1128,7 @@ PRAGMA_ENABLE_DEPRECATION_WARNINGS
 			{
 				if (UObject* MetaSoundObject = AssetData.GetAsset())
 				{
-					MetaSoundObject->Modify();
-					FMetaSoundFrontendDocumentBuilder& Builder = FDocumentBuilderRegistry::GetChecked().FindOrBeginBuilding(MetaSoundObject);
-					const bool bDependencyUpdated = Builder.UpdateDependencyRegistryData(OldToNewReferenceKeys);
-					if (bDependencyUpdated)
-					{
-						bReferencesReplaced = true;
-						Builder.RemoveUnusedDependencies();
-						if (FMetasoundAssetBase* AssetBase = GetAsAsset(*MetaSoundObject); ensure(AssetBase))
-						{
-							AssetBase->RebuildReferencedAssetClasses();
-						}
-					}
+					bReferencesReplaced |= ReplaceReferencesInAssetInternal(MetaSoundObject, OldToNewReferenceKeys);
 				}
 			});
 		}
diff --git a/Engine/Plugins/Runtime/Metasound/Source/MetasoundEngine/Public/MetasoundAssetSubsystem.h b/Engine/Plugins/Runtime/Metasound/Source/MetasoundEngine/Public/MetasoundAssetSubsystem.h
--- a/Engine/Plugins/Runtime/Metasound/Source/MetasoundEngine/Public/MetasoundAssetSubsystem.h
+++ b/Engine/Plugins/Runtime/Metasound/Source/MetasoundEngine/Public/MetasoundAssetSubsystem.h
@@ -96,6 +96,7 @@ PRAGMA_ENABLE_DEPRECATION_WARNINGS
 		UE_API void RegisterAssetClassesInDirectories(const TArray<FMetaSoundAssetDirectory>& Directories);
 
 #if WITH_EDITOR
+		UE_API bool ReplaceReferencesInAsset(TScriptInterface<IMetaSoundDocumentInterface> InMetaSound, const Frontend::FNodeRegistryKey& OldClassKey, const Frontend::FNodeRegistryKey& NewClassKey) const;
 		UE_API bool ReplaceReferencesInDirectory(const TArray<FMetaSoundAssetDirectory>& InDirectories, const FNodeRegistryKey& OldClassKey, const FNodeRegistryKey& NewClassKey) const;
 #endif // WITH_EDITOR
 
@@ -188,10 +189,15 @@ PRAGMA_ENABLE_DEPRECATION_WARNINGS
 		FMetaSoundAsyncAssetDependencies* FindLoadingDependencies(int32 InLoadID);
 		void RemoveLoadingDependencies(int32 InLoadID);
 		void OnReferencedAssetsLoaded(int32 InLoadID);
+
 #endif // WITH_EDITORONLY_DATA
 
 		bool GetReferencedPresetHierarchyInternal(FMetasoundAssetBase& InAsset, TArray<FMetasoundAssetBase*>& OutReferencedAssets) const;
 
+#if WITH_EDITOR
+		bool ReplaceReferencesInAssetInternal(UObject* MetaSoundObject, const TMap<FNodeRegistryKey, FNodeRegistryKey>& InOldToNewReferenceKeys) const;
+#endif // WITH_EDITOR
+
 		struct FPackageLoadedArgs
 		{
 			FName PackageName;
