Move errors to validation so tests pass accordingly.  Logging errors at configure time is ultimately up to the configuration author.  Generally speaking, should be encouraged to log errors during asset validation. Any configuration transform must should result in a valid asset produced by design input, so logging errors is better suited for validation.

diff --git a/Engine/Plugins/Runtime/Metasound/Source/MetasoundEditor/Private/MetasoundDetailCustomization.cpp b/Engine/Plugins/Runtime/Metasound/Source/MetasoundEditor/Private/MetasoundDetailCustomization.cpp
--- a/Engine/Plugins/Runtime/Metasound/Source/MetasoundEditor/Private/MetasoundDetailCustomization.cpp
+++ b/Engine/Plugins/Runtime/Metasound/Source/MetasoundEditor/Private/MetasoundDetailCustomization.cpp
@@ -194,14 +194,14 @@ namespace Metasound::Editor
 					TSharedPtr<IPropertyHandle> DocConfigHandle = DetailLayout.GetProperty(GetDocConfigPropertyPath());
 					if (DocConfigHandle.IsValid())
 					{
-						FSimpleDelegate OnConfigChanged = FSimpleDelegate::CreateLambda([this]()
+						FSimpleDelegate OnConfigChanged = FSimpleDelegate::CreateSPLambda(AsShared(), [this]()
 						{
 							if (Builder.IsValid())
 							{
 								FMetaSoundFrontendDocumentBuilder& DocBuilder = Builder->GetBuilder();
 								if (DocBuilder.IsValid())
 								{
-									DocBuilder.ConfigureDocument(true /* bLogFailure */);
+									DocBuilder.ConfigureDocument();
 									DocBuilder.GetDocumentDelegates().OnDocumentConfigurationChanged.Broadcast({ EDocumentConfigChangeType::Struct });
 								}
 							}
@@ -217,7 +217,7 @@ namespace Metasound::Editor
 								{
 									if (FDocumentConfiguration* Config = DocBuilder.GetDocumentConfiguration().GetPtr())
 									{
-										DocBuilder.ConfigureDocument(true /* bLogFailure */);
+										DocBuilder.ConfigureDocument();
 										Config->OnPropertyChanged(InEvent, DocBuilder);
 										DocBuilder.GetDocumentDelegates().OnDocumentConfigurationChanged.Broadcast({ EDocumentConfigChangeType::Property });
 									}
diff --git a/Engine/Plugins/Runtime/Metasound/Source/MetasoundEditor/Private/MetasoundEditor.cpp b/Engine/Plugins/Runtime/Metasound/Source/MetasoundEditor/Private/MetasoundEditor.cpp
--- a/Engine/Plugins/Runtime/Metasound/Source/MetasoundEditor/Private/MetasoundEditor.cpp
+++ b/Engine/Plugins/Runtime/Metasound/Source/MetasoundEditor/Private/MetasoundEditor.cpp
@@ -2558,9 +2558,8 @@ namespace Metasound
 							check(ConfigStruct);
 							ConfigInstance.InitializeAsScriptStruct(ConfigStruct);
 
-							constexpr bool bLogErrors = true;
 							DocBuilder->SetDocumentConfiguration(MoveTemp(ConfigInstance));
-							DocBuilder->ConfigureDocument(bLogErrors);
+							DocBuilder->ConfigureDocument();
 
 							ReloadAllDocConfigTabs();
 							UpdateEnabledTabs();
diff --git a/Engine/Plugins/Runtime/Metasound/Source/MetasoundEditor/Private/MetasoundEditorSubsystem.cpp b/Engine/Plugins/Runtime/Metasound/Source/MetasoundEditor/Private/MetasoundEditorSubsystem.cpp
--- a/Engine/Plugins/Runtime/Metasound/Source/MetasoundEditor/Private/MetasoundEditorSubsystem.cpp
+++ b/Engine/Plugins/Runtime/Metasound/Source/MetasoundEditor/Private/MetasoundEditorSubsystem.cpp
@@ -383,7 +383,7 @@ void UMetaSoundEditorSubsystem::InitAsset(UObject& InNewMetaSound, UMetaSoundEdi
 		FMetaSoundFrontendDocumentBuilder& DocBuilder = Frontend::IDocumentBuilderRegistry::GetChecked().FindOrBeginBuilding(&InNewMetaSound);
 		DocBuilder.SetDocumentConfiguration(MoveTemp(InitAssetArgs.DocumentConfiguration));
 		DocConfig->OnAssetInitialized(MoveTemp(InitAssetArgs.SelectedObjects), DocBuilder);
-		DocBuilder.ConfigureDocument(true /* bLogFailure */);
+		DocBuilder.ConfigureDocument();
 	}
 }
 
diff --git a/Engine/Plugins/Runtime/Metasound/Source/MetasoundEngine/Private/DocumentConfigurations/MetasoundRandomizerDocumentConfiguration.cpp b/Engine/Plugins/Runtime/Metasound/Source/MetasoundEngine/Private/DocumentConfigurations/MetasoundRandomizerDocumentConfiguration.cpp
--- a/Engine/Plugins/Runtime/Metasound/Source/MetasoundEngine/Private/DocumentConfigurations/MetasoundRandomizerDocumentConfiguration.cpp
+++ b/Engine/Plugins/Runtime/Metasound/Source/MetasoundEngine/Private/DocumentConfigurations/MetasoundRandomizerDocumentConfiguration.cpp
@@ -193,7 +193,7 @@ namespace Metasound::Engine
 } // namespace Metasound::Engine
 
 #if WITH_EDITORONLY_DATA
-bool FMetaSoundRandomizerDocumentConfiguration::ConfigureDocument(FMetaSoundFrontendDocumentBuilder& OutBuilder, bool bLogFailure)
+bool FMetaSoundRandomizerDocumentConfiguration::ConfigureDocument(FMetaSoundFrontendDocumentBuilder& OutBuilder)
 {
 	using namespace Metasound;
 	using namespace Metasound::Engine;
diff --git a/Engine/Plugins/Runtime/Metasound/Source/MetasoundEngine/Private/MetasoundBuilderBase.cpp b/Engine/Plugins/Runtime/Metasound/Source/MetasoundEngine/Private/MetasoundBuilderBase.cpp
--- a/Engine/Plugins/Runtime/Metasound/Source/MetasoundEngine/Private/MetasoundBuilderBase.cpp
+++ b/Engine/Plugins/Runtime/Metasound/Source/MetasoundEngine/Private/MetasoundBuilderBase.cpp
@@ -1354,7 +1354,7 @@ void UMetaSoundBuilderBase::SetConfiguration(TInstancedStruct<FMetaSoundFrontend
 	Builder.SetDocumentConfiguration(MoveTemp(Configuration));
 
 PRAGMA_DISABLE_EXPERIMENTAL_WARNINGS
-	Builder.ConfigureDocument(true /* bLogFailure */);
+	Builder.ConfigureDocument();
 PRAGMA_ENABLE_EXPERIMENTAL_WARNINGS
 }
 
@@ -1363,7 +1363,7 @@ void UMetaSoundBuilderBase::SetConfigurationProperties(TFunctionRef<void(TStruct
 	Builder.SetDocumentConfigurationProperties(SetPropertyFunction);
 
 	PRAGMA_DISABLE_EXPERIMENTAL_WARNINGS
-	Builder.ConfigureDocument(true /* bLogFailure */);
+	Builder.ConfigureDocument();
 	PRAGMA_ENABLE_EXPERIMENTAL_WARNINGS
 }
 
diff --git a/Engine/Plugins/Runtime/Metasound/Source/MetasoundEngine/Private/MetasoundEngineAsset.h b/Engine/Plugins/Runtime/Metasound/Source/MetasoundEngine/Private/MetasoundEngineAsset.h
--- a/Engine/Plugins/Runtime/Metasound/Source/MetasoundEngine/Private/MetasoundEngineAsset.h
+++ b/Engine/Plugins/Runtime/Metasound/Source/MetasoundEngine/Private/MetasoundEngineAsset.h
@@ -204,6 +204,15 @@ namespace Metasound::Engine
 					ErrorIfMissing(PageID, FText::Format(LOCTEXT("InputPageDefaultDescriptorFormat", "input '{0}' default value"), FText::FromName(ClassInput.Name)));
 				});
 			}
+
+			if (const FMetaSoundFrontendDocumentConfiguration* DocConfig = Document.Configuration.GetPtr())
+			{
+				FDocumentBuilderRegistry& BuilderRegistry = FDocumentBuilderRegistry::GetChecked();
+				UObject* MetaSoundObject = FAssetData(&MetaSound).GetAsset();
+				check(MetaSoundObject);
+				const UMetaSoundBuilderBase& Builder = BuilderRegistry.FindOrBeginBuilding(*MetaSoundObject);
+				CombineDataValidationResults(Result, DocConfig->IsDataValid(Builder.GetConstBuilder(), InOutContext));
+			}
 			return Result;
 		}
 
diff --git a/Engine/Plugins/Runtime/Metasound/Source/MetasoundEngine/Public/DocumentConfigurations/MetasoundRandomizerDocumentConfiguration.h b/Engine/Plugins/Runtime/Metasound/Source/MetasoundEngine/Public/DocumentConfigurations/MetasoundRandomizerDocumentConfiguration.h
--- a/Engine/Plugins/Runtime/Metasound/Source/MetasoundEngine/Public/DocumentConfigurations/MetasoundRandomizerDocumentConfiguration.h
+++ b/Engine/Plugins/Runtime/Metasound/Source/MetasoundEngine/Public/DocumentConfigurations/MetasoundRandomizerDocumentConfiguration.h
@@ -23,7 +23,7 @@ struct FMetaSoundRandomizerDocumentConfiguration : public FMetaSoundFrontendDocu
 	virtual ~FMetaSoundRandomizerDocumentConfiguration() = default;
 
 #if WITH_EDITORONLY_DATA
-	UE_INTERNAL UE_API virtual bool ConfigureDocument(FMetaSoundFrontendDocumentBuilder& OutBuilder, bool bLogFailure) override;
+	UE_INTERNAL UE_API virtual bool ConfigureDocument(FMetaSoundFrontendDocumentBuilder& OutBuilder) override;
 #endif // WITH_EDITORONLY_DATA
 
 #if WITH_EDITOR
diff --git a/Engine/Plugins/Runtime/Metasound/Source/MetasoundEngine/Public/MetasoundBuilderBase.h b/Engine/Plugins/Runtime/Metasound/Source/MetasoundEngine/Public/MetasoundBuilderBase.h
--- a/Engine/Plugins/Runtime/Metasound/Source/MetasoundEngine/Public/MetasoundBuilderBase.h
+++ b/Engine/Plugins/Runtime/Metasound/Source/MetasoundEngine/Public/MetasoundBuilderBase.h
@@ -492,7 +492,7 @@ public:
 		if (bConfigUpdated)
 		{
 			PRAGMA_DISABLE_EXPERIMENTAL_WARNINGS
-			GetBuilder().ConfigureDocument(true /* bLogFailure */);
+			GetBuilder().ConfigureDocument();
 			PRAGMA_ENABLE_EXPERIMENTAL_WARNINGS
 		}
 		return bConfigUpdated;
diff --git a/Engine/Plugins/Runtime/Metasound/Source/MetasoundFrontend/Private/DocumentConfigurations/MetasoundFrontendDocumentConfiguration.cpp b/Engine/Plugins/Runtime/Metasound/Source/MetasoundFrontend/Private/DocumentConfigurations/MetasoundFrontendDocumentConfiguration.cpp
--- a/Engine/Plugins/Runtime/Metasound/Source/MetasoundFrontend/Private/DocumentConfigurations/MetasoundFrontendDocumentConfiguration.cpp
+++ b/Engine/Plugins/Runtime/Metasound/Source/MetasoundFrontend/Private/DocumentConfigurations/MetasoundFrontendDocumentConfiguration.cpp
@@ -10,7 +10,7 @@
 
 #define LOCTEXT_NAMESPACE "MetaSound"
 
-bool FMetaSoundFrontendDocumentConfiguration::ConfigureDocument(FMetaSoundFrontendDocumentBuilder& OutBuilder, bool bLogFailure)
+bool FMetaSoundFrontendDocumentConfiguration::ConfigureDocument(FMetaSoundFrontendDocumentBuilder& OutBuilder)
 {
 #if WITH_EDITORONLY_DATA
 	OutBuilder.SetIsGraphEditable(false);
@@ -36,6 +36,11 @@ const FMetaSoundFrontendDocumentConfiguration::FEditorOptions& FMetaSoundFronten
 	return FEditorOptions::GetDefaultOptions();
 }
 
+EDataValidationResult FMetaSoundFrontendDocumentConfiguration::IsDataValid(const FMetaSoundFrontendDocumentBuilder& Builder, FDataValidationContext& InOutContext) const
+{
+	return EDataValidationResult::Valid;
+}
+
 void FMetaSoundFrontendDocumentConfiguration::OnAssetInitialized(TArray<UObject*> SelectedObjects, FMetaSoundFrontendDocumentBuilder& OutBuilder)
 {
 }
diff --git a/Engine/Plugins/Runtime/Metasound/Source/MetasoundFrontend/Private/DocumentConfigurations/MetasoundFrontendPresetDocumentConfiguration.cpp b/Engine/Plugins/Runtime/Metasound/Source/MetasoundFrontend/Private/DocumentConfigurations/MetasoundFrontendPresetDocumentConfiguration.cpp
--- a/Engine/Plugins/Runtime/Metasound/Source/MetasoundFrontend/Private/DocumentConfigurations/MetasoundFrontendPresetDocumentConfiguration.cpp
+++ b/Engine/Plugins/Runtime/Metasound/Source/MetasoundFrontend/Private/DocumentConfigurations/MetasoundFrontendPresetDocumentConfiguration.cpp
@@ -68,7 +68,7 @@ namespace Metasound::Frontend::PresetDocConfigPrivate
 } // namespace Metasound::Frontend::PresetDocConfigPrivate
 
 #if WITH_EDITORONLY_DATA
-bool FMetaSoundFrontendPresetDocumentConfiguration::ConfigureDocument(FMetaSoundFrontendDocumentBuilder& OutBuilder, bool bLogFailure)
+bool FMetaSoundFrontendPresetDocumentConfiguration::ConfigureDocument(FMetaSoundFrontendDocumentBuilder& OutBuilder)
 {
 	using namespace Metasound::Frontend;
 
@@ -84,20 +84,34 @@ bool FMetaSoundFrontendPresetDocumentConfiguration::ConfigureDocument(FMetaSound
 			FRebuildPresetRootGraph RebuildPresetRootGraph(ParentBuilder);
 			return RebuildPresetRootGraph.Transform(OutBuilder);
 		}
-		else
+	}
+
+	return false;
+}
+
+EDataValidationResult FMetaSoundFrontendPresetDocumentConfiguration::IsDataValid(const FMetaSoundFrontendDocumentBuilder& Builder, FDataValidationContext& InOutContext) const
+{
+	using namespace Metasound::Frontend;
+
+	EDataValidationResult Result = EDataValidationResult::Valid;
+	if (UObject* RefObject = Parent.GetObject())
+	{
+		if (!RefObject->IsA(&Builder.GetConstDocumentInterfaceChecked().GetBaseMetaSoundUClass()))
 		{
-			if (bLogFailure)
-			{
-				UE_LOG(LogMetaSound, Error, TEXT("The referenced node type must match the base MetaSound class when converting builder '%s' to a preset (ex. source preset must reference another source)"), *OutBuilder.GetDebugName());
-			}
+			Result = EDataValidationResult::Invalid;
+			InOutContext.AddError(FText::Format(LOCTEXT("MetaSoundPreset_InvalidReferenceTypeFormat", "The referenced node type must match the base MetaSound class when converting "
+				"builder '{0}' to a preset (ex. source preset must reference another source)"),
+				FText::FromString(Builder.GetDebugName())
+			));
 		}
 	}
-	else if (bLogFailure)
+	else
 	{
-		UE_LOG(LogMetaSound, Error, TEXT("Preset '%s' generation failed: Referenced class not set."), *OutBuilder.GetDebugName());
+		Result = EDataValidationResult::Invalid;
+		InOutContext.AddError(FText::Format(LOCTEXT("MetaSoundPreset_ReferenceUnsetFormat", "Preset '%s' generation failed: Referenced class not set."), FText::FromString(Builder.GetDebugName())));
 	}
 
-	return false;
+	return Result;
 }
 #endif // WITH_EDITORONLY_DATA
 
diff --git a/Engine/Plugins/Runtime/Metasound/Source/MetasoundFrontend/Private/MetasoundFrontendDocumentBuilder.cpp b/Engine/Plugins/Runtime/Metasound/Source/MetasoundFrontend/Private/MetasoundFrontendDocumentBuilder.cpp
--- a/Engine/Plugins/Runtime/Metasound/Source/MetasoundFrontend/Private/MetasoundFrontendDocumentBuilder.cpp
+++ b/Engine/Plugins/Runtime/Metasound/Source/MetasoundFrontend/Private/MetasoundFrontendDocumentBuilder.cpp
@@ -2520,7 +2520,7 @@ bool FMetaSoundFrontendDocumentBuilder::ConformObjectToDocument()
 	return GetDocumentInterfaceChecked().ConformObjectToDocument();
 }
 
-bool FMetaSoundFrontendDocumentBuilder::ConfigureDocument(bool bLogFailure)
+bool FMetaSoundFrontendDocumentBuilder::ConfigureDocument()
 {
 	using namespace Metasound::Frontend;
 
@@ -2535,7 +2535,7 @@ bool FMetaSoundFrontendDocumentBuilder::ConfigureDocument(bool bLogFailure)
 	FMetasoundFrontendDocument& Document = GetDocumentChecked();
 	if (FMetaSoundFrontendDocumentConfiguration* DocConfig = Document.Configuration.GetMutablePtr())
 	{
-		if (DocConfig->ConfigureDocument(*this, bLogFailure))
+		if (DocConfig->ConfigureDocument(*this))
 		{
 			GetDocumentInterfaceChecked().ConformObjectToDocument();
 
@@ -2631,7 +2631,7 @@ bool FMetaSoundFrontendDocumentBuilder::ConvertToPreset(const FMetaSoundFrontend
 	FPresetConfigStruct NewConfig = FPresetConfigStruct::Make<FMetaSoundFrontendPresetDocumentConfiguration>();
 	NewConfig.GetMutable().Parent = &InReferencedDocumentBuilder.CastDocumentObjectChecked<UObject>();
 	SetDocumentConfiguration(MoveTemp(NewConfig));
-	ConfigureDocument(true /* bLogFailure */);
+	ConfigureDocument();
 
 	return true;
 }
diff --git a/Engine/Plugins/Runtime/Metasound/Source/MetasoundFrontend/Private/MetasoundFrontendTransform.cpp b/Engine/Plugins/Runtime/Metasound/Source/MetasoundFrontend/Private/MetasoundFrontendTransform.cpp
--- a/Engine/Plugins/Runtime/Metasound/Source/MetasoundFrontend/Private/MetasoundFrontendTransform.cpp
+++ b/Engine/Plugins/Runtime/Metasound/Source/MetasoundFrontend/Private/MetasoundFrontendTransform.cpp
@@ -509,7 +509,7 @@ namespace Metasound
 			// If document config provided, run configure
 			if (TConstStructView<FMetaSoundFrontendDocumentConfiguration> DocConfig = InOutBuilderToTransform.GetConstDocumentConfiguration(); DocConfig.IsValid())
 			{
-				InOutBuilderToTransform.ConfigureDocument(true /* bLogFailure */);
+				InOutBuilderToTransform.ConfigureDocument();
 				
 				TScriptInterface<IMetaSoundDocumentInterface> DocInterface(&InOutBuilderToTransform.CastDocumentObjectChecked<UObject>());
 				DocInterface->ConformObjectToDocument();
diff --git a/Engine/Plugins/Runtime/Metasound/Source/MetasoundFrontend/Public/DocumentConfigurations/MetasoundFrontendDocumentConfiguration.h b/Engine/Plugins/Runtime/Metasound/Source/MetasoundFrontend/Public/DocumentConfigurations/MetasoundFrontendDocumentConfiguration.h
--- a/Engine/Plugins/Runtime/Metasound/Source/MetasoundFrontend/Public/DocumentConfigurations/MetasoundFrontendDocumentConfiguration.h
+++ b/Engine/Plugins/Runtime/Metasound/Source/MetasoundFrontend/Public/DocumentConfigurations/MetasoundFrontendDocumentConfiguration.h
@@ -5,6 +5,7 @@
 
 #if WITH_EDITOR
 #include "Delegates/DelegateCombinations.h"
+#include "Misc/DataValidation.h"
 #endif // WITH_EDITOR
 
 #include "Internationalization/Text.h"
@@ -44,8 +45,9 @@ struct FMetaSoundFrontendDocumentConfiguration
 
 	virtual ~FMetaSoundFrontendDocumentConfiguration() = default;
 
-	// Executes document configuration. Returns true if document configuration was successful, false if not.
-	UE_INTERNAL UE_API virtual bool ConfigureDocument(FMetaSoundFrontendDocumentBuilder& OutBuilder, bool bLogFailure);
+	// Executes document configuration. Returns true if document configuration applied, false if not
+	// (resulting in a follow-up call to reset the document).
+	UE_INTERNAL UE_API virtual bool ConfigureDocument(FMetaSoundFrontendDocumentBuilder& OutBuilder);
 
 #if WITH_EDITOR
 	struct FEditorOptions
@@ -97,6 +99,10 @@ struct FMetaSoundFrontendDocumentConfiguration
 	UE_EXPERIMENTAL(5.8, "Document Configuration is in active development and subject to change")
 	UE_API virtual const FEditorOptions& GetEditorOptions() const;
 
+#if WITH_EDITOR
+	UE_API virtual EDataValidationResult IsDataValid(const FMetaSoundFrontendDocumentBuilder& Builder, FDataValidationContext& InOutContext) const;
+#endif // WITH_EDITOR
+
 	UE_EXPERIMENTAL(5.8, "Document Configuration is in active development and subject to change")
 	UE_API virtual void OnAssetInitialized(TArray<UObject*> SelectedObjects, FMetaSoundFrontendDocumentBuilder& OutBuilder);
 
diff --git a/Engine/Plugins/Runtime/Metasound/Source/MetasoundFrontend/Public/DocumentConfigurations/MetasoundFrontendPresetDocumentConfiguration.h b/Engine/Plugins/Runtime/Metasound/Source/MetasoundFrontend/Public/DocumentConfigurations/MetasoundFrontendPresetDocumentConfiguration.h
--- a/Engine/Plugins/Runtime/Metasound/Source/MetasoundFrontend/Public/DocumentConfigurations/MetasoundFrontendPresetDocumentConfiguration.h
+++ b/Engine/Plugins/Runtime/Metasound/Source/MetasoundFrontend/Public/DocumentConfigurations/MetasoundFrontendPresetDocumentConfiguration.h
@@ -23,13 +23,16 @@ struct FMetaSoundFrontendPresetDocumentConfiguration final : public FMetaSoundFr
 	virtual ~FMetaSoundFrontendPresetDocumentConfiguration() = default;
 
 #if WITH_EDITORONLY_DATA
-	UE_INTERNAL UE_API virtual bool ConfigureDocument(FMetaSoundFrontendDocumentBuilder& OutBuilder, bool bLogFailure) override;
+	UE_INTERNAL UE_API virtual bool ConfigureDocument(FMetaSoundFrontendDocumentBuilder& OutBuilder) override;
 #endif // WITH_EDITORONLY_DATA
 
 #if WITH_EDITOR
 	UE_EXPERIMENTAL(5.8, "Document Configuration is in active development and subject to change")
 	UE_API virtual const FMetaSoundFrontendDocumentConfiguration::FEditorOptions& GetEditorOptions() const override;
 
+	UE_EXPERIMENTAL(5.8, "Document Configuration is in active development and subject to change")
+	UE_API EDataValidationResult IsDataValid(const FMetaSoundFrontendDocumentBuilder& Builder, FDataValidationContext& InOutContext) const override;
+
 	UE_EXPERIMENTAL(5.8, "Document Configuration is in active development and subject to change")
 	UE_API virtual void OnAssetInitialized(TArray<UObject*> SelectedObjects, FMetaSoundFrontendDocumentBuilder& OutBuilder) override;
 
diff --git a/Engine/Plugins/Runtime/Metasound/Source/MetasoundFrontend/Public/MetasoundFrontendDocumentBuilder.h b/Engine/Plugins/Runtime/Metasound/Source/MetasoundFrontend/Public/MetasoundFrontendDocumentBuilder.h
--- a/Engine/Plugins/Runtime/Metasound/Source/MetasoundFrontend/Public/MetasoundFrontendDocumentBuilder.h
+++ b/Engine/Plugins/Runtime/Metasound/Source/MetasoundFrontend/Public/MetasoundFrontendDocumentBuilder.h
@@ -285,7 +285,7 @@ PRAGMA_ENABLE_EXPERIMENTAL_WARNINGS
 #endif // WITH_EDITORONLY_DATA
 
 	UE_EXPERIMENTAL(5.8, "Configurator is in active development and subject to change")
-	UE_API bool ConfigureDocument(bool bLogFailure = false);
+	UE_API bool ConfigureDocument();
 
 	UE_API bool ConformObjectToDocument();
 
