Return true but clear doc when preset has no parent MetaSound set.
#rnx

#Synced-CL 50173398

diff --git a/Engine/Plugins/Runtime/Metasound/Source/MetasoundFrontend/Private/DocumentConfigurations/MetasoundFrontendDocumentConfiguration.cpp b/Engine/Plugins/Runtime/Metasound/Source/MetasoundFrontend/Private/DocumentConfigurations/MetasoundFrontendDocumentConfiguration.cpp
--- a/Engine/Plugins/Runtime/Metasound/Source/MetasoundFrontend/Private/DocumentConfigurations/MetasoundFrontendDocumentConfiguration.cpp
+++ b/Engine/Plugins/Runtime/Metasound/Source/MetasoundFrontend/Private/DocumentConfigurations/MetasoundFrontendDocumentConfiguration.cpp
@@ -6,16 +6,29 @@
 #include "MetasoundFrontendDocument.h"
 #include "MetasoundFrontendDocumentBuilder.h"
 #include "MetasoundFrontendDocumentBuilderRegistry.h"
+#include "MetasoundFrontendSearchEngine.h"
 
 
 #define LOCTEXT_NAMESPACE "MetaSound"
 
 bool FMetaSoundFrontendDocumentConfiguration::ConfigureDocument(FMetaSoundFrontendDocumentBuilder& OutBuilder)
 {
+	using namespace Metasound::Frontend;
+
+	TSharedRef<FDocumentModifyDelegates> ModifyDelegates = MakeShared<FDocumentModifyDelegates>(OutBuilder.GetDocumentDelegates());
+	OutBuilder.ClearDocument(ModifyDelegates);
+
 #if WITH_EDITORONLY_DATA
 	OutBuilder.SetIsGraphEditable(false);
 #endif // WITH_EDITORONLY_DATA
 
+	// Add required interfaces, but disable editing to ensure document maintains valid state
+	// (ex. MetaSound sources need  the source interface to properly initialize a null-rendering
+	// MetaSound generator.)
+	TArray<FMetasoundFrontendVersion> InitVersions = ISearchEngine::Get().FindUClassDefaultInterfaceVersions(OutBuilder.GetBuilderClassPath());
+	FModifyInterfaceOptions Options({ }, MoveTemp(InitVersions));
+	OutBuilder.ModifyInterfaces(MoveTemp(Options));
+
 	return true;
 }
 
diff --git a/Engine/Plugins/Runtime/Metasound/Source/MetasoundFrontend/Private/DocumentConfigurations/MetasoundFrontendPresetDocumentConfiguration.cpp b/Engine/Plugins/Runtime/Metasound/Source/MetasoundFrontend/Private/DocumentConfigurations/MetasoundFrontendPresetDocumentConfiguration.cpp
--- a/Engine/Plugins/Runtime/Metasound/Source/MetasoundFrontend/Private/DocumentConfigurations/MetasoundFrontendPresetDocumentConfiguration.cpp
+++ b/Engine/Plugins/Runtime/Metasound/Source/MetasoundFrontend/Private/DocumentConfigurations/MetasoundFrontendPresetDocumentConfiguration.cpp
@@ -86,7 +86,10 @@ bool FMetaSoundFrontendPresetDocumentConfiguration::ConfigureDocument(FMetaSound
 		}
 	}
 
-	return false;
+	// Need to return default document configuration when reference is invalid as by default, no referenced parent is set and
+	// configuration must be valid when default constructed for testing purposes.  Data Validation provides place to surface
+	// error (i.e. when saving).
+	return Super::ConfigureDocument(OutBuilder);
 }
 
 EDataValidationResult FMetaSoundFrontendPresetDocumentConfiguration::IsDataValid(const FMetaSoundFrontendDocumentBuilder& Builder, FDataValidationContext& InOutContext) const
diff --git a/Engine/Plugins/Runtime/Metasound/Source/MetasoundFrontend/Public/MetasoundFrontendDocumentBuilder.h b/Engine/Plugins/Runtime/Metasound/Source/MetasoundFrontend/Public/MetasoundFrontendDocumentBuilder.h
--- a/Engine/Plugins/Runtime/Metasound/Source/MetasoundFrontend/Public/MetasoundFrontendDocumentBuilder.h
+++ b/Engine/Plugins/Runtime/Metasound/Source/MetasoundFrontend/Public/MetasoundFrontendDocumentBuilder.h
@@ -859,6 +859,8 @@ PRAGMA_ENABLE_EXPERIMENTAL_WARNINGS
 	};
 #endif // WITH_EDITORONLY_DATA
 
+	UE_API const FTopLevelAssetPath GetBuilderClassPath() const;
+
 private:
 	bool bInDocConfigScope = false;
 
@@ -944,8 +946,6 @@ private:
 
 	void IterateNodesConnectedWithVertex(const FMetasoundFrontendVertexHandle& Vertex, TFunctionRef<void(const FMetasoundFrontendEdge&, FMetasoundFrontendNode&)> NodeIndexIterFunc, const FGuid& InPageID);
 
-	const FTopLevelAssetPath GetBuilderClassPath() const;
-	
 	UE_API FMetasoundFrontendDocument& GetDocumentChecked() const;
 	UE_API IMetaSoundDocumentInterface& GetDocumentInterfaceChecked() const;
 
