Hide experimental Randomizer from selectable configurations
#rnx

#Synced-CL 50243364

diff --git a/Engine/Plugins/Runtime/Metasound/Source/MetasoundEngine/Public/DocumentConfigurations/MetasoundRandomizerDocumentConfiguration.h b/Engine/Plugins/Runtime/Metasound/Source/MetasoundEngine/Public/DocumentConfigurations/MetasoundRandomizerDocumentConfiguration.h
--- a/Engine/Plugins/Runtime/Metasound/Source/MetasoundEngine/Public/DocumentConfigurations/MetasoundRandomizerDocumentConfiguration.h
+++ b/Engine/Plugins/Runtime/Metasound/Source/MetasoundEngine/Public/DocumentConfigurations/MetasoundRandomizerDocumentConfiguration.h
@@ -15,7 +15,7 @@ class UMetaSoundFrontendMemberMetadata;
 enum class EMetasoundFrontendClassType : uint8;
 
 
-USTRUCT(BlueprintType, DisplayName = "Randomizer")
+USTRUCT(BlueprintType, meta = (Hidden, DisplayName = "Randomizer"))
 struct FMetaSoundRandomizerDocumentConfiguration : public FMetaSoundFrontendDocumentConfiguration
 {
 	GENERATED_BODY()
