Fix duplication logic to only allow duplication when not cooked for editor (and omit check for restricted mode)
#rnx

diff --git a/Engine/Plugins/Runtime/Metasound/Source/MetasoundEditor/Private/MetasoundAssetDefinitions.cpp b/Engine/Plugins/Runtime/Metasound/Source/MetasoundEditor/Private/MetasoundAssetDefinitions.cpp
--- a/Engine/Plugins/Runtime/Metasound/Source/MetasoundEditor/Private/MetasoundAssetDefinitions.cpp
+++ b/Engine/Plugins/Runtime/Metasound/Source/MetasoundEditor/Private/MetasoundAssetDefinitions.cpp
@@ -303,15 +303,9 @@ namespace Metasound::Editor
 					PresetConfig.GetMutable().Parent = ParentMetaSound;
 
 					IMetasoundEditorModule& MetaSoundEditorModule = FModuleManager::GetModuleChecked<IMetasoundEditorModule>("MetaSoundEditor");
-					const bool bCanDuplicate = ParentPackage->bIsCookedForEditor || MetaSoundEditorModule.IsRestrictedMode();
-					if (bCanDuplicate)
+					if (ParentPackage->bIsCookedForEditor)
 					{
-						// Duplicate asset to preserve properties from parent asset (ex. quality settings, soundwave properties)
-						NewMetaSound = IAssetTools::Get().DuplicateAssetWithDialogAndTitle(AssetName, FPackageName::GetLongPackagePath(PackagePath), ParentMetaSound, LOCTEXT("CreateMetaSoundPresetTitle", "Create MetaSound Preset"));
-					}
-					else
-					{
-						// Cannot duplicate cooked assets or when in restricted mode, so create new object and copy over properties
+						// Cannot duplicate cooked assets or when underlying asset is cooked, so create new object and copy over properties
 						// for sources in InitAsset below. Since copying properties is done manually, 
 						// SetSoundWaveSettingsFromTemplate may need to be updated with properties to be copied. 
 						UMetaSoundBaseFactory* Factory = NewObject<TFactoryClass>();
@@ -319,6 +313,11 @@ namespace Metasound::Editor
 
 						NewMetaSound = IAssetTools::Get().CreateAssetWithDialog(AssetName, FPackageName::GetLongPackagePath(PackagePath), Factory->GetSupportedClass(), Factory);
 					}
+					else
+					{
+						// Duplicate asset to preserve properties from parent asset (ex. quality settings, soundwave properties)
+						NewMetaSound = IAssetTools::Get().DuplicateAssetWithDialogAndTitle(AssetName, FPackageName::GetLongPackagePath(PackagePath), ParentMetaSound, LOCTEXT("CreateMetaSoundPresetTitle", "Create MetaSound Preset"));
+					}
 				
 					if (NewMetaSound)
 					{
@@ -331,8 +330,8 @@ namespace Metasound::Editor
 
 						FGraphBuilder::RegisterGraphWithFrontend(*NewMetaSound);
 
-						// Have to set wave settings this when not duplicating but after asset is initialized
-						if (!bCanDuplicate)
+						// Have to set wave settings on new asset when unable to duplicate manually (but after asset is initialized)
+						if (ParentPackage->bIsCookedForEditor)
 						{
 							SetObjectSettings(EditorSubsystem, *CastChecked<TClass>(NewMetaSound), *ParentMetaSound);
 						}
