Disable new wave player by default
#rnx
#Synced-CL 50486660

diff --git a/Engine/Plugins/Runtime/Metasound/Source/MetasoundEngine/Private/MetasoundWavePlayerNode.cpp b/Engine/Plugins/Runtime/Metasound/Source/MetasoundEngine/Private/MetasoundWavePlayerNode.cpp
--- a/Engine/Plugins/Runtime/Metasound/Source/MetasoundEngine/Private/MetasoundWavePlayerNode.cpp
+++ b/Engine/Plugins/Runtime/Metasound/Source/MetasoundEngine/Private/MetasoundWavePlayerNode.cpp
@@ -332,7 +332,7 @@ namespace Metasound
 		TEXT("Default: 512"),
 		ECVF_Default);
 
-	static bool UseNewWavePlayerOperator = true;
+	static bool UseNewWavePlayerOperator = false;
 	FAutoConsoleVariableRef CVarMetaSoundWavePlayerUseNewWavePlayer(
 		TEXT("au.MetaSound.WavePlayer.UseNewOperator"),
 		UseNewWavePlayerOperator,
