- Fix for editor locking up when attempting to select parent for MetaSound preset that causes loop
- Remove exposition of preset conversion when in restricted mode

diff --git a/Engine/Plugins/Runtime/Metasound/Source/MetasoundEditor/Private/MetasoundEditor.cpp b/Engine/Plugins/Runtime/Metasound/Source/MetasoundEditor/Private/MetasoundEditor.cpp
--- a/Engine/Plugins/Runtime/Metasound/Source/MetasoundEditor/Private/MetasoundEditor.cpp
+++ b/Engine/Plugins/Runtime/Metasound/Source/MetasoundEditor/Private/MetasoundEditor.cpp
@@ -2525,24 +2525,29 @@ namespace Metasound
 				return EVisibility::Collapsed;
 			});
 
-			MenuBuilder.BeginSection("ConvertConfigurationHeader", LOCTEXT("ConvertConfigurationDisplayName", "Convert Configuration"), ClearVisibility);
+
+			IMetasoundEditorModule& MetaSoundEditorModule = FModuleManager::GetModuleChecked<IMetasoundEditorModule>("MetaSoundEditor");
+			if (!MetaSoundEditorModule.IsRestrictedMode())
 			{
-				if (const FMetaSoundFrontendDocumentBuilder* DocBuilder = GetFrontendBuilder())
+				MenuBuilder.BeginSection("ConvertConfigurationHeader", LOCTEXT("ConvertConfigurationDisplayName", "Convert Configuration"), ClearVisibility);
 				{
-					if (const UScriptStruct* ConfigStruct = DocBuilder->GetConstDocumentConfiguration().GetScriptStruct())
+					if (const FMetaSoundFrontendDocumentBuilder* DocBuilder = GetFrontendBuilder())
 					{
-						MenuBuilder.AddMenuEntry(
-							FText::Format(LOCTEXT("ClearConfigurationButtonText", "Convert From '{0}' ..."), ConfigStruct->GetDisplayNameText()),
-							LOCTEXT("ConvertFromConfigurationTooltip", "Removes the MetaSound's configuration, leaving the MetaSound state intact and enabling editing of configuration-managed MetaSound functionality."),
-							FSlateIcon(),
-							FUIAction(FExecuteAction::CreateSP(this, &FEditor::ShowConvertFromConfigDialog)),
-							{ },
-							EUserInterfaceActionType::Button
-						);
+						if (const UScriptStruct* ConfigStruct = DocBuilder->GetConstDocumentConfiguration().GetScriptStruct())
+						{
+							MenuBuilder.AddMenuEntry(
+								FText::Format(LOCTEXT("ClearConfigurationButtonText", "Convert From '{0}' ..."), ConfigStruct->GetDisplayNameText()),
+								LOCTEXT("ConvertFromConfigurationTooltip", "Removes the MetaSound's configuration, leaving the MetaSound state intact and enabling editing of configuration-managed MetaSound functionality."),
+								FSlateIcon(),
+								FUIAction(FExecuteAction::CreateSP(this, &FEditor::ShowConvertFromConfigDialog)),
+								{ },
+								EUserInterfaceActionType::Button
+							);
+						}
 					}
 				}
+				MenuBuilder.EndSection();
 			}
-			MenuBuilder.EndSection();
 
 			MenuBuilder.BeginSection("SetConfigurationSectionHeader", LOCTEXT("SelectConfigurationSectionName", "Set Configuration"));
 			{
diff --git a/Engine/Plugins/Runtime/Metasound/Source/MetasoundFrontend/Private/DocumentConfigurations/MetasoundFrontendPresetDocumentConfiguration.cpp b/Engine/Plugins/Runtime/Metasound/Source/MetasoundFrontend/Private/DocumentConfigurations/MetasoundFrontendPresetDocumentConfiguration.cpp
--- a/Engine/Plugins/Runtime/Metasound/Source/MetasoundFrontend/Private/DocumentConfigurations/MetasoundFrontendPresetDocumentConfiguration.cpp
+++ b/Engine/Plugins/Runtime/Metasound/Source/MetasoundFrontend/Private/DocumentConfigurations/MetasoundFrontendPresetDocumentConfiguration.cpp
@@ -160,11 +160,20 @@ void FMetaSoundFrontendPresetDocumentConfiguration::OnPropertyChanged(const FPro
 
 	if (InEvent.GetPropertyName() == GET_MEMBER_NAME_CHECKED(FMetaSoundFrontendPresetDocumentConfiguration, Parent))
 	{
-		if (Parent.GetObject() != nullptr)
+		if (UObject* ParentObject = Parent.GetObject())
 		{
-			if (FMetasoundAssetBase* MetaSoundAsset = IMetaSoundAssetManager::GetChecked().GetAsAsset(*Parent.GetObject()))
+			if (FMetasoundAssetBase* ParentAsset = IMetaSoundAssetManager::GetChecked().GetAsAsset(*ParentObject))
 			{
-				MetaSoundAsset->UpdateAndRegisterForExecution(IMetaSoundAssetManager::GetChecked().GetRegistrationOptions());
+				const bool bCreatesLoop = OutBuilder.GetMetasoundAsset().AddingReferenceCausesLoop(*ParentAsset);
+				if (bCreatesLoop)
+				{
+					UE_LOG(LogMetaSound, Error, TEXT("Parent could not be set to '%s': Causes reference loop."), *ParentObject->GetFullName());
+					Parent = nullptr;
+				}
+				else
+				{
+					ParentAsset->UpdateAndRegisterForExecution(IMetaSoundAssetManager::GetChecked().GetRegistrationOptions());
+				}
 			}
 		}
 	}
