Fix negative triggers coming from the trigger repeat node

[RN] Fixed MetaSound crash when changing trigger repeat counts on trigger playing a wave.

#Synced-CL 50516542

diff --git a/Engine/Plugins/Runtime/Metasound/Source/MetasoundFrontend/Private/MetasoundTrigger.cpp b/Engine/Plugins/Runtime/Metasound/Source/MetasoundFrontend/Private/MetasoundTrigger.cpp
--- a/Engine/Plugins/Runtime/Metasound/Source/MetasoundFrontend/Private/MetasoundTrigger.cpp
+++ b/Engine/Plugins/Runtime/Metasound/Source/MetasoundFrontend/Private/MetasoundTrigger.cpp
@@ -38,15 +38,18 @@ namespace Metasound
 
 	void FTrigger::TriggerFrame(int32 InFrameToTrigger)
 	{
-		// Insert trigger frame index into sorted triggered frames
-		TriggeredFrames.Insert(InFrameToTrigger, Algo::LowerBound(TriggeredFrames, InFrameToTrigger));
-
-		if (InFrameToTrigger < NumFramesPerBlock)
+		if (ensureMsgf(InFrameToTrigger >= 0, TEXT("Do not trigger negative frames [%d]"), InFrameToTrigger))
 		{
-			UpdateLastTriggerIndexInBlock();
-		}
+			// Insert trigger frame index into sorted triggered frames
+			TriggeredFrames.Insert(InFrameToTrigger, Algo::LowerBound(TriggeredFrames, InFrameToTrigger));
+
+			if (InFrameToTrigger < NumFramesPerBlock)
+			{
+				UpdateLastTriggerIndexInBlock();
+			}
 
-		bHasTrigger = true;
+			bHasTrigger = true;
+		}
 	}
 
 	void FTrigger::AdvanceBlock()
diff --git a/Engine/Plugins/Runtime/Metasound/Source/MetasoundStandardNodes/Private/MetasoundTriggerRepeatNode.cpp b/Engine/Plugins/Runtime/Metasound/Source/MetasoundStandardNodes/Private/MetasoundTriggerRepeatNode.cpp
--- a/Engine/Plugins/Runtime/Metasound/Source/MetasoundStandardNodes/Private/MetasoundTriggerRepeatNode.cpp
+++ b/Engine/Plugins/Runtime/Metasound/Source/MetasoundStandardNodes/Private/MetasoundTriggerRepeatNode.cpp
@@ -137,6 +137,8 @@ namespace Metasound
 
 				if (!bShouldRepeat)
 				{
+					// Disable repeat node if have satisfied the number of repeats
+					bEnabled = false;
 					break;
 				}
 
