[Metasounds][Experimental] Move Channel Agnostic work into main plugin

#tests LowLevelTests Audio

#Synced-CL 50628516

diff --git a/Engine/Plugins/Experimental/MetasoundExperimental/Source/AudioExperimentalRuntime/AudioExperimentalRuntime.Build.cs b/Engine/Plugins/Experimental/MetasoundExperimental/Source/AudioExperimentalRuntime/AudioExperimentalRuntime.Build.cs
--- a/Engine/Plugins/Experimental/MetasoundExperimental/Source/AudioExperimentalRuntime/AudioExperimentalRuntime.Build.cs
+++ b/Engine/Plugins/Experimental/MetasoundExperimental/Source/AudioExperimentalRuntime/AudioExperimentalRuntime.Build.cs
@@ -15,6 +15,7 @@ public class AudioExperimentalRuntime : ModuleRules
 			new string[]
 			{
 				"Core",
+				"AudioChannelAgnosticCore"
 			}
 		);
 		
diff --git a/Engine/Plugins/Experimental/MetasoundExperimental/Source/AudioExperimentalRuntime/Private/AudioExperimentalRuntimeModule.cpp b/Engine/Plugins/Experimental/MetasoundExperimental/Source/AudioExperimentalRuntime/Private/AudioExperimentalRuntimeModule.cpp
--- a/Engine/Plugins/Experimental/MetasoundExperimental/Source/AudioExperimentalRuntime/Private/AudioExperimentalRuntimeModule.cpp
+++ b/Engine/Plugins/Experimental/MetasoundExperimental/Source/AudioExperimentalRuntime/Private/AudioExperimentalRuntimeModule.cpp
@@ -8,9 +8,7 @@ class FAudioCoreExperimentalModule final : public IModuleInterface
 {
 public:
 	virtual void StartupModule() override
-	{
-		Audio::RegisterChannelLayouts();
-	}
+	{}
 	
 	virtual void ShutdownModule() override
 	{}
diff --git a/Engine/Plugins/Experimental/MetasoundExperimental/Source/AudioExperimentalRuntime/Private/ChannelAgnostic/ChannelAgnosticTranscoding.cpp b/Engine/Plugins/Experimental/MetasoundExperimental/Source/AudioExperimentalRuntime/Private/ChannelAgnostic/ChannelAgnosticTranscoding.cpp
deleted file mode 100644
diff --git a/Engine/Plugins/Experimental/MetasoundExperimental/Source/AudioExperimentalRuntime/Private/ChannelAgnostic/ChannelAgnosticType.cpp b/Engine/Plugins/Experimental/MetasoundExperimental/Source/AudioExperimentalRuntime/Private/ChannelAgnostic/ChannelAgnosticType.cpp
deleted file mode 100644
diff --git a/Engine/Plugins/Experimental/MetasoundExperimental/Source/AudioExperimentalRuntime/Private/ChannelAgnostic/ChannelAgnosticTypeUtils.cpp b/Engine/Plugins/Experimental/MetasoundExperimental/Source/AudioExperimentalRuntime/Private/ChannelAgnostic/ChannelAgnosticTypeUtils.cpp
deleted file mode 100644
diff --git a/Engine/Plugins/Experimental/MetasoundExperimental/Source/AudioExperimentalRuntime/Private/TypeFamily/ChannelTypeFamily.cpp b/Engine/Plugins/Experimental/MetasoundExperimental/Source/AudioExperimentalRuntime/Private/TypeFamily/ChannelTypeFamily.cpp
deleted file mode 100644
diff --git a/Engine/Plugins/Experimental/MetasoundExperimental/Source/AudioExperimentalRuntime/Public/ChannelAgnostic/ChannelAgnosticTranscoding.h b/Engine/Plugins/Experimental/MetasoundExperimental/Source/AudioExperimentalRuntime/Public/ChannelAgnostic/ChannelAgnosticTranscoding.h
deleted file mode 100644
diff --git a/Engine/Plugins/Experimental/MetasoundExperimental/Source/AudioExperimentalRuntime/Public/ChannelAgnostic/ChannelAgnosticType.h b/Engine/Plugins/Experimental/MetasoundExperimental/Source/AudioExperimentalRuntime/Public/ChannelAgnostic/ChannelAgnosticType.h
deleted file mode 100644
diff --git a/Engine/Plugins/Experimental/MetasoundExperimental/Source/AudioExperimentalRuntime/Public/ChannelAgnostic/ChannelAgnosticTypeUtils.h b/Engine/Plugins/Experimental/MetasoundExperimental/Source/AudioExperimentalRuntime/Public/ChannelAgnostic/ChannelAgnosticTypeUtils.h
deleted file mode 100644
diff --git a/Engine/Plugins/Experimental/MetasoundExperimental/Source/AudioExperimentalRuntime/Public/ScratchBuffer.h b/Engine/Plugins/Experimental/MetasoundExperimental/Source/AudioExperimentalRuntime/Public/ScratchBuffer.h
deleted file mode 100644
diff --git a/Engine/Plugins/Experimental/MetasoundExperimental/Source/AudioExperimentalRuntime/Public/SimpleAlloc/SimpleAllocBase.h b/Engine/Plugins/Experimental/MetasoundExperimental/Source/AudioExperimentalRuntime/Public/SimpleAlloc/SimpleAllocBase.h
deleted file mode 100644
diff --git a/Engine/Plugins/Experimental/MetasoundExperimental/Source/AudioExperimentalRuntime/Public/SimpleAlloc/SimpleHeapAllocator.h b/Engine/Plugins/Experimental/MetasoundExperimental/Source/AudioExperimentalRuntime/Public/SimpleAlloc/SimpleHeapAllocator.h
deleted file mode 100644
diff --git a/Engine/Plugins/Experimental/MetasoundExperimental/Source/AudioExperimentalRuntime/Public/SimpleAlloc/SimpleLinearAllocator.h b/Engine/Plugins/Experimental/MetasoundExperimental/Source/AudioExperimentalRuntime/Public/SimpleAlloc/SimpleLinearAllocator.h
deleted file mode 100644
diff --git a/Engine/Plugins/Experimental/MetasoundExperimental/Source/AudioExperimentalRuntime/Public/TypeFamily/ChannelTypeFamily.h b/Engine/Plugins/Experimental/MetasoundExperimental/Source/AudioExperimentalRuntime/Public/TypeFamily/ChannelTypeFamily.h
deleted file mode 100644
diff --git a/Engine/Plugins/Experimental/MetasoundExperimental/Source/AudioExperimentalRuntime/Public/TypeFamily/TypeFamily.h b/Engine/Plugins/Experimental/MetasoundExperimental/Source/AudioExperimentalRuntime/Public/TypeFamily/TypeFamily.h
deleted file mode 100644
diff --git a/Engine/Plugins/Experimental/MetasoundExperimental/Source/MetasoundExperimentalEngineRuntime/MetasoundExperimentalEngineRuntime.Build.cs b/Engine/Plugins/Experimental/MetasoundExperimental/Source/MetasoundExperimentalEngineRuntime/MetasoundExperimentalEngineRuntime.Build.cs
--- a/Engine/Plugins/Experimental/MetasoundExperimental/Source/MetasoundExperimentalEngineRuntime/MetasoundExperimentalEngineRuntime.Build.cs
+++ b/Engine/Plugins/Experimental/MetasoundExperimental/Source/MetasoundExperimentalEngineRuntime/MetasoundExperimentalEngineRuntime.Build.cs
@@ -6,6 +6,7 @@ public class MetasoundExperimentalEngineRuntime : ModuleRules
 {
 	public MetasoundExperimentalEngineRuntime(ReadOnlyTargetRules Target) : base(Target)
 	{
+		bValidateExperimentalApi = false;
 		PCHUsage = ModuleRules.PCHUsageMode.UseExplicitOrSharedPCHs;
 		PublicDependencyModuleNames.AddRange(
 			new string[]
@@ -35,7 +36,8 @@ public class MetasoundExperimentalEngineRuntime : ModuleRules
 				"MetasoundEngine",
 				"MetasoundStandardNodes",
 				"SignalProcessing",
-				"MetasoundExperimentalRuntime"
+				"MetasoundExperimentalRuntime",
+				"AudioChannelAgnosticCore",
 			}
 		);
 
diff --git a/Engine/Plugins/Experimental/MetasoundExperimental/Source/MetasoundExperimentalEngineRuntime/Private/MetasoundGranulatorNode.cpp b/Engine/Plugins/Experimental/MetasoundExperimental/Source/MetasoundExperimentalEngineRuntime/Private/MetasoundGranulatorNode.cpp
--- a/Engine/Plugins/Experimental/MetasoundExperimental/Source/MetasoundExperimentalEngineRuntime/Private/MetasoundGranulatorNode.cpp
+++ b/Engine/Plugins/Experimental/MetasoundExperimental/Source/MetasoundExperimentalEngineRuntime/Private/MetasoundGranulatorNode.cpp
@@ -474,14 +474,14 @@ namespace Metasound
 			float PanningFraction = 0.0f;
 			GetOutputSpeakerChannelData(InGrain, PanningFraction, PrevData, NextData);
 
-			if (PrevData.ChannelID == INDEX_NONE || NextData.ChannelID == INDEX_NONE)
+			if (PrevData.ChannelID.IsNone() || NextData.ChannelID.IsNone())
 			{
 				return;
 			}
-			int32 PrevSpeakerIndex = DiscreteFamily.FindSpeakerIndex(static_cast<ESpeakerShortNames>(PrevData.ChannelID));
+			const int32 PrevSpeakerIndex = DiscreteFamily.FindSpeakerIndex(PrevData.ChannelID);
 			check(PrevSpeakerIndex != INDEX_NONE);
 
-			int32 NextSpeakerIndex = DiscreteFamily.FindSpeakerIndex(static_cast<ESpeakerShortNames>(NextData.ChannelID));
+			const int32 NextSpeakerIndex = DiscreteFamily.FindSpeakerIndex(NextData.ChannelID);
 			check(NextSpeakerIndex != INDEX_NONE);
 
 			TArrayView<float> PrevChannelOutAudioBufferView = OutputAudio->GetChannel(PrevSpeakerIndex);
diff --git a/Engine/Plugins/Experimental/MetasoundExperimental/Source/MetasoundExperimentalRuntime/MetasoundExperimentalRuntime.Build.cs b/Engine/Plugins/Experimental/MetasoundExperimental/Source/MetasoundExperimentalRuntime/MetasoundExperimentalRuntime.Build.cs
--- a/Engine/Plugins/Experimental/MetasoundExperimental/Source/MetasoundExperimentalRuntime/MetasoundExperimentalRuntime.Build.cs
+++ b/Engine/Plugins/Experimental/MetasoundExperimental/Source/MetasoundExperimentalRuntime/MetasoundExperimentalRuntime.Build.cs
@@ -41,7 +41,8 @@ public class MetasoundExperimentalRuntime : ModuleRules
 				"MetasoundStandardNodes",
 				"SignalProcessing",
 				"Serialization",
-				"Projects"
+				"Projects",
+				"AudioChannelAgnosticCore"
 			}
 		);
 		PrivateDefinitions.AddRange(
@@ -51,8 +52,9 @@ public class MetasoundExperimentalRuntime : ModuleRules
 				"METASOUND_MODULE=MetasoundExperimentalRuntime"
 			});
 		
-		RuntimeDependencies.Add("$(PluginDir)/Content/CatSchemas/Composite/*.json", StagedFileType.UFS);
-		RuntimeDependencies.Add("$(PluginDir)/Content/CatSchemas/Discrete/*.json", StagedFileType.UFS);
-		RuntimeDependencies.Add("$(PluginDir)/Content/CatSchemas/Soundfield/*.json", StagedFileType.UFS);
+		// These should probably move back to the experimental plugin again.
+		RuntimeDependencies.Add("$(PluginDir)/Content/Schemas/ChannelAgnostic/Composite/*.json", StagedFileType.UFS);
+		RuntimeDependencies.Add("$(PluginDir)/Content/Schemas/ChannelAgnostic/Discrete/*.json", StagedFileType.UFS);
+		RuntimeDependencies.Add("$(PluginDir)/Content/Schemas/ChannelAgnostic/Soundfield/*.json", StagedFileType.UFS);
 	}
 }
diff --git a/Engine/Plugins/Experimental/MetasoundExperimental/Source/MetasoundExperimentalRuntime/Private/MetasoundCatBreakNode.cpp b/Engine/Plugins/Experimental/MetasoundExperimental/Source/MetasoundExperimentalRuntime/Private/MetasoundCatBreakNode.cpp
--- a/Engine/Plugins/Experimental/MetasoundExperimental/Source/MetasoundExperimentalRuntime/Private/MetasoundCatBreakNode.cpp
+++ b/Engine/Plugins/Experimental/MetasoundExperimental/Source/MetasoundExperimentalRuntime/Private/MetasoundCatBreakNode.cpp
@@ -135,7 +135,7 @@ namespace Metasound
 	class FCatBreakOperator final : public TExecutableOperator<FCatBreakOperator>
 	{
 	public:
-		using FTranscoder = Audio::FChannelTypeFamily::FTranscoder;
+		using FTranscoder = Audio::ChannelAgnosticTranscoder::FTranscoder; 
 		FCatBreakOperator(const FBuildOperatorParams& InParams, FChannelAgnosticTypeReadRef&& InInputCat, TArray<TDataWriteReference<FAudioBuffer>>&& InOutputVerticies,
 			FTranscoder&& InTranscoder, const FName InFormat)
 			: InputCat(MoveTemp(InInputCat))
@@ -178,12 +178,13 @@ namespace Metasound
 
 			// Always ask for the transcoder 
 			// In the trivial case where we are the same format, this will be just a memcpy.
-			FTranscoder Transcoder = InputPin->GetType().GetTranscoder(
+			FTranscoder Transcoder = Audio::FTranscoderResolver::Resolve(
+				InputPin->GetType(),
 				{
 					.ToType = *ConcreteToType,
 					.TranscodeMethod = CatBreakData->GetTranscodeMethod(),
 					.MixMethod = CatBreakData->GetMixMethod(),
-			});
+				});
 			
 			return MakeUnique<FCatBreakOperator>(
 				InParams,
diff --git a/Engine/Plugins/Experimental/MetasoundExperimental/Source/MetasoundExperimentalRuntime/Private/MetasoundCatCastingNode.cpp b/Engine/Plugins/Experimental/MetasoundExperimental/Source/MetasoundExperimentalRuntime/Private/MetasoundCatCastingNode.cpp
--- a/Engine/Plugins/Experimental/MetasoundExperimental/Source/MetasoundExperimentalRuntime/Private/MetasoundCatCastingNode.cpp
+++ b/Engine/Plugins/Experimental/MetasoundExperimental/Source/MetasoundExperimentalRuntime/Private/MetasoundCatCastingNode.cpp
@@ -2,7 +2,7 @@
 
 #include "MetasoundCatCastingNode.h"
 
-#include "MetasoundCatDataSchemas.h"
+#include "MetasoundChannelAgnosticDataSchemas.h"
 #include "MetasoundChannelAgnosticType.h"
 #include "MetasoundExecutableOperator.h"
 #include "MetasoundFacade.h"
@@ -140,7 +140,8 @@ namespace Metasound
 		InOutVertexData.BindWriteVertex(METASOUND_GET_PARAM_NAME(OutputToCat), OutputCastResult);
 
 		// Create transcoder 
-		Transcoder = InputFrom->GetType().GetTranscoder(
+		Transcoder = Audio::FTranscoderResolver::Resolve(
+			InputFrom->GetType(),
 			{
 				.ToType = OutputCastResult->GetType(),
 				.TranscodeMethod = TranscodeMethod,
diff --git a/Engine/Plugins/Experimental/MetasoundExperimental/Source/MetasoundExperimentalRuntime/Private/MetasoundCatCastingNode.h b/Engine/Plugins/Experimental/MetasoundExperimental/Source/MetasoundExperimentalRuntime/Private/MetasoundCatCastingNode.h
--- a/Engine/Plugins/Experimental/MetasoundExperimental/Source/MetasoundExperimentalRuntime/Private/MetasoundCatCastingNode.h
+++ b/Engine/Plugins/Experimental/MetasoundExperimental/Source/MetasoundExperimentalRuntime/Private/MetasoundCatCastingNode.h
@@ -5,6 +5,7 @@
 #include "MetasoundChannelAgnosticType.h"
 #include "MetasoundExecutableOperator.h"
 #include "MetasoundFrontendDocument.h"
+#include "ChannelAgnostic/ChannelAgnosticTranscoding.h"
 #include "Templates/SharedPointer.h"
 #include "TypeFamily/ChannelTypeFamily.h"
 #include "UObject/PropertyText.h"
@@ -91,7 +92,7 @@ namespace Metasound
 		static FNodeClassMetadata GetNodeInfo();
 	
 	private:
-		using FTranscoder = Audio::FChannelTypeFamily::FTranscoder;
+		using FTranscoder = Audio::ChannelAgnosticTranscoder::FTranscoder;
 		FChannelAgnosticTypeReadRef InputFrom;
 		FName ToFormatName;
 		TDataWriteReference<FChannelAgnosticType> OutputCastResult;
diff --git a/Engine/Plugins/Experimental/MetasoundExperimental/Source/MetasoundExperimentalRuntime/Private/MetasoundCatDataSchemas.cpp b/Engine/Plugins/Experimental/MetasoundExperimental/Source/MetasoundExperimentalRuntime/Private/MetasoundCatDataSchemas.cpp
deleted file mode 100644
diff --git a/Engine/Plugins/Experimental/MetasoundExperimental/Source/MetasoundExperimentalRuntime/Private/MetasoundCatDataSchemas.h b/Engine/Plugins/Experimental/MetasoundExperimental/Source/MetasoundExperimentalRuntime/Private/MetasoundCatDataSchemas.h
deleted file mode 100644
diff --git a/Engine/Plugins/Experimental/MetasoundExperimental/Source/MetasoundExperimentalRuntime/Private/MetasoundCatMixerNode.cpp b/Engine/Plugins/Experimental/MetasoundExperimental/Source/MetasoundExperimentalRuntime/Private/MetasoundCatMixerNode.cpp
--- a/Engine/Plugins/Experimental/MetasoundExperimental/Source/MetasoundExperimentalRuntime/Private/MetasoundCatMixerNode.cpp
+++ b/Engine/Plugins/Experimental/MetasoundExperimental/Source/MetasoundExperimentalRuntime/Private/MetasoundCatMixerNode.cpp
@@ -13,6 +13,7 @@
 #include "MetasoundPrimitives.h"
 #include "MetasoundStandardNodesCategories.h"
 #include "MetasoundVertex.h"
+#include "ChannelAgnostic/ChannelAgnosticTypeUtils.h"
 #include "TypeFamily/ChannelTypeFamily.h"
 
 #define LOCTEXT_NAMESPACE "MetasoundStandardNodes_CatMixerNode"
@@ -113,12 +114,29 @@ namespace Metasound
 			return nullptr;
 		}
 		static const Audio::FChannelTypeFamily* ChooseOutputFormat(
-			const TArray<FChannelAgnosticTypeReadRef>& InOutputs, const Audio::IChannelTypeRegistry& InRegistry, const FCatMixerOperatorData& InOperatorData)
+			const TArray<FChannelAgnosticTypeReadRef>& InOutputs, const Audio::IChannelTypeRegistry& InRegistry, const FCatMixerOperatorData& InOperatorData, const FMetasoundEnvironment& InEnvironment)
 		{
 			using enum EMetasoundMixerFormatChoosingMethod;
 			switch (InOperatorData.FormatChoosingMethod)
 			{
-			case MetasoundOutput:	// TODO. 
+			case MetasoundOutput:
+			{
+#if 0 // This is in another followup CL 
+				// Ask for number of channels AudioMixer is configured to use.
+				const int32 NumChannels = FMath::Clamp(InEnvironment.GetValue<int32>(Frontend::SourceInterface::Environment::AudioMixerNumChannels)
+					, 1, AUDIO_MIXER_MAX_OUTPUT_CHANNELS);
+
+				// Convert that using standard mapping to ChannelMask.
+				TOptional<uint32> Mask = Audio::FChannelAgnosticUtils::FindChannelMaskFromAudioMixerNumChannels(NumChannels); 
+				if (!ensure(Mask.IsSet()))
+				{
+					return nullptr;
+				}
+					
+				// Lookup mask in discrete registry. 
+				return InRegistry.FindDiscreteChannelFromBitMask(*Mask).Get();
+#endif
+			}
 			case HighestInput:		return FindHighestInputType(InOutputs);
 			case LowestInput:		return FindLowestInputType(InOutputs);
 			case Custom:			return InRegistry.FindChannel(InOperatorData.ToMixType);
@@ -127,9 +145,9 @@ namespace Metasound
 			return nullptr;
 		}
 		static const Audio::FChannelTypeFamily* ResolveOutputFormat(
-		const TArray<FChannelAgnosticTypeReadRef>& InOutputs, const Audio::IChannelTypeRegistry& InRegistry, const FCatMixerOperatorData& InOperatorData)
+		const TArray<FChannelAgnosticTypeReadRef>& InOutputs, const Audio::IChannelTypeRegistry& InRegistry, const FCatMixerOperatorData& InOperatorData, const FMetasoundEnvironment& InEnvironment)
 		{
-			if (const Audio::FChannelTypeFamily* Found = ChooseOutputFormat(InOutputs, InRegistry, InOperatorData))
+			if (const Audio::FChannelTypeFamily* Found = ChooseOutputFormat(InOutputs, InRegistry, InOperatorData, InEnvironment ))
 			{
 				return InRegistry.FindConcreteChannel(Found->GetName());
 			}
@@ -155,12 +173,14 @@ namespace Metasound
 			}
 			for  (FChannelAgnosticTypeReadRef& Input : Inputs)
 			{
-				Transcoders.Emplace(Input->GetType().GetTranscoder(
-					{
-						.ToType = Outputs->GetType(),
-						.TranscodeMethod = static_cast<Audio::EChannelTranscodeMethod>(InOpData.CatCastingMethod),
-						.MixMethod = static_cast<Audio::EChannelMapMonoUpmixMethod>(InOpData.CatCastingMethod)
-					}));
+				Transcoders.Emplace(
+					Audio::FTranscoderResolver::Resolve(
+						Input->GetType(),
+				{
+							.ToType = Outputs->GetType(),
+							.TranscodeMethod = static_cast<Audio::EChannelTranscodeMethod>(InOpData.CatCastingMethod),
+							.MixMethod = static_cast<Audio::EChannelMapMonoUpmixMethod>(InOpData.ChannelMapMonoUpmixMethod)
+						}));
 			}
 		}
 		virtual ~FCatMixerOperator() override = default;
@@ -177,7 +197,7 @@ namespace Metasound
 			);
 			return Metadata;
 		}
-		
+				
 		static TUniquePtr<IOperator> CreateOperator(const FBuildOperatorParams& InParams, FBuildResults& OutResults)
 		{
 			using namespace CatMixerPrivate;
@@ -187,13 +207,7 @@ namespace Metasound
 			{
 				return MakeUnique<FNoOpOperator>();
 			}
-				
-			const Audio::FChannelTypeFamily* ConcreteToType = Audio::GetChannelRegistry().FindConcreteChannel(OperatorData->ToMixType);
-			if (!ConcreteToType)
-			{
-				return MakeUnique<FNoOpOperator>();
-			}
-
+			
 			// inputs.
 			TArray<FChannelAgnosticTypeReadRef> Inputs;
 			TArray<FFloatReadRef> Gains;
@@ -203,7 +217,8 @@ namespace Metasound
 				Gains.Emplace(InParams.InputData.GetOrCreateDefaultDataReadReference<float>(MakeGainInputVertexName(i), InParams.OperatorSettings));
 			}
 
-			const Audio::FChannelTypeFamily* Resolved = ResolveOutputFormat(Inputs,Audio::GetChannelRegistry(), *OperatorData);
+			// output.
+			const Audio::FChannelTypeFamily* Resolved = ResolveOutputFormat(Inputs,Audio::GetChannelRegistry(), *OperatorData, InParams.Environment);
 			if (!Resolved)
 			{
 				return MakeUnique<FNoOpOperator>();
@@ -293,7 +308,7 @@ namespace Metasound
 		TArray<FChannelAgnosticTypeReadRef> Inputs;
 		FChannelAgnosticTypeWriteRef Outputs;
 		FOperatorSettings Settings;
-		TArray<Audio::FChannelTypeFamily::FTranscoder> Transcoders;
+		TArray<Audio::ChannelAgnosticTranscoder::FTranscoder> Transcoders;
 		FChannelAgnosticType TranscodeScratch;
 
 		static FNodeClassMetadata CreateNodeClassMetadata(const FName& InOperatorName, const FText& InDisplayName, const FText& InDescription, const FVertexInterface& InDefaultInterface)
diff --git a/Engine/Plugins/Experimental/MetasoundExperimental/Source/MetasoundExperimentalRuntime/Private/MetasoundCatWaveWriterNode.cpp b/Engine/Plugins/Experimental/MetasoundExperimental/Source/MetasoundExperimentalRuntime/Private/MetasoundCatWaveWriterNode.cpp
--- a/Engine/Plugins/Experimental/MetasoundExperimental/Source/MetasoundExperimentalRuntime/Private/MetasoundCatWaveWriterNode.cpp
+++ b/Engine/Plugins/Experimental/MetasoundExperimental/Source/MetasoundExperimentalRuntime/Private/MetasoundCatWaveWriterNode.cpp
@@ -140,7 +140,7 @@ namespace Metasound
 				// If we have a valid writer and enabled.
 				if (Writer && *Enabled && NumInputChannels > 0)
 				{
-					Audio::FCatUtils::Interleave(*AudioInputs, InterleaveBuffer);
+					Audio::FChannelAgnosticUtils::Interleave(*AudioInputs, InterleaveBuffer);
 					Writer->Write(InterleaveBuffer);
 				}
 			}
diff --git a/Engine/Plugins/Experimental/MetasoundExperimental/Source/MetasoundExperimentalRuntime/Private/MetasoundChannelAgnosticType.cpp b/Engine/Plugins/Experimental/MetasoundExperimental/Source/MetasoundExperimentalRuntime/Private/MetasoundChannelAgnosticType.cpp
deleted file mode 100644
diff --git a/Engine/Plugins/Experimental/MetasoundExperimental/Source/MetasoundExperimentalRuntime/Private/MetasoundExperimentalRuntimeModule.cpp b/Engine/Plugins/Experimental/MetasoundExperimental/Source/MetasoundExperimentalRuntime/Private/MetasoundExperimentalRuntimeModule.cpp
--- a/Engine/Plugins/Experimental/MetasoundExperimental/Source/MetasoundExperimentalRuntime/Private/MetasoundExperimentalRuntimeModule.cpp
+++ b/Engine/Plugins/Experimental/MetasoundExperimental/Source/MetasoundExperimentalRuntime/Private/MetasoundExperimentalRuntimeModule.cpp
@@ -2,22 +2,37 @@
 
 #include "MetasoundExperimentalRuntimeModule.h"
 
-#include "MetasoundCatConversionNode.h"
-#include "MetasoundCatDataSchemas.h"
+#include "MetasoundChannelAgnosticDataSchemas.h"
 #include "MetasoundFrontendModuleRegistrationMacros.h"
+#include "Interfaces/IPluginManager.h"
+#include "Misc/Paths.h"
 #include "Modules/ModuleManager.h"
 
+namespace MetasoundExperimentalRuntimePrivate
+{
+	static FString GetPluginContentPath()
+	{
+		if (const TSharedPtr<IPlugin> Plugin = IPluginManager::Get().FindPlugin(TEXT("MetasoundExperimental")); Plugin.IsValid())
+		{
+			return Plugin->GetContentDir();
+		}
+
+		// Some targets don't appear to return the plugin, so fall-back to this.
+		const FString ContentPath = FPaths::Combine(FPaths::EnginePluginsDir(), TEXT("Runtime/Experimental/MetasoundExperimental/Content"));;
+		return ContentPath;
+	}
+}
+
 void FMetasoundExperimentalRuntimeModule::StartupModule()
 {
 	FModuleManager::Get().LoadModule("AudioExperimentalRuntime");
-	Metasound::LoadPolymorphicDataSchemas();
-	Metasound::RegisterCatConvertNode();
-	
+	RegisteredDataSchemas.Append(Metasound::LoadChannelAgnosticDataSchemas(MetasoundExperimentalRuntimePrivate::GetPluginContentPath()));
 	METASOUND_REGISTER_ITEMS_IN_MODULE
 }
 
 void FMetasoundExperimentalRuntimeModule::ShutdownModule()
 {
+	Metasound::UnloadChannelAgnosticDataSchemas(RegisteredDataSchemas);
 	METASOUND_UNREGISTER_ITEMS_IN_MODULE
 }
 
diff --git a/Engine/Plugins/Experimental/MetasoundExperimental/Source/MetasoundExperimentalRuntime/Private/MetasoundExperimentalRuntimeModule.h b/Engine/Plugins/Experimental/MetasoundExperimental/Source/MetasoundExperimentalRuntime/Private/MetasoundExperimentalRuntimeModule.h
--- a/Engine/Plugins/Experimental/MetasoundExperimental/Source/MetasoundExperimentalRuntime/Private/MetasoundExperimentalRuntimeModule.h
+++ b/Engine/Plugins/Experimental/MetasoundExperimental/Source/MetasoundExperimentalRuntime/Private/MetasoundExperimentalRuntimeModule.h
@@ -2,6 +2,7 @@
 
 #pragma once
 
+#include "Containers/Array.h"
 #include "Modules/ModuleInterface.h"
 
 
@@ -14,4 +15,5 @@ public:
 	//~ End IModuleInterface API
 
 private:
+	TArray<FName> RegisteredDataSchemas;
 };
diff --git a/Engine/Plugins/Experimental/MetasoundExperimental/Source/MetasoundExperimentalRuntime/Public/MetasoundChannelAgnosticType.h b/Engine/Plugins/Experimental/MetasoundExperimental/Source/MetasoundExperimentalRuntime/Public/MetasoundChannelAgnosticType.h
deleted file mode 100644
diff --git a/Engine/Plugins/Runtime/Metasound/Source/MetasoundFrontend/MetasoundFrontend.Build.cs b/Engine/Plugins/Runtime/Metasound/Source/MetasoundFrontend/MetasoundFrontend.Build.cs
--- a/Engine/Plugins/Runtime/Metasound/Source/MetasoundFrontend/MetasoundFrontend.Build.cs
+++ b/Engine/Plugins/Runtime/Metasound/Source/MetasoundFrontend/MetasoundFrontend.Build.cs
@@ -22,9 +22,16 @@ namespace UnrealBuildTool.Rules
 					"JsonUtilities",
 					"Serialization",
 					"SignalProcessing",
-					"MetasoundGraphCore"
+					"MetasoundGraphCore",
+					"AudioChannelAgnosticCore",
 				}
 			);
+			
+			PrivateDependencyModuleNames.AddRange(
+				new string[]
+				{
+					"Projects",	// For IPluginManager
+				});
 
 			if (Target.bBuildWithEditorOnlyData)
 			{
@@ -47,6 +54,9 @@ namespace UnrealBuildTool.Rules
 
 			PrivateDefinitions.Add("METASOUND_PLUGIN=Metasound");
 			PrivateDefinitions.Add("METASOUND_MODULE=MetasoundFrontend");
+
+			// Disable experimental warnings
+			bValidateExperimentalApi = false;
 		}
 	}
 }
diff --git a/Engine/Plugins/Runtime/Metasound/Source/MetasoundFrontend/Private/MetasoundChannelAgnosticDataSchemas.cpp b/Engine/Plugins/Runtime/Metasound/Source/MetasoundFrontend/Private/MetasoundChannelAgnosticDataSchemas.cpp
new file mode 100644
--- /dev/null
+++ b/Engine/Plugins/Runtime/Metasound/Source/MetasoundFrontend/Private/MetasoundChannelAgnosticDataSchemas.cpp
@@ -0,0 +1,983 @@
+﻿// Copyright Epic Games, Inc. All Rights Reserved.
+
+#include "MetasoundChannelAgnosticDataSchemas.h"
+
+#include "Containers/UnrealString.h"
+#include "MetasoundChannelAgnosticType.h"
+#include "MetasoundFrontendDataTypeRegistry.h"
+#include "MetasoundFrontend/Private/MetasoundJsonBackend.h"
+#include "HAL/FileManager.h"
+#include "HAL/IConsoleManager.h"
+#include "Interfaces/IPluginManager.h"
+#include "Misc/Paths.h"
+#include "Modules/ModuleManager.h"
+#include "StructDeserializer.h"
+#include "TypeFamily/ChannelTypeFamily.h"
+
+#define LOCTEXT_NAMESPACE "MetasoundFrontend"
+
+namespace Metasound
+{
+	static FAutoConsoleCommand DumpDataTypes(TEXT("au.metasound.dump_frontend_datatypes"), TEXT(""), FConsoleCommandDelegate::CreateLambda([]
+	{
+		using namespace Metasound::Frontend;
+		TArray<FName> Names;
+		IDataTypeRegistry::Get().GetRegisteredDataTypeNames(Names);
+		for (FName i : Names)
+		{
+			FDataTypeRegistryInfo Info;
+			IDataTypeRegistry::Get().GetDataTypeInfo(i, Info);
+			UE_LOG(LogMetaSound, Display, TEXT("GetDataTypeInfo: %s, bIsPolymorphic=%d, bIsAbstract=%d, Base=%s"),
+				*i.ToString(), (int32)Info.bIsPolymorphic,(int32)Info.bIsAbstract, *Info.ParentDataTypeName.ToString()  );
+		}
+	}));
+
+	static FAutoConsoleCommand DumpPolyTypes(TEXT("au.metasound.dump_poly_types"), TEXT(""), FConsoleCommandDelegate::CreateLambda([]
+	{
+		using namespace Metasound;
+	
+		TArray<FPolyTypeInfo> All;
+		FPolyRegistry::Get().GetAllRegisteredTypes(All);
+
+		for (const FPolyTypeInfo& i : All)
+		{
+			UE_LOG(LogMetaSound, Display, TEXT("PolyType: %s, bIsPolymorphic=%d, Base=%s"),
+				*i.TypeName.ToString(), (int32)i.bIsPolymorphic, *i.BaseTypeName.ToString());
+		}
+	}));
+
+	// Brute this for now.
+	TOptional<ESpeakerShortNames> LexFromString(const FString& InName)
+	{
+		for (int32 i = static_cast<int32>(ESpeakerShortNames::DefaultChannel); i < static_cast<int32>(ESpeakerShortNames::NumChannels); i++ )
+		{
+			const FString SpeakerName = LexToString(static_cast<ESpeakerShortNames>(i));
+			if (InName.Equals(SpeakerName, ESearchCase::IgnoreCase))
+			{
+				return static_cast<ESpeakerShortNames>(i);
+			}
+		}
+		// fail.
+		return {};
+	}
+
+	static FString MakePrettyString(const TArray<Audio::FDiscreteChannelTypeFamily::FSpeaker>& InSpeakers)
+	{
+		return FString::JoinBy(InSpeakers,TEXT(", "), [](const Audio::FDiscreteChannelTypeFamily::FSpeaker InSpeaker) -> FString
+		{
+			return FString::Printf(TEXT("[Name=%s,Az=%2.2f, El=%2.2f]"), *InSpeaker.ID.ToString(), InSpeaker.AzimuthDegrees,
+				InSpeaker.ElevationDegrees);
+		});
+	}
+
+	// Metasound Discrete, ingests the Schema.
+	class FMetasoundDiscreteChannelTypeFamily final : public Audio::FDiscreteChannelTypeFamily
+	{
+	public:
+		using Super = FDiscreteChannelTypeFamily;
+		static TArray<FSpeaker> MakeOrderFromSchema(const FCatDiscreteDataSchema& InSchema)
+		{
+			TArray<FSpeaker> Speakers;
+			for (const FCatDiscreteChannel& i : InSchema.Speakers)
+			{
+				Speakers.Emplace(
+					FSpeaker
+					{
+						.ID = i.ID,
+						.AzimuthDegrees = i.AzimuthDegrees,
+						.ElevationDegrees = i.ElevationDegrees,
+						.bIsSpatialized = i.bIsSpatialized,
+					});
+			}
+			
+			// Make sure if we're using the "ChannelEnum" as the order, everything can resolve to an enum entry.
+			auto NotConvertableToShortSpeakerName = [](const FSpeaker& i) -> bool { return !NameToShortSpeakerName(i.ID).IsSet(); };
+			if (InSchema.OrderPolicy == ECatDiscreteOrderPolicy::ChannelEnum && 
+				Speakers.FindByPredicate(NotConvertableToShortSpeakerName) == nullptr)
+			{
+				auto WavSortOrder = [](const FSpeaker& A,const FSpeaker& B) -> bool 
+				{
+					TOptional<ESpeakerShortNames> SpeakerA = NameToShortSpeakerName(A.ID);
+					TOptional<ESpeakerShortNames> SpeakerB = NameToShortSpeakerName(B.ID);
+					ensure(SpeakerA.IsSet());
+					ensure(SpeakerB.IsSet());
+					return static_cast<uint32>(*SpeakerA) < static_cast<uint32>(*SpeakerB);
+				};
+				Speakers.Sort(WavSortOrder);
+			}
+			else if (InSchema.OrderPolicy == ECatDiscreteOrderPolicy::Explicit)
+			{	
+				auto IndexOf = [&InSchema](const FName InName) -> int32 { return InSchema.Order.IndexOfByKey(InName); };
+				auto ExplicitSortOrder = [&IndexOf](const FSpeaker& A,const FSpeaker& B) -> bool
+					{
+						return IndexOf(A.ID) < IndexOf(B.ID);
+					};
+				Speakers.Sort(ExplicitSortOrder);
+			}
+			else
+			{
+				checkNoEntry();
+			}
+			return Speakers;
+		}
+		
+		FMetasoundDiscreteChannelTypeFamily(const FCatDiscreteDataSchema& InSchema, Audio::IChannelTypeRegistry& InRegistry)
+			: Super(InSchema.Name, InRegistry.FindChannel(InSchema.BaseType), *InSchema.FriendlyName, MakeOrderFromSchema(InSchema), InSchema.bIsParentsDefault, InSchema.bIsAbstract)
+		{
+			UE_LOG(LogMetaSound, Verbose, TEXT("Discrete: Unique=%s\tNumChannels=%d\tParent=%s\tFriendlyName=%s\tDefault=%s\tOrder=[%s]\tAbstract=%s"),
+				*InSchema.Name.ToString(),
+				Order.Num(),
+				*InSchema.BaseType.ToString(),
+				*InSchema.FriendlyName,
+				ToCStr(LexToString(InSchema.bIsParentsDefault)),
+				*MakePrettyString(Order),
+				ToCStr(LexToString(InSchema.bIsAbstract)));
+		}
+	};
+
+	class FMetasoundSoundfieldChannelTypeFamily final : public Audio::FSoundfieldChannelTypeFamily
+	{
+	public:
+		using Super = FSoundfieldChannelTypeFamily;
+		
+		FMetasoundSoundfieldChannelTypeFamily(const FCatSoundfieldDataSchema& InSchema, Audio::IChannelTypeRegistry& InRegistry)
+			: Super(InSchema.Name,  InSchema.NumOrders, InRegistry.FindChannel(InSchema.BaseType), *InSchema.FriendlyName, InSchema.bIsParentsDefault, InSchema.bIsAbstract)
+		{
+			UE_LOG(LogMetaSound, Verbose, TEXT("Soundfield: Unique=%s\tParent=%s\tFriendlyName=%s\tOrder=%d\tDefault=%s\tAbstract=%s"),
+				*InSchema.Name.ToString(),
+				*InSchema.BaseType.ToString(),
+				*InSchema.FriendlyName,
+				GetAmbisonicsOrder(),
+				ToCStr(LexToString(InSchema.bIsParentsDefault)),
+				ToCStr(LexToString(InSchema.bIsAbstract)));
+		}	
+	};
+
+	class FMetasoundCompositeChannelTypeFamily final : public Audio::FCompositeChannelTypeFamily
+	{
+	public:
+		using Super = FCompositeChannelTypeFamily;
+
+		FMetasoundCompositeChannelTypeFamily(const FCatCompositeDataSchema& InSchema, Audio::IChannelTypeRegistry& InRegistry)
+			: Super(InSchema.Name,InRegistry.FindChannel(InSchema.BaseType), *InSchema.FriendlyName, InSchema.bIsParentsDefault)
+		{
+			UE_LOG(LogMetaSound, Verbose, TEXT("Composite: Unique=%s\tParent=%s\tFriendlyName=%s\tDefault=%s\tAbstract=%s"),
+				*InSchema.Name.ToString(),
+				*InSchema.BaseType.ToString(),
+				*InSchema.FriendlyName,
+				ToCStr(LexToString(InSchema.bIsParentsDefault)),
+				ToCStr(LexToString(InSchema.bIsAbstract)));
+		}	
+	};
+
+	template <typename DataType, EVertexAccessType VertexAccess = EVertexAccessType::Reference>
+	class TPolyInputNodeOperatorFactory : public IOperatorFactory
+	{
+		static constexpr bool bIsReferenceVertexAccess = VertexAccess == EVertexAccessType::Reference;
+		static constexpr bool bIsValueVertexAccess = VertexAccess == EVertexAccessType::Value;
+
+		static_assert(bIsValueVertexAccess || bIsReferenceVertexAccess, "Unsupported EVertexAccessType");
+
+		// Choose which data reference type is created based on template parameters
+		using FPassThroughDataReference = std::conditional_t<bIsReferenceVertexAccess, TDataReadReference<DataType>, TDataValueReference<DataType>>;
+
+		// Return correct data reference type based on vertex access type for pass through scenario.
+		FPassThroughDataReference CreatePassThroughDataReference(const FAnyDataReference& InRef)
+		{
+			if constexpr (bIsReferenceVertexAccess)
+			{
+				return InRef.GetDataReadReference<DataType>();
+			}
+			else if constexpr (bIsValueVertexAccess)
+			{
+				return InRef.GetDataValueReference<DataType>();
+			}
+			else
+			{
+				static_assert("Unsupported EVertexAccessType");
+			}
+		}
+
+	public:
+		explicit TPolyInputNodeOperatorFactory(const FVertexName& InVertexName)
+			: VertexName(InVertexName)
+		{
+		}
+
+		virtual TUniquePtr<IOperator> CreateOperator(const FBuildOperatorParams& InParams, FBuildResults& OutResults) override
+		{
+			using namespace MetasoundInputNodePrivate;
+
+			if (const FAnyDataReference* Ref = InParams.InputData.FindDataReference(VertexName))
+			{
+				if constexpr (bIsReferenceVertexAccess)
+				{
+					if (EDataReferenceAccessType::Write == Ref->GetAccessType())
+					{
+						return MakeUnique<TNonOwningInputOperator<DataType, VertexAccess>>(VertexName, Ref->GetDataWriteReference<DataType>());
+					}
+				}
+				// Pass through input value
+				return MakeUnique<TPassThroughOperator<DataType, VertexAccess>>(VertexName, CreatePassThroughDataReference(*Ref));
+			}
+			else
+			{
+				// Owned input value
+				return MakeUnique<TInputOperator<DataType, VertexAccess>>(VertexName, InParams.OperatorSettings, InParams.InputData);
+			}
+		}
+
+	private:
+		FVertexName VertexName;
+	};
+	
+	// Polymorphic Specific INPUT NODE.
+	template<typename DataType, EVertexAccessType VertexAccess=EVertexAccessType::Reference>
+	class TPolyInputNode : public FInputNode
+	{
+		static constexpr bool bIsConstructorInput = VertexAccess == EVertexAccessType::Value;
+		static constexpr bool bIsSupportedConstructorInput = TIsConstructorVertexSupported<DataType>::Value && bIsConstructorInput;
+		static constexpr bool bIsReferenceInput = VertexAccess == EVertexAccessType::Reference;
+		static constexpr bool bIsSupportedReferenceInput = TLiteralTraits<DataType>::bIsParsableFromAnyLiteralType && bIsReferenceInput;
+
+		static constexpr bool bIsSupportedInput = bIsSupportedConstructorInput || bIsSupportedReferenceInput;
+
+	public:
+		// If true, this node can be instantiated by the Frontend.
+		static constexpr bool bCanRegister = bIsSupportedInput;
+		
+		static FNodeClassMetadata CreateNodeClassMetadata(const FVertexName& InVertexName, const FName InConcreteTypeName)
+		{
+			return FInputNode::CreateNodeClassMetadata(InVertexName, InConcreteTypeName, VertexAccess);
+		}
+
+			
+		explicit TPolyInputNode(FInputNodeConstructorParams&& InParams)
+		:	FInputNode(MoveTemp(InParams), InParams.PolymorphicTypeName, VertexAccess, MakeShared<TPolyInputNodeOperatorFactory<DataType, VertexAccess>>(InParams.VertexName))
+		{
+		}
+
+		explicit TPolyInputNode(const FVertexName& InVertexName, FNodeData InNodeData, TSharedRef<const FNodeClassMetadata> InClassMetadata)
+		: FInputNode(MakeShared<TPolyInputNodeOperatorFactory<DataType, VertexAccess>>(InVertexName), MoveTemp(InNodeData), MoveTemp(InClassMetadata))
+		{
+		}
+	};
+		
+	template<typename SchemaType, typename CppChannelType, typename CppDataType>
+	void AddCatDataFormat(const SchemaType& InSchema)
+	{
+		using namespace Metasound;
+		using namespace Metasound::Frontend;
+
+		using FBaseType = CppDataType;
+
+		struct FDataDrivenEntry : IDataTypeRegistryEntry 
+		{
+			static FDataTypeRegistryInfo MakeInfo(const SchemaType& InSchema)
+			{
+				FDataTypeRegistryInfo Info;
+
+				Info.DataTypeName =  InSchema.Name; 
+
+				// The preferred constructor argument type for creating instances of the data type.
+				Info.PreferredLiteralType = ELiteralType::String;
+
+				// Constructor argument support in TDataTypeLiteralFactory<TDataType>;
+				Info.bIsParsable = true;
+				Info.bIsDefaultParsable = false;
+				Info.bIsBoolParsable = false;
+				Info.bIsIntParsable = false;
+				Info.bIsFloatParsable = false;
+				Info.bIsStringParsable = true;
+				Info.bIsProxyParsable = false;
+				Info.bIsUniquePtrProxyParsable_DEPRECATED = false;
+				Info.bIsDefaultArrayParsable = false;
+				Info.bIsBoolArrayParsable = false;
+				Info.bIsIntArrayParsable = false;
+				Info.bIsFloatArrayParsable = false;
+				Info.bIsStringArrayParsable = false;
+				Info.bIsProxyArrayParsable = false;
+				Info.bIsUniquePtrProxyArrayParsable_DEPRECATED = false;
+
+				// Is a TEnum wrapped enum
+				Info.bIsEnum = false;
+
+				// Whether exact types are required
+				Info.bIsExplicit = false;
+
+				// Determines whether the type can be used with send/receive transmitters
+				Info.bIsTransmittable = false;
+
+				// Returns if DataType is a Variable type
+				Info.bIsVariable = false;
+
+				// Returns if DataType can be used for constructor vertices.
+				Info.bIsConstructorType = false;
+
+				// Returns if DataType represents an array type (ex. TArray<float>, TArray<int32>, etc.).
+				Info.bIsArrayType = false;
+
+				// Returns if DataType supports array parsing and passing array of base type to constructor.
+				Info.bIsArrayParseable = false;
+
+				// If provided in registration call, UClass this datatype was registered with.
+				Info.ProxyGeneratorClass = nullptr;
+			
+				// If this type is polymorphic.
+				Info.bIsPolymorphic = true;
+				Info.bIsAbstract = InSchema.bIsAbstract;
+
+				// If this type is polymorphic, what is its parent.
+				Info.ParentDataTypeName = InSchema.BaseType;
+			
+				return Info;
+			}
+
+			void InitNodeClasses()
+			{
+				const FName DataTypeName = Info.DataTypeName; 
+				const FName UnnamedVertex;
+
+				if (Info.bIsParsable)
+				{
+					const FNodeClassMetadata InputMeta  = FInputNode::CreateNodeClassMetadata(UnnamedVertex, DataTypeName, EVertexAccessType::Reference);
+					const FNodeClassMetadata OutputMeta = FOutputNode::CreateNodeClassMetadata(UnnamedVertex, DataTypeName, EVertexAccessType::Reference);
+					const FNodeClassMetadata LiteralMeta = FLiteralNode::GetNodeMetadata(DataTypeName, EVertexAccessType::Reference);
+
+					using namespace Metasound::VariableNames;
+
+					// Variable.
+					FNodeClassMetadata VarMeta = TVariableNode<FBaseType>::CreateNodeClassMetadata();                      
+					VarMeta.ClassName = GetVariableNodeClassName(DataTypeName);                                           
+
+					// Mutator.
+					FNodeClassMetadata MutMeta = MetasoundVariableNodesPrivate::TVariableMutatorOperator<FBaseType>::GetNodeInfo();
+					MutMeta.ClassName = GetVariableMutatorNodeClassName(DataTypeName);                                   
+
+					// Variable accessor
+					FNodeClassMetadata AccMeta = MetasoundVariableNodesPrivate::TVariableAccessorOperator<FBaseType>::GetNodeInfo();
+					AccMeta.ClassName = GetVariableAccessorNodeClassName(DataTypeName);                                 
+
+					// Deffered.
+					FNodeClassMetadata DefMeta = MetasoundVariableNodesPrivate::TVariableDeferredAccessorOperator<FBaseType>::GetNodeInfo();
+					DefMeta.ClassName = GetVariableDeferredAccessorNodeClassName(DataTypeName);
+					InitNodeClassesBase(InputMeta, OutputMeta, LiteralMeta, VarMeta, MutMeta, AccMeta, DefMeta);
+				}
+
+				if (Info.bIsConstructorType)
+				{
+					InitNodeClassesBaseConstructor(
+						FInputNode::CreateNodeClassMetadata(UnnamedVertex, DataTypeName, EVertexAccessType::Value),	
+						FOutputNode::CreateNodeClassMetadata(UnnamedVertex, DataTypeName, EVertexAccessType::Value)
+					);
+				}
+			}
+
+			void InitRawAssignmentFunction()
+			{
+				this->RawAssignmentFunction = [](const void* Src, void* Dest)
+					{
+						*(static_cast<FBaseType*>(Dest)) = *(static_cast<const FBaseType*>(Src));
+					};
+			}
+
+			void InitLiteralAssignmentFunction()
+			{
+				if  (Info.bIsParsable)
+				{
+					this->LiteralAssignmentFunction = [](const FOperatorSettings& InOperatorSettings, const FLiteral& InLiteral, const FAnyDataReference& OutDataRef)
+						{
+							*OutDataRef.GetWritableValue<FBaseType>() = TDataTypeLiteralFactory<FBaseType>::CreateExplicitArgs(InOperatorSettings, InLiteral);
+						};
+				}
+			}
+
+			void InitPolymorphicInterface(const SchemaType& InSchema)
+			{
+				this->PolymorphicInterface = MakeShared<FPolymorphicDataTypeBase>(InSchema.Name, InSchema.BaseType);
+			}
+	
+
+			FDataDrivenEntry(const SchemaType& InSchema)
+				: Info(MakeInfo(InSchema))
+			{
+				InitNodeClasses();
+				InitRawAssignmentFunction();
+				InitLiteralAssignmentFunction();
+				InitPolymorphicInterface(InSchema);
+			}
+
+			virtual ~FDataDrivenEntry() override = default;
+		
+			virtual TUniquePtr<IDataTypeRegistryEntry> Clone() const override
+			{
+				return {};
+			}
+		
+			virtual const FDataTypeRegistryInfo& GetDataTypeInfo() const override
+			{
+				return Info;
+			}
+			virtual TSharedPtr<const IEnumDataTypeInterface> GetEnumInterface() const override
+			{
+				return {}; 
+			}
+			virtual TSharedPtr<IPolymorphicDataTypeInterface> GetPolymorphicInterface() const override
+			{
+				return PolymorphicInterface;;
+			}
+		
+			virtual const FMetasoundFrontendClass& GetFrontendInputClass() const override
+			{
+				return InputClass;
+			}
+
+			virtual TSharedPtr<const Metasound::FNodeClassMetadata> GetInputClassMetadata() const override
+			{
+				return InputClassMetadata;
+			}
+
+			virtual const FMetasoundFrontendClass& GetFrontendConstructorInputClass() const override
+			{
+				return ConstructorInputClass;
+			}
+
+			virtual TSharedPtr<const Metasound::FNodeClassMetadata> GetConstructorInputClassMetadata() const override
+			{
+				return ConstructorInputClassMetadata;
+			}
+
+			virtual const FMetasoundFrontendClass& GetFrontendLiteralClass() const override
+			{
+				return LiteralClass;
+			}
+
+			virtual const FMetasoundFrontendClass& GetFrontendOutputClass() const override
+			{
+				return OutputClass;
+			}
+
+			virtual TSharedPtr<const FNodeClassMetadata> GetOutputClassMetadata() const override
+			{
+				return OutputClassMetadata;
+			}
+
+			virtual const FMetasoundFrontendClass& GetFrontendConstructorOutputClass() const override
+			{
+				return ConstructorOutputClass;
+			}
+
+			virtual TSharedPtr<const FNodeClassMetadata> GetConstructorOutputClassMetadata() const override
+			{
+				return ConstructorOutputClassMetadata;
+			}
+
+			virtual const FMetasoundFrontendClass& GetFrontendVariableClass() const override
+			{
+				return VariableClass;
+			}
+
+			virtual TSharedPtr<const FNodeClassMetadata> GetVariableClassMetadata() const override
+			{
+				return VariableClassMetadata;
+			}
+
+			virtual const FMetasoundFrontendClass& GetFrontendVariableMutatorClass() const override
+			{
+				return VariableMutatorClass;
+			}
+
+			virtual TSharedPtr<const FNodeClassMetadata> GetVariableMutatorClassMetadata() const override
+			{
+				return VariableMutatorClassMetadata;
+			}
+
+			virtual const FMetasoundFrontendClass& GetFrontendVariableAccessorClass() const override
+			{
+				return VariableAccessorClass;
+			}
+
+			virtual TSharedPtr<const FNodeClassMetadata> GetVariableAccessorClassMetadata() const override
+			{
+				return VariableAccessorClassMetadata;
+			}
+
+			virtual const FMetasoundFrontendClass& GetFrontendVariableDeferredAccessorClass() const override
+			{
+				return VariableDeferredAccessorClass;
+			}
+
+			virtual TSharedPtr<const FNodeClassMetadata> GetVariableDeferredAccessorClassMetadata() const override
+			{
+				return VariableDeferredAccessorClassMetadata;
+			}
+
+			virtual const IParameterAssignmentFunction& GetRawAssignmentFunction() const override
+			{
+				return RawAssignmentFunction;
+			}
+
+			virtual FLiteralAssignmentFunction GetLiteralAssignmentFunction() const override
+			{
+				return LiteralAssignmentFunction;
+			}
+
+			virtual TUniquePtr<INode> CreateOutputNode(FNodeData InNodeData) const override
+			{
+				if (OutputClassMetadata)
+				{
+					const FOutputVertexInterface& Outputs = InNodeData.Interface.GetOutputInterface();
+					if (ensure(Outputs.Num() == 1))
+					{
+						const FVertexName& VertexName = Outputs.At(0).VertexName;
+						return MakeUnique<FOutputNode>(VertexName, MoveTemp(InNodeData), OutputClassMetadata.ToSharedRef());
+					}
+				}
+
+				return TUniquePtr<INode>(nullptr);
+			}
+
+			virtual TUniquePtr<INode> CreateConstructorOutputNode(FNodeData InNodeData) const override
+			{
+				if (ConstructorOutputClassMetadata)
+				{
+					const FOutputVertexInterface& Outputs = InNodeData.Interface.GetOutputInterface();
+					if (ensure(Outputs.Num() == 1))
+					{
+						const FVertexName& VertexName = Outputs.At(0).VertexName;
+						return MakeUnique<FOutputNode>(VertexName, MoveTemp(InNodeData), ConstructorOutputClassMetadata.ToSharedRef());
+					}
+				}
+				return TUniquePtr<INode>(nullptr);
+			}
+
+			FORCENOINLINE void InitNodeClassesBase(
+				const FNodeClassMetadata& InInputClassMetadata,
+				const FNodeClassMetadata& InOutputClassMetadata,
+				const FNodeClassMetadata& InLiteralPrototypeMetadata,
+				const FNodeClassMetadata& InVariableClassMetadata,
+				const FNodeClassMetadata& InVariableMutatorClassMetadata,
+				const FNodeClassMetadata& InVariableAccessorClassMetadata,
+				const FNodeClassMetadata& InVariableDeferredAccessorClassMetadata
+			)
+			{
+				this->InputClassMetadata = MakeShared<FNodeClassMetadata>(InInputClassMetadata);
+				this->InputClass = GenerateClass(*this->InputClassMetadata, EMetasoundFrontendClassType::Input);
+
+				this->OutputClassMetadata = MakeShared<Metasound::FNodeClassMetadata>(InOutputClassMetadata);
+				this->OutputClass = GenerateClass(*this->OutputClassMetadata, EMetasoundFrontendClassType::Output);
+
+				this->LiteralClass = GenerateClass(InLiteralPrototypeMetadata, EMetasoundFrontendClassType::Literal);
+
+				this->VariableClassMetadata = MakeShared<FNodeClassMetadata>(InVariableClassMetadata);
+				this->VariableClass = GenerateClass(*this->VariableClassMetadata, EMetasoundFrontendClassType::Variable);
+
+				this->VariableMutatorClassMetadata = MakeShared<FNodeClassMetadata>(InVariableMutatorClassMetadata);
+				this->VariableMutatorClass = GenerateClass(*this->VariableMutatorClassMetadata, EMetasoundFrontendClassType::VariableMutator);
+
+				this->VariableAccessorClassMetadata = MakeShared<FNodeClassMetadata>(InVariableAccessorClassMetadata);
+				this->VariableAccessorClass = GenerateClass(*this->VariableAccessorClassMetadata, EMetasoundFrontendClassType::VariableAccessor);
+
+				this->VariableDeferredAccessorClassMetadata = MakeShared<FNodeClassMetadata>(InVariableDeferredAccessorClassMetadata);
+				this->VariableDeferredAccessorClass = GenerateClass(*this->VariableDeferredAccessorClassMetadata, EMetasoundFrontendClassType::VariableDeferredAccessor);
+			}
+
+			FORCENOINLINE void InitNodeClassesBaseConstructor(
+				const FNodeClassMetadata& InConstructorInputClassMetadata,
+				const FNodeClassMetadata& InConstructorOutputClassMetadata
+			)
+			{
+				this->ConstructorInputClassMetadata = MakeShared<Metasound::FNodeClassMetadata>(InConstructorInputClassMetadata);
+				this->ConstructorInputClass = Metasound::Frontend::GenerateClass(*this->ConstructorInputClassMetadata, EMetasoundFrontendClassType::Input);
+				this->ConstructorOutputClassMetadata = MakeShared<Metasound::FNodeClassMetadata>(InConstructorOutputClassMetadata);
+				this->ConstructorOutputClass = Metasound::Frontend::GenerateClass(*this->ConstructorOutputClassMetadata, EMetasoundFrontendClassType::Output);
+			}
+		
+		
+			virtual TUniquePtr<INode> CreateInputNode(Metasound::FNodeData InNodeData) const override
+			{
+				if (InputClassMetadata)
+				{
+					const FInputVertexInterface& Inputs = InNodeData.Interface.GetInputInterface();
+					if (ensure(Inputs.Num() == 1))
+					{
+						const FVertexName& VertexName = Inputs.At(0).VertexName;
+						return MakeUnique<TPolyInputNode<FBaseType, EVertexAccessType::Reference>>(VertexName, MoveTemp(InNodeData), InputClassMetadata.ToSharedRef());
+					}
+				}
+				return TUniquePtr<INode>(nullptr);
+			}
+
+			using TDataType = FBaseType;
+
+			virtual TUniquePtr<INode> CreateConstructorInputNode(FNodeData InNodeData) const override
+			{
+				if (Info.bIsParsable && Info.bIsConstructorType)
+				{
+					if (ConstructorInputClassMetadata)
+					{
+						const FInputVertexInterface& Inputs = InNodeData.Interface.GetInputInterface();
+						if (ensure(Inputs.Num() == 1))
+						{
+							const FVertexName& VertexName = Inputs.At(0).VertexName;
+							return MakeUnique<TPolyInputNode<TDataType, EVertexAccessType::Value>>(
+								VertexName, MoveTemp(InNodeData), ConstructorInputClassMetadata.ToSharedRef());
+						}
+					}
+				}
+				return TUniquePtr<INode>(nullptr);
+			}
+
+			virtual TUniquePtr<INode> CreateVariableNode(FLiteral InLiteral, FNodeData InNodeData) const override
+			{
+				if (Info.bIsParsable)
+				{
+					if (VariableClassMetadata)
+					{
+						return MakeUnique<TVariableNode<TDataType>>(MoveTemp(InLiteral), MoveTemp(InNodeData), VariableClassMetadata.ToSharedRef());
+					}
+				}
+				return TUniquePtr<INode>(nullptr);
+			}
+
+			virtual TUniquePtr<INode> CreateVariableMutatorNode(FNodeData InNodeData) const override
+			{
+				if (Info.bIsParsable)
+				{
+					if (VariableMutatorClassMetadata)
+					{
+						return MakeUnique<TVariableMutatorNode<TDataType>>(MoveTemp(InNodeData), VariableMutatorClassMetadata.ToSharedRef());
+					}
+				}
+				return TUniquePtr<INode>(nullptr);
+			}
+
+			virtual TUniquePtr<INode> CreateVariableAccessorNode(const FNodeInitData& InParams) const override
+			{
+				if (Info.bIsParsable)
+				{
+					return MakeUnique<TVariableAccessorNode<TDataType>>(InParams);
+				}
+				else
+				{
+					return TUniquePtr<INode>(nullptr);
+				}
+			}
+
+			virtual TUniquePtr<INode> CreateVariableAccessorNode(FNodeData InNodeData) const override
+			{
+				if (Info.bIsParsable)
+				{
+					if (VariableAccessorClassMetadata)
+					{
+						return MakeUnique<TVariableAccessorNode<TDataType>>(MoveTemp(InNodeData), VariableAccessorClassMetadata.ToSharedRef());
+					}
+				}
+				return TUniquePtr<INode>(nullptr);
+			}
+
+			virtual TUniquePtr<INode> CreateVariableDeferredAccessorNode(FNodeData InNodeData) const override
+			{
+				if (Info.bIsParsable)
+				{
+					if (VariableAccessorClassMetadata)
+					{
+						return MakeUnique<TVariableDeferredAccessorNode<TDataType>>(MoveTemp(InNodeData), VariableAccessorClassMetadata.ToSharedRef());
+					}
+				}
+				return TUniquePtr<INode>(nullptr);
+			}
+
+			virtual TOptional<FAnyDataReference> CreateDataReference(EDataReferenceAccessType InAccessType, const FLiteral& InLiteral,
+																	 const FOperatorSettings& InOperatorSettings) const override
+			{
+				using enum EDataReferenceAccessType;;
+				// Use the CPP types ID. For cat that will mean Derived,Composite,Soundfield.
+				auto ID = GetMetasoundDataTypeId<CppDataType>();
+				switch (InAccessType)
+				{
+				case Read: return  FAnyDataReference(TDataReadReference<FChannelAgnosticType>::CreateNewDerivedAs<CppDataType>(Info.DataTypeName, ID, InOperatorSettings, Info.DataTypeName));
+				case Write: return FAnyDataReference(TDataWriteReference<FChannelAgnosticType>::CreateNewDerivedAs<CppDataType>(Info.DataTypeName, ID, InOperatorSettings, Info.DataTypeName));
+				case Value: return FAnyDataReference(TDataValueReference<FChannelAgnosticType>::CreateNewDerivedAs<CppDataType>(Info.DataTypeName, ID, InOperatorSettings, Info.DataTypeName));
+				default:
+					checkNoEntry();
+				}
+				return {};
+			}
+
+
+			virtual TSharedPtr<Audio::IProxyData> CreateProxy(UObject* InObject) const override { return {}; }
+			virtual TSharedPtr<IDataChannel, ESPMode::ThreadSafe> CreateDataChannel(const Metasound::FOperatorSettings&) const override { return {}; }
+
+			// stubbed.
+			virtual TUniquePtr<INode> CreateVariableNode(Metasound::FVariableNodeConstructorParams&& InInitParams) const override
+			{
+				return {};
+			}
+
+			virtual TUniquePtr<INode> CreateLiteralNode(Metasound::FLiteralNodeConstructorParams&&) const override
+			{
+				return {};
+			}
+
+			virtual TUniquePtr<INode> CreateReceiveNode(const Metasound::FNodeInitData&) const override
+			{
+				return {};
+			}
+
+			virtual TUniquePtr<INode> CreateConstructorOutputNode(Metasound::FOutputNodeConstructorParams&&) const override
+			{
+				return {};
+			}
+
+			virtual TUniquePtr<INode> CreateOutputNode(Metasound::FOutputNodeConstructorParams&&) const override
+			{
+				return {};
+			}
+
+			virtual TUniquePtr<INode> CreateVariableMutatorNode(const Metasound::FNodeInitData&) const override
+			{
+				return {};
+			}
+
+			virtual TUniquePtr<INode> CreateVariableDeferredAccessorNode(const Metasound::FNodeInitData&) const override
+			{
+				return {};
+			}
+
+			virtual TUniquePtr<INode> CreateInputNode(Metasound::FInputNodeConstructorParams&&) const override
+			{
+				return {};
+			}
+
+			virtual TUniquePtr<INode> CreateConstructorInputNode(Metasound::FInputNodeConstructorParams&&) const override
+			{
+				return {};
+			}
+
+			FDataTypeRegistryInfo Info;
+			FMetasoundFrontendClass InputClass;
+			TSharedPtr<const FNodeClassMetadata> InputClassMetadata;
+			FMetasoundFrontendClass ConstructorInputClass;
+			TSharedPtr<const FNodeClassMetadata> ConstructorInputClassMetadata;
+			FMetasoundFrontendClass OutputClass;
+			TSharedPtr<const FNodeClassMetadata> OutputClassMetadata;
+			FMetasoundFrontendClass ConstructorOutputClass;
+			TSharedPtr<const FNodeClassMetadata> ConstructorOutputClassMetadata;
+
+			FMetasoundFrontendClass LiteralClass;
+
+			FMetasoundFrontendClass VariableClass;
+			TSharedPtr<const FNodeClassMetadata> VariableClassMetadata;
+			FMetasoundFrontendClass VariableMutatorClass;
+			TSharedPtr<const FNodeClassMetadata> VariableMutatorClassMetadata;
+			FMetasoundFrontendClass VariableAccessorClass;
+			TSharedPtr<const FNodeClassMetadata> VariableAccessorClassMetadata;
+			FMetasoundFrontendClass VariableDeferredAccessorClass;
+			TSharedPtr<const FNodeClassMetadata> VariableDeferredAccessorClassMetadata;
+			TSharedPtr<IEnumDataTypeInterface> EnumInterface;
+			TSharedPtr<IPolymorphicDataTypeInterface> PolymorphicInterface;
+			IParameterAssignmentFunction RawAssignmentFunction;
+			FLiteralAssignmentFunction LiteralAssignmentFunction = nullptr;
+		};
+
+		ensure(Audio::GetChannelRegistry().RegisterType(InSchema.Name, MakeUnique<CppChannelType>(InSchema, Audio::GetChannelRegistry())));		
+		ensure(IDataTypeRegistry::Get().RegisterDataType(MakeUnique<FDataDrivenEntry>(InSchema)));
+	}
+
+	template<typename SchemaType>
+	bool BuildSchemaLoadOrder(const TMap<FName, SchemaType>& InSchemas, 
+		TArray<FName>& OutLoadOrder, FName& OutProblemNode, FName& OutProblemParent, TArray<FName>& OutCyclicNodes)
+	{
+		OutLoadOrder.Reset();
+		OutCyclicNodes.Reset();
+		OutProblemNode = {};
+		OutProblemParent = {};
+		
+		TMap<FName, TArray<FName>> ChildrenMap;
+		TMap<FName, int32> Degree;
+		ChildrenMap.Reserve(InSchemas.Num());
+		Degree.Reserve(InSchemas.Num());
+		
+		for (auto i : InSchemas)
+		{
+			Degree.Add(i.Key, 0);
+			ChildrenMap.FindOrAdd(i.Key);
+		}
+		for (auto i : InSchemas)
+		{
+			const FName Child = i.Key;
+			const FName Parent = i.Value.BaseType;
+
+			// No parent?
+			if (Parent.IsNone())
+			{
+				continue;
+			}
+			// In Loading Q?
+			else if (InSchemas.Contains(Parent))
+			{
+				ChildrenMap.FindOrAdd(Parent).Add(Child);
+				++Degree.FindChecked(Child);
+			}
+			// Already registered as a type?
+			else if (!Audio::GetChannelRegistry().FindChannel(Parent))
+			{
+				OutProblemNode = Child;
+				OutProblemParent = Parent;
+				return false;  // fail.
+			}
+		}
+		
+		TArray<FName> Ready;
+		Ready.Reserve(InSchemas.Num());
+		for (auto i : Degree)
+		{
+			// Add everything with 
+			if (i.Value == 0)
+			{
+				Ready.Add(i.Key);
+			}
+		}
+		
+		while (Ready.Num() > 0)
+		{
+			const FName Current = Ready.Pop();
+			OutLoadOrder.Add(Current);
+			if (const TArray<FName>* Children  = ChildrenMap.Find(Current))
+			{
+				for (const FName Child : *Children)
+				{
+					if (int32& Deg = Degree.FindChecked(Child); !--Deg)
+					{
+						Ready.Add(Child);
+					}
+				}
+			}
+		}
+				
+		if (OutLoadOrder.Num() != InSchemas.Num())
+		{
+			for (auto i : Degree)
+			{
+				if (i.Value > 0)
+				{
+					OutCyclicNodes.Add(i.Key);
+				}
+			}
+			
+			// Cyclic, fail.
+			return false;
+		}
+		
+		// Success.
+		return true;
+	}
+	
+	template<typename SchemaType, typename CppChannelType, typename CppDataType>
+	TArray<FName> LoadSchemas(const FString& Path)
+	{
+		TArray<FString> SchemasFiles;
+		IFileManager::Get().FindFiles(SchemasFiles, *Path, TEXT(".json"));
+		
+		TMap<FName, SchemaType> Schemas;
+		for (int32 i = 0; i < SchemasFiles.Num(); i++)
+		{
+			const FString File = Path / SchemasFiles[i];
+			if (TUniquePtr<FArchive> Reader = TUniquePtr<FArchive>(IFileManager::Get().CreateFileReader(*File)))
+			{
+				TJsonStructDeserializerBackend<DefaultCharType> Backend(*Reader);
+				SchemaType Schema;
+				FStructDeserializer::Deserialize(Schema, Backend);
+				if (ensureMsgf(Schema.FamilyType == CppChannelType::GetFamilyTypeName(),
+					TEXT("Failed load schema. File=%s, LastError=%s"), *File, *Backend.GetLastErrorMessage()))
+				{
+					Schema.Name = AddCatNamespaceToName(Schema.Name);
+					Schema.BaseType = AddCatNamespaceToName(Schema.BaseType);
+					ensure(Schemas.Find(Schema.Name) == nullptr);
+					Schemas.Add(Schema.Name, Schema);
+				}
+			}
+		}
+		
+		TArray<FName> Order;
+		TArray<FName> CyclicNodes;
+		if (FName ProblemNode, ProblemParent; !BuildSchemaLoadOrder(Schemas, Order, ProblemNode, ProblemParent, CyclicNodes))
+		{
+			UE_CLOG(!ProblemNode.IsNone(), LogMetaSound, Error, TEXT("LoadChannelAgosticDataSchemas: Error failed to find parent '%s' in schema '%s'"), 
+				*ProblemParent.ToString(), *ProblemNode.ToString());	
+			
+			// Log each cyclic.
+			if (CyclicNodes.Num() > 0)
+			{
+				auto WalkCycle = [&Schemas](FName Node) -> TArray<FName>
+				{
+					TArray<FName> Cycles;
+					TSet<FName> Seen;
+					while (!Seen.Contains(Node))
+					{
+						Cycles.Add(Node); 
+						Seen.Add(Node);
+						Node = Schemas[Node].BaseType;
+					}
+					return Cycles;
+				};
+				for (FName i : CyclicNodes)
+				{
+					TArray<FName> Cyclic = WalkCycle(i);
+					UE_LOG(LogMetaSound, Error, TEXT("LoadChannelAgosticDataSchemas: Found cyclic loop in data schema '%s'"), 
+						*FString::JoinBy(Cyclic, TEXT(" -> "), [](const FName Name){ return Name.ToString(); }));
+				}
+			}
+				
+			return {};
+		}
+		for (const FName i : Order)
+		{
+			AddCatDataFormat<SchemaType,CppChannelType, CppDataType>(*Schemas.Find(i));
+		}
+		
+		return Order;			
+	}
+	
+	TArray<FName> LoadChannelAgnosticDataSchemas(const FString& InContentPath)
+	{
+		TArray<FName> LoadedSchemas;
+		const FString SchemaPath = InContentPath / TEXT("Schemas") / TEXT("ChannelAgnostic");
+		LoadedSchemas.Append(LoadSchemas<FCatDiscreteDataSchema,		FMetasoundDiscreteChannelTypeFamily,	FDiscreteChannelAgnosticType>(		SchemaPath / TEXT("Discrete")));
+		LoadedSchemas.Append(LoadSchemas<FCatSoundfieldDataSchema,	FMetasoundSoundfieldChannelTypeFamily,	FSoundfieldChannelAgnosticType>(	SchemaPath / TEXT("Soundfield")));
+		LoadedSchemas.Append(LoadSchemas<FCatCompositeDataSchema,	FMetasoundCompositeChannelTypeFamily,	FCompositeChannelAgnosticType>(		SchemaPath / TEXT("Composite")));
+		UE_LOG(LogMetaSound, Display, TEXT("LoadChannelAgosticDataSchemas: Loaded %d schemas from [%s]"),LoadedSchemas.Num(), *SchemaPath);	
+		
+		return LoadedSchemas;
+	}
+
+	void UnloadChannelAgnosticDataSchemas(const TArray<FName>& InSchemas)
+	{
+		for (const FName i : InSchemas)
+		{
+			Frontend::IDataTypeRegistry::Get().UnregisterDataType(i);
+			Audio::GetChannelRegistry().UnregisterType(i);
+		}
+	}
+	
+	// Put these here for now.
+	FName AddNamespaceToName(const FName InDatatype, const FName InNamespace)
+	{
+		if (InNamespace != InDatatype && !InDatatype.ToString().StartsWith(InNamespace.ToString()))
+		{
+			return *FString::Printf(TEXT("%s%s"), *InNamespace.ToString(), *InDatatype.ToString());
+		}
+		return InDatatype;
+	}
+
+	FName AddCatNamespaceToName(const FName InDatatype)
+	{
+		static const FName Cat(TEXT("Cat:"));
+		return AddNamespaceToName(InDatatype, Cat);
+	}		
+
+
+} 
+
+#undef LOCTEXT_NAMESPACE
diff --git a/Engine/Plugins/Runtime/Metasound/Source/MetasoundFrontend/Private/MetasoundChannelAgnosticType.cpp b/Engine/Plugins/Runtime/Metasound/Source/MetasoundFrontend/Private/MetasoundChannelAgnosticType.cpp
new file mode 100644
--- /dev/null
+++ b/Engine/Plugins/Runtime/Metasound/Source/MetasoundFrontend/Private/MetasoundChannelAgnosticType.cpp
@@ -0,0 +1,113 @@
+﻿// Copyright Epic Games, Inc. All Rights Reserved.
+
+#include "MetasoundChannelAgnosticType.h"
+#include "TypeFamily/TypeFamily.h"
+#include "TypeFamily/ChannelTypeFamily.h"
+#include "MetasoundPolymorphic.h"
+
+namespace Metasound
+{
+	namespace ChannelAgnosticTypePrivate
+	{
+		static void EnsureDefaultTypesRegistered();
+		
+		static const FName NAME_Mono(TEXT("Cat:Mono"));
+		static const FName NAME_Mono1Dot0(TEXT("Cat:Mono1Dot0"));
+		static const FName NAME_Discrete(TEXT("Cat:Discrete"));
+	}
+	
+	/**
+	 * Returns the default (i.e. default constructed) format of a CAT. Static.
+	 * @return Name to Pass to Registry. 
+	 */
+	FName FChannelAgnosticType::GetDefaultCatFormat()
+	{
+		static const FName DefaultFormat(TEXT("Cat:Mono1Dot0"));
+		if (Audio::GetChannelRegistry().FindConcreteChannel(DefaultFormat) == nullptr)
+		{
+			UE_LOG(LogMetaSound, Warning, TEXT("The default channel type '%s' is not registered. Calling EnsureDefaultTypesRegistered() "), 
+				*DefaultFormat.ToString());
+			ChannelAgnosticTypePrivate::EnsureDefaultTypesRegistered();
+		}
+		return DefaultFormat;
+	}
+
+	FChannelAgnosticType::FChannelAgnosticType(const FOperatorSettings& InSettings, const FName& InChannelTypeName)
+		: FChannelAgnosticType(InSettings.GetNumFramesPerBlock(), InChannelTypeName)
+	{}
+	
+	FChannelAgnosticType::FChannelAgnosticType(const int32 InNumFramesPerBlock, const FName& InChannelTypeName)
+		: Super(Audio::GetChannelRegistry().FindConcreteChannelChecked(InChannelTypeName), InNumFramesPerBlock)
+	{}
+	
+	// Register.
+	REGISTER_METASOUND_DATATYPE(FChannelAgnosticType, "Cat");
+	REGISTER_METASOUND_POLY_TYPE(FChannelAgnosticType);
+
+	REGISTER_METASOUND_DATATYPE(FCompositeChannelAgnosticType, "Cat:Composite");
+	REGISTER_METASOUND_POLY_TYPE(FCompositeChannelAgnosticType);
+
+	REGISTER_METASOUND_DATATYPE(FDiscreteChannelAgnosticType, "Cat:Discrete");
+	REGISTER_METASOUND_POLY_TYPE(FDiscreteChannelAgnosticType);
+
+	REGISTER_METASOUND_DATATYPE(FSoundfieldChannelAgnosticType, "Cat:Soundfield");
+	REGISTER_METASOUND_POLY_TYPE(FSoundfieldChannelAgnosticType);
+	
+	namespace ChannelAgnosticTypePrivate 
+	{
+		void EnsureDefaultTypesRegistered()
+		{
+			using namespace Audio; 
+			
+			IChannelTypeRegistry& Registry = GetChannelRegistry();
+		
+			if (!Registry.FindChannel(NAME_Discrete))
+			{
+				RegisterCatBaseClassesWithChannelRegistry();
+			}
+					
+			if (!Registry.FindChannel(NAME_Mono))
+			{
+				// Mono abstract.
+				Registry.RegisterType(
+					NAME_Mono,
+					MakeUnique<FDiscreteChannelTypeFamily>(
+						NAME_Mono,
+						Registry.FindChannel(NAME_Discrete), 
+						TEXT("Mono"),
+						TArray<FDiscreteChannelTypeFamily::FSpeaker>(),
+						true,		// Parents default
+						true));		// Abstract
+			}
+			
+			if (!Registry.FindChannel(NAME_Mono1Dot0))
+			{
+				// Mono (1.0)
+				Registry.RegisterType(
+					NAME_Mono1Dot0,
+					MakeUnique<FDiscreteChannelTypeFamily>(
+						NAME_Mono1Dot0,
+						Registry.FindChannel(NAME_Mono), 
+						TEXT("Mono (1.0)"),
+						TArray
+							{
+								FDiscreteChannelTypeFamily::FSpeaker
+								{
+									.ID =TEXT("FL"),		// Note use of FL (but azimuth of center).
+									.AzimuthDegrees = 0,
+									.ElevationDegrees = 0,
+									.bIsSpatialized = true
+								}
+							},
+						true,			// Parents default.
+						false));		// Abstract
+			}
+		}
+	}
+}
+
+
+
+
+
+
diff --git a/Engine/Plugins/Runtime/Metasound/Source/MetasoundFrontend/Private/MetasoundFrontend.cpp b/Engine/Plugins/Runtime/Metasound/Source/MetasoundFrontend/Private/MetasoundFrontend.cpp
--- a/Engine/Plugins/Runtime/Metasound/Source/MetasoundFrontend/Private/MetasoundFrontend.cpp
+++ b/Engine/Plugins/Runtime/Metasound/Source/MetasoundFrontend/Private/MetasoundFrontend.cpp
@@ -198,6 +198,8 @@ class FMetasoundFrontendModule : public IModuleInterface
 #if WITH_EDITORONLY_DATA
 		FVersioningManager::Initialize();
 #endif // WITH_EDITORONLY_DATA
+
+	    FModuleManager::LoadModuleChecked<IModuleInterface>("AudioChannelAgnosticCore");
 	}
 
 	virtual void ShutdownModule() override
diff --git a/Engine/Plugins/Runtime/Metasound/Source/MetasoundFrontend/Public/MetasoundChannelAgnosticDataSchemas.h b/Engine/Plugins/Runtime/Metasound/Source/MetasoundFrontend/Public/MetasoundChannelAgnosticDataSchemas.h
new file mode 100644
--- /dev/null
+++ b/Engine/Plugins/Runtime/Metasound/Source/MetasoundFrontend/Public/MetasoundChannelAgnosticDataSchemas.h
@@ -0,0 +1,108 @@
+﻿// Copyright Epic Games, Inc. All Rights Reserved.
+
+#pragma once
+
+#include "UObject/NameTypes.h"
+
+#include "MetasoundChannelAgnosticDataSchemas.generated.h"
+
+namespace Metasound
+{
+	UE_EXPERIMENTAL(5.8, "Channel Angostic Types are Experimental")
+	METASOUNDFRONTEND_API TArray<FName> LoadChannelAgnosticDataSchemas(const FString& InContentPath);
+	
+	UE_EXPERIMENTAL(5.8, "Channel Angostic Types are Experimental")		
+	METASOUNDFRONTEND_API void UnloadChannelAgnosticDataSchemas(const TArray<FName>&);
+
+	// Put these here for now.
+	UE_EXPERIMENTAL(5.8, "Channel Angostic Types are Experimental")
+	METASOUNDFRONTEND_API FName AddNamespaceToName(const FName InDatatype, const FName InNamespace);
+
+	UE_EXPERIMENTAL(5.8, "Channel Angostic Types are Experimental")	
+	METASOUNDFRONTEND_API FName AddCatNamespaceToName(const FName InDatatype);
+}
+
+USTRUCT()
+struct FCatDataTypeSchema
+{
+	UE_EXPERIMENTAL(5.8, "Channel Angostic Types are Experimental")
+	GENERATED_BODY()
+
+	UPROPERTY()
+	FName Name;
+
+	UPROPERTY()
+	FName BaseType;
+
+	UPROPERTY()
+	FName FamilyType;
+
+	UPROPERTY()
+	FString FriendlyName;
+
+	UPROPERTY()
+	bool bIsAbstract = false;
+
+	UPROPERTY()
+	bool bIsParentsDefault = false;
+};
+
+USTRUCT()
+struct FCatSoundfieldDataSchema : public FCatDataTypeSchema
+{
+	UE_EXPERIMENTAL(5.8, "Channel Angostic Types are Experimental")	
+	GENERATED_BODY()
+
+	UPROPERTY()
+	int32 NumOrders = 1;
+};
+
+USTRUCT()
+struct FCatDiscreteChannel
+{
+	UE_EXPERIMENTAL(5.8, "Channel Angostic Types are Experimental")	
+	GENERATED_BODY()
+
+	UPROPERTY()
+	float ElevationDegrees = 0;			// In Degrees. +180 to -180
+
+	UPROPERTY()
+	float AzimuthDegrees = 0;			// -180 to 180 azimuth.
+
+	UPROPERTY()
+	FName ID;							// ID of the Speaker. This could be unique ID or match the Short form of the Speaker ID. 
+
+	UPROPERTY()
+	bool bIsSpatialized = true;			// If this is off, azimuth+degrees do nothing. (for excluding from panner). 
+};
+
+UENUM()
+UE_EXPERIMENTAL(5.8, "Channel Angostic Types are Experimental")	
+enum class ECatDiscreteOrderPolicy : uint32
+{
+	ChannelEnum,					// Format uses the order of the AudioMixerChannel (which is the same as WavFormatEx) 
+	Explicit,						// If explicit, use the Order[] array here to define it. 
+};
+
+USTRUCT()
+struct FCatDiscreteDataSchema : public FCatDataTypeSchema
+{
+	UE_EXPERIMENTAL(5.8, "Channel Angostic Types are Experimental")	
+	GENERATED_BODY()
+
+	UPROPERTY()
+	ECatDiscreteOrderPolicy OrderPolicy = ECatDiscreteOrderPolicy::ChannelEnum;
+	
+	UPROPERTY()
+	TArray<FName> Order;
+	
+	UPROPERTY()
+	TArray<FCatDiscreteChannel> Speakers;
+};
+
+USTRUCT()
+struct FCatCompositeDataSchema : public FCatDataTypeSchema
+{
+	UE_EXPERIMENTAL(5.8, "Channel Angostic Types are Experimental")	
+	GENERATED_BODY()
+};
\ No newline at end of file
diff --git a/Engine/Plugins/Runtime/Metasound/Source/MetasoundFrontend/Public/MetasoundChannelAgnosticType.h b/Engine/Plugins/Runtime/Metasound/Source/MetasoundFrontend/Public/MetasoundChannelAgnosticType.h
new file mode 100644
--- /dev/null
+++ b/Engine/Plugins/Runtime/Metasound/Source/MetasoundFrontend/Public/MetasoundChannelAgnosticType.h
@@ -0,0 +1,101 @@
+// Copyright Epic Games, Inc. All Rights Reserved.
+
+#pragma once
+
+#include "MetasoundPolymorphic.h"
+#include "ChannelAgnostic/ChannelAgnosticType.h"
+#include "MetasoundDataTypeRegistrationMacro.h"
+
+namespace Metasound
+{
+	class FChannelAgnosticType : public Audio::FChannelAgnosticType
+	{
+	public:
+		using Super = Audio::FChannelAgnosticType;
+		using Super::Super;
+
+		static FName GetRegisteredTypeName()
+		{
+			static const FName Name("Cat");			
+			return Name;
+		}						
+
+		/**
+		 * Returns the default (i.e. default constructed) format of a CAT.
+		 * @return Name to Pass to Registry. 
+		 */
+		UE_EXPERIMENTAL(5.8, "Channel Angostic Types are Experimental")
+		METASOUNDFRONTEND_API static FName GetDefaultCatFormat();
+
+		UE_EXPERIMENTAL(5.8, "Channel Angostic Types are Experimental")
+		METASOUNDFRONTEND_API explicit FChannelAgnosticType(const FOperatorSettings& InSettings, const FName& InChannelTypeName = GetDefaultCatFormat());
+		
+		UE_EXPERIMENTAL(5.8, "Channel Angostic Types are Experimental")
+		METASOUNDFRONTEND_API explicit FChannelAgnosticType(const int32 InNumFramesPerBlock, const FName& InChannelTypeName = GetDefaultCatFormat());
+		
+		static FName LiteralToFName(const FString& In)
+		{
+			if (In.IsEmpty())
+			{
+				static const FName Default = TEXT("Cat:Mono1Dot0");
+				return Default;
+			}
+			return FName(In);
+		}
+		/* FLiteral compatible Constructor
+		 * TODO: Literals only support FString and not FNames, so there is a conversion here.
+		 */
+		PRAGMA_DISABLE_EXPERIMENTAL_WARNINGS		
+		UE_EXPERIMENTAL(5.8, "Channel Angostic Types are Experimental")
+		explicit FChannelAgnosticType(const FOperatorSettings& InSettings, const FString& InChannelTypeName)
+			: FChannelAgnosticType(InSettings, LiteralToFName((InChannelTypeName)))
+		{}
+		PRAGMA_ENABLE_EXPERIMENTAL_WARNINGS
+	};
+
+	class FDiscreteChannelAgnosticType : public FChannelAgnosticType
+	{
+		static FName GetRegisteredTypeName()
+		{
+			static const FName Name("Cat:Discrete");			
+			return Name;
+		}				
+		using Super = FChannelAgnosticType;
+		using Super::Super;
+	};
+	class FSoundfieldChannelAgnosticType : public FChannelAgnosticType
+	{
+		static FName GetRegisteredTypeName()
+		{
+			static const FName Name("Cat:Soundfield");			
+			return Name;
+		}		
+		using Super = FChannelAgnosticType;
+		using Super::Super;
+	};
+	class FCompositeChannelAgnosticType : public FChannelAgnosticType
+	{
+		static FName GetRegisteredTypeName()
+		{
+			static const FName Name("Cat:Composite");			
+			return Name;
+		}
+		using Super = FChannelAgnosticType;
+		using Super::Super;
+	};
+}
+
+namespace Metasound
+{
+	// Declare Metasound Data reference Types. 
+	DECLARE_METASOUND_DATA_REFERENCE_TYPES(Metasound::FChannelAgnosticType,				METASOUNDFRONTEND_API, FChannelAgnosticTypeTypeInfo,				FChannelAgnosticTypeReadRef,			FChannelAgnosticTypeWriteRef);
+	DECLARE_METASOUND_DATA_REFERENCE_TYPES(Metasound::FCompositeChannelAgnosticType,	METASOUNDFRONTEND_API, FCompositeChannelAgnosticTypeTypeInfo,	FCompositeChannelAgnosticTypeReadRef,	FCompositeChannelAgnosticTypeWriteRef);
+	DECLARE_METASOUND_DATA_REFERENCE_TYPES(Metasound::FDiscreteChannelAgnosticType,		METASOUNDFRONTEND_API, FDiscreteChannelAgnosticTypeTypeInfo,		FDiscreteChannelAgnosticTypeReadRef,	FDiscreteChannelAgnosticTypeWriteRef);
+	DECLARE_METASOUND_DATA_REFERENCE_TYPES(Metasound::FSoundfieldChannelAgnosticType,	METASOUNDFRONTEND_API, FSoundfieldChannelAgnosticTypeTypeInfo,	FSoundfieldChannelAgnosticTypeReadRef,	FSoundChannelAgnosticTypeWriteRef);
+
+	// Declare Polymorphic DataTypes.
+	DECLARE_METASOUND_POLY_TYPE(FChannelAgnosticType,			void, true);
+	DECLARE_METASOUND_POLY_TYPE(FCompositeChannelAgnosticType,	Metasound::FChannelAgnosticType, true);
+	DECLARE_METASOUND_POLY_TYPE(FDiscreteChannelAgnosticType,	Metasound::FChannelAgnosticType, true);
+	DECLARE_METASOUND_POLY_TYPE(FSoundfieldChannelAgnosticType,	Metasound::FChannelAgnosticType, true);
+}
\ No newline at end of file
diff --git a/Engine/Source/Runtime/AudioChannelAgnosticCore/AudioChannelAgnosticCore.Build.cs b/Engine/Source/Runtime/AudioChannelAgnosticCore/AudioChannelAgnosticCore.Build.cs
new file mode 100644
--- /dev/null
+++ b/Engine/Source/Runtime/AudioChannelAgnosticCore/AudioChannelAgnosticCore.Build.cs
@@ -0,0 +1,44 @@
+﻿// Copyright Epic Games, Inc. All Rights Reserved.
+
+using UnrealBuildTool;
+
+public class AudioChannelAgnosticCore : ModuleRules
+{
+	public AudioChannelAgnosticCore(ReadOnlyTargetRules Target) : base(Target)
+	{
+		PCHUsage = PCHUsageMode.UseExplicitOrSharedPCHs;
+
+		// Disable warnings of Experimental API
+		bValidateExperimentalApi = false;
+		
+		PublicDependencyModuleNames.AddRange(
+			new string[]
+			{
+				"Core",
+			}
+		);
+		
+		// We want to use enum definitions in AudioMixer without an Engine dependency.
+		PublicIncludePathModuleNames.AddRange(
+			new string[]
+			{
+				"AudioMixer"
+			}
+		);
+		
+		// When the Editor builds all plugins, it will build with engine. In that case add a dependency on CoreUObject.
+		// If we don't do this we will get linker errors for the Engine expecting to find stuff defined here.
+		if (Target.bCompileAgainstEngine)
+		{
+			PublicDependencyModuleNames.Add("CoreUObject");
+		}
+		
+		PrivateDependencyModuleNames.AddRange(
+			new string[]
+			{
+				"Core",
+				"SignalProcessing",
+			}
+		);
+	}
+}
\ No newline at end of file
diff --git a/Engine/Source/Runtime/AudioChannelAgnosticCore/AudioChannelAgnosticCoreModule.cpp b/Engine/Source/Runtime/AudioChannelAgnosticCore/AudioChannelAgnosticCoreModule.cpp
new file mode 100644
--- /dev/null
+++ b/Engine/Source/Runtime/AudioChannelAgnosticCore/AudioChannelAgnosticCoreModule.cpp
@@ -0,0 +1,21 @@
+﻿// Copyright Epic Games, Inc. All Rights Reserved.
+
+#include "CoreMinimal.h"
+#include "Modules/ModuleManager.h"
+#include "TypeFamily/ChannelTypeFamily.h"
+
+class FAudioChannelAgnosticCoreModule final : public IModuleInterface
+{
+public:
+	virtual void StartupModule() override
+	{
+		Audio::RegisterCatBaseClassesWithChannelRegistry();
+	}
+	
+	virtual void ShutdownModule() override
+	{
+		Audio::UnregisterCatBaseClassesWithChannelRegistry();
+	}
+};
+    
+IMPLEMENT_MODULE(FAudioChannelAgnosticCoreModule, AudioChannelAgnosticCore);
\ No newline at end of file
diff --git a/Engine/Source/Runtime/AudioChannelAgnosticCore/Private/ChannelAgnostic/ChannelAgnosticTranscoding.cpp b/Engine/Source/Runtime/AudioChannelAgnosticCore/Private/ChannelAgnostic/ChannelAgnosticTranscoding.cpp
new file mode 100644
--- /dev/null
+++ b/Engine/Source/Runtime/AudioChannelAgnosticCore/Private/ChannelAgnostic/ChannelAgnosticTranscoding.cpp
@@ -0,0 +1,334 @@
+﻿// Copyright Epic Games, Inc. All Rights Reserved.
+
+#include "ChannelAgnostic/ChannelAgnosticTranscoding.h"
+
+#include "Algo/Transform.h"
+#include "DSP/Ambisonics.h"
+#include "DSP/MultiMono.h"
+#include "DSP/SphericalHarmonicCalculator.h"
+#include "DSP/Vbap.h"
+#include "TypeFamily/ChannelTypeFamily.h"
+
+namespace Audio
+{
+	bool MakeGainMatrixFrom(const FDiscreteChannelTypeFamily& InSrc, const FDiscreteChannelTypeFamily& InDst, TArray<float>& OutMatrix);
+	
+	namespace ChannelAgnosticTranscoder
+	{
+		// Discrete to Discrete
+		FTranscoder GetTranscoder(const FDiscrete& InSrcType, const FDiscrete& InDstType, const FParams& InParams)
+		{
+			// Exact match? We can just memcpy each channel.
+			// TODO. in future these could be shared-ptrs from the main CAT memory block.
+			if (&InDstType == &InSrcType)
+			{
+				const int32 NumChannels = InDstType.NumChannels();
+				return [NumChannels](TArrayView<const float*> Src, TArrayView<float*> Dst, const int32 NumFrames)
+					{
+						for (int32 i = 0; i < NumChannels; ++i)
+						{
+							FMemory::Memcpy(Dst[i], Src[i], NumFrames * sizeof(float));
+						}
+					};
+			}
+
+			switch (InParams.TranscodeMethod)
+			{
+			case EChannelTranscodeMethod::ChannelDrop:
+				{
+					return [&InSrcType,&InDstType](TArrayView<const float*> SrcChannels, TArrayView<float*> DstChannels, const int32 NumFrames)
+						{
+							checkSlow(DstChannels.Num() == NumFrames);
+							// Copy everything destination wants, and nothing else.
+							const TArray<FDiscreteChannelTypeFamily::FSpeaker>& SrcOrder = InSrcType.GetSpeakerOrder();
+							const TArray<FDiscreteChannelTypeFamily::FSpeaker>& DstOrder = InDstType.GetSpeakerOrder();
+							const int32 NumDstChannels = InDstType.NumChannels();	
+							for (int32 i = 0; i < NumDstChannels; ++i)
+							{
+								const FName DstSpeaker = DstOrder[i].ID;
+								if (const int32 SrcChannelIndex = SrcOrder.IndexOfByPredicate([&](const FDiscreteChannelTypeFamily::FSpeaker& j) { return j.ID == DstSpeaker; });
+									SrcChannelIndex != INDEX_NONE)
+								{
+									FMemory::Memcpy(DstChannels[i], SrcChannels[SrcChannelIndex], NumFrames * sizeof(float));
+								}
+								else
+								{
+									FMemory::Memzero(DstChannels[i], NumFrames * sizeof(float));	
+								}
+							}
+						};
+				}
+
+			case EChannelTranscodeMethod::MixUpOrDown:
+				{
+					if (TArray<float> Gains; MakeGainMatrixFrom(InSrcType, InDstType, Gains))
+					{
+						return [MixGains = MoveTemp(Gains)](TArrayView<const float*> InSrc, TArrayView<float*> InDst, const int32 NumFrames)
+							{
+								MultiMonoMixUpOrDown(InSrc, InDst, NumFrames, MixGains);
+							};					
+					}
+				}
+			}
+			// fail.
+			return {};
+		}
+
+		// Discrete to Ambisonics
+		FTranscoder GetTranscoder(const FDiscrete& InFromType, const FSoundfield& InToType, const FParams&)
+		{
+			const TArray<FDiscreteChannelTypeFamily::FSpeaker>& Speakers = InFromType.GetSpeakerOrder();
+			TArray<TArray<float>> AllChannelGainArrays;
+			AllChannelGainArrays.SetNum(InFromType.NumChannels());
+
+			for (int32 i = 0; i < InFromType.NumChannels(); ++i)
+			{
+				const FDiscreteChannelTypeFamily::FSpeaker& Speaker = Speakers[i];
+			
+				TArray<float>& GainArray = AllChannelGainArrays[i];
+				GainArray.SetNumZeroed(InToType.NumChannels());
+
+				// Skip LFE.
+				if (Speakers[i].ID == LexToName(ESpeakerShortNames::LFE))
+				{
+					continue;
+				}
+			
+				const float AzimuthRads = FMath::DegreesToRadians(Speaker.AzimuthDegrees);
+				const float ElevationRads = FMath::DegreesToRadians(Speaker.ElevationDegrees);
+				FSphericalHarmonicCalculator::ComputeSoundfieldChannelGains(InToType.GetAmbisonicsOrder(), AzimuthRads, ElevationRads, GainArray);
+				FSphericalHarmonicCalculator::NormalizeGains(GainArray);
+			}
+	
+			return [Gains = MoveTemp(AllChannelGainArrays)](TArrayView<const float*> Src, TArrayView<float*> Dst, const int32 NumFrames)
+				{
+					const int32 NumSrcChannels = Src.Num();
+					for( int32 i = 0; i < Dst.Num(); ++i)
+					{
+						FMemory::Memzero(Dst[i], sizeof(float) * NumFrames);
+					}
+					for (int32 SourceChannelIndex = 0; SourceChannelIndex < NumSrcChannels; ++SourceChannelIndex)
+					{
+						EncodeMonoAmbisonicMixIn(MakeArrayView(Src[SourceChannelIndex], NumFrames), Dst, Gains[SourceChannelIndex]);
+					}
+				};
+		}
+	
+		// Discrete to Composite
+		FTranscoder GetTranscoder(const FDiscrete& InFrom,	const FComposite&	InTo,	const FParams&)
+		{
+			return {}; // TBD.
+		}
+		
+		// Ambisonics to Discrete
+		FTranscoder GetTranscoder(const FSoundfield& InFromType, const FDiscrete& InToType, const FParams&)
+		{
+			TArray<TArray<float>> AllChannelGainArrays;
+			AllChannelGainArrays.SetNum(InToType.NumChannels());
+			const TArray<FDiscreteChannelTypeFamily::FSpeaker>& Order = InToType.GetSpeakerOrder();
+			for (int32 i = 0; i < InToType.NumChannels(); ++i)
+			{
+				const FDiscreteChannelTypeFamily::FSpeaker& Speaker = Order[i];
+			
+				TArray<float>& GainArray = AllChannelGainArrays[i];
+				GainArray.SetNum(InFromType.NumChannels());
+
+				// Skip LFE.
+				if (Order[i].ID == LexToName(ESpeakerShortNames::LFE))
+				{
+					continue;
+				}
+			
+				check(GainArray.Num() == FSoundfieldChannelTypeFamily::OrderToNumChannels(InFromType.GetAmbisonicsOrder()));
+				const float AzimuthRads = FMath::DegreesToRadians(Speaker.AzimuthDegrees);
+				const float ElevationRads = FMath::DegreesToRadians(Speaker.ElevationDegrees);
+				FSphericalHarmonicCalculator::ComputeSoundfieldChannelGains(InFromType.GetAmbisonicsOrder(), AzimuthRads, ElevationRads, GainArray);
+				FSphericalHarmonicCalculator::NormalizeGains(GainArray);
+			}
+	
+			return [NumSpeakers = InToType.NumChannels(), Gains = MoveTemp(AllChannelGainArrays)](TArrayView<const float*> Src, TArrayView<float*> Dst, const int32 NumFrames)
+				{
+					for (int32 i = 0; i < NumSpeakers; ++i)
+					{
+						TArrayView<float> DstChannel = MakeArrayView(Dst[i],NumFrames);
+						FMemory::Memzero(DstChannel.GetData(), DstChannel.NumBytes()); // Decoder is MixIn.
+						DecodeMonoAmbisonicMixIn(Src, DstChannel, Gains[i]);
+					}
+				};
+		}
+
+		// Ambisonics to Ambisonics
+		FTranscoder GetTranscoder(const FSoundfield& InFromType,const FSoundfield& InToType, const FParams& InParams)
+		{
+			// Ambisonics to Ambisonics
+			if (&InToType == &InFromType)
+			{
+				// Same, so memcpy.
+				return [](TArrayView<const float*> InSrc, TArrayView<float*> InDst, const int32 NumFrames)
+					{
+						for (int32 i = 0; i < InSrc.Num(); ++i)
+						{
+							FMemory::Memcpy(InDst[i], InSrc[i], sizeof(float) * NumFrames);
+						}
+					};
+			}
+
+			const int32 NumChannelsToCopy = FMath::Min(InToType.NumChannels(), InFromType.NumChannels());
+			return [NumChannelsToCopy](TArrayView<const float*> InSrc, TArrayView<float*> InDst, const int32 NumFrames)
+				{
+					checkSlow(InSrc.Num() <= NumChannelsToCopy);
+					checkSlow(InDst.Num() <= NumChannelsToCopy);
+					// Copy as many channels as destination will hold.
+					for (int32 i = 0; i < NumChannelsToCopy; ++i)
+					{
+						FMemory::Memcpy(InDst[i], InSrc[i], sizeof(float) * NumFrames);
+					}
+
+					// Zero unwritten channels.
+					for (int32 i = NumChannelsToCopy; i < InDst.Num(); ++i)
+					{
+						FMemory::Memzero(InDst[i], sizeof(float) * NumFrames);					
+					}
+				};
+		}
+		
+		// Ambisonics to Composite
+		FTranscoder GetTranscoder(const FSoundfield&	InFrom,	const FComposite&	InTo,	const FParams&)
+		{
+			return {};	// TBD.
+		}
+				
+		static FTranscoder GetCompositeToX(const FComposite&	InFrom,	const FChannelTypeFamily&	InTo,	const FParams& InParams)
+		{
+			using FCoderInfo = TTuple<FComposite::FContainedType, FTranscoder>;
+			TArray<FCoderInfo> Coders;
+			Coders.Reserve(InFrom.NumTypes());
+			for (int32 i = 0; i < InFrom.NumTypes(); ++i)
+			{
+				const FComposite::FContainedType* Contained = InFrom.GetType(i);
+				const FChannelTypeFamily* ContainedFrom = GetChannelRegistry().FindChannel(Contained->Type);
+				if (ensure(ContainedFrom))
+				{
+					Coders.Emplace(
+						*Contained,
+						FTranscoderResolver::Resolve(*ContainedFrom, InParams)
+					);
+				}
+			}
+			
+			return [Coders](TArrayView<const float*> Src, TArrayView<float*> Dst, const int32 NumFrames)
+				{	
+					for (const FCoderInfo& i : Coders)
+					{
+						// Pull out stuff from tuple. 
+						const int32 CompositeChannelOffset = i.Get<0>().Offset;
+						const int32 CompositeNumChannels = i.Get<0>().Count;
+						const FTranscoder& Coder = i.Get<1>();
+						
+						// Slice SRC info just what this composite item should see
+						const TArrayView<const float*> SrcSlice = Src.Slice(CompositeChannelOffset, CompositeNumChannels); 
+						
+						// Execute.
+						Coder(SrcSlice, Dst, NumFrames);
+					}
+				};
+		}	
+		
+		// Composite to Soundfield
+		FTranscoder GetTranscoder(const FComposite&	InFrom,	const FSoundfield&	InTo,	const FParams& InParams)
+		{
+			return GetCompositeToX(InFrom, InTo, InParams);
+		}
+		// Composite to Composite
+		FTranscoder GetTranscoder(const FComposite&	InFrom,	const FComposite&	InTo,	const FParams& InParams)
+		{
+			return GetCompositeToX(InFrom, InTo, InParams);
+		} 
+		
+		// Composite to Discrete
+		FTranscoder GetTranscoder(const FComposite&	InFrom,	const FDiscrete&	InTo,	const FParams& InParams)
+		{
+			return GetCompositeToX(InFrom, InTo, InParams);			
+		}
+	
+		static void TranslateGainMatrix(
+			const TArrayView<float>& InGainMatrix,
+			const int32 InNumInputChannels,
+			const int32 InNumOutputChannels,
+			const TArrayView<const ESpeakerShortNames>& InOldOrder,
+			const TArrayView<const ESpeakerShortNames>& InNewOrder,
+			TArray<float>& OutGains)
+		{
+			// Channels not found in new order will be zeroed.
+			OutGains.SetNumZeroed(InNumOutputChannels*InNumInputChannels);
+		
+			// Lookup (old->new).
+			int32 Lookup[static_cast<int32>(ESpeakerShortNames::NumChannels)] = { INDEX_NONE };
+			for (int32 i = 0; i < InOldOrder.Num(); ++i)
+			{
+				Lookup[static_cast<int32>(InOldOrder[i])] = InNewOrder.IndexOfByKey(InOldOrder[i]);
+			}
+		
+			// Translate.
+			for (int32 InputChannelIndex=0; InputChannelIndex < InNumInputChannels; ++InputChannelIndex)
+			{
+				// This input channel has a mapping in the new?
+				const ESpeakerShortNames OldIn = InOldOrder[InputChannelIndex];
+				if (const int32 NewIn = Lookup[static_cast<int32>(OldIn)]; NewIn != INDEX_NONE)
+				{
+					for (int32 OutputChannelIndex=0; OutputChannelIndex < InNumOutputChannels; ++OutputChannelIndex)
+					{
+						const ESpeakerShortNames OldOut = InOldOrder[OutputChannelIndex];
+						if (const int32 NewOut = Lookup[static_cast<int32>(OldOut)]; NewOut != INDEX_NONE)
+						{
+							OutGains[(NewIn				* InNumOutputChannels)		+ NewOut] =
+						InGainMatrix[(InputChannelIndex * InNumOutputChannels)		+ OutputChannelIndex];
+						}
+					}
+				}
+			}
+		}
+	
+		static const IDiscretePanner* ChoosePanner(const FDiscreteChannelTypeFamily& InFrom, const FDiscreteChannelTypeFamily& InTo)
+		{
+			// Return To panner.
+			// All using 2d vbap for now, so just return that.
+			return InTo.GetPanner();
+		}
+	}
+	
+	// Static speakers can generate a matrix of speaker gains and cache that to avoid running the panner everytime.
+	bool MakeGainMatrixFrom(const FDiscreteChannelTypeFamily& InSrc, const FDiscreteChannelTypeFamily& InDst, TArray<float>& OutMatrix)
+	{
+		if (const IDiscretePanner* Panner = ChannelAgnosticTranscoder::ChoosePanner(InSrc, InDst))
+		{
+			const int32 DstChannels = InDst.NumChannels();
+			const int32 SrcChannels = InSrc.NumChannels();
+			OutMatrix.SetNumZeroed(SrcChannels*DstChannels);
+
+			const TArray<FDiscreteChannelTypeFamily::FSpeaker>& Order = InSrc.GetSpeakerOrder();
+		
+			// For each src channel, create row vector of gains for each dst channel.
+			TArray<float> SpeakerGains;
+			SpeakerGains.SetNumZeroed(DstChannels);
+			for (int32 i = 0; i < SrcChannels; ++i)
+			{
+				IDiscretePanner::FOutputParams Results;
+				IDiscretePanner::FInputParams Params { .AzimuthDegrees = Order[i].AzimuthDegrees, .ElevationDegrees = Order[i].ElevationDegrees };
+				if (Panner->ComputeGains(Params,Results) )
+				{
+					// Copy results into gain array.
+					for (int32 j = 0; j < Results.Results.Num(); ++j)
+					{
+						if (const int32 DstChannelIndex = Results.Results[j].ChannelIndex; DstChannelIndex != INDEX_NONE)
+						{
+							OutMatrix[i*DstChannels+DstChannelIndex] = Results.Results[j].Gain;
+						}
+					}
+				}
+			}
+			return true;
+		}
+		return false;
+	}
+}
\ No newline at end of file
diff --git a/Engine/Source/Runtime/AudioChannelAgnosticCore/Private/ChannelAgnostic/ChannelAgnosticType.cpp b/Engine/Source/Runtime/AudioChannelAgnosticCore/Private/ChannelAgnostic/ChannelAgnosticType.cpp
new file mode 100644
--- /dev/null
+++ b/Engine/Source/Runtime/AudioChannelAgnosticCore/Private/ChannelAgnostic/ChannelAgnosticType.cpp
@@ -0,0 +1,55 @@
+﻿// Copyright Epic Games, Inc. All Rights Reserved.
+
+
+#include "ChannelAgnostic/ChannelAgnosticType.h"
+#include "SimpleAlloc/SimpleHeapAllocator.h"
+#include "TypeFamily/ChannelTypeFamily.h"
+
+namespace Audio
+{
+	FSimpleAllocBase& FChannelAgnosticType::GetDefaultAllocator()
+	{
+		static FSimpleHeapAllocator Allocator;
+		return Allocator;
+	}
+
+	FChannelAgnosticType::FChannelAgnosticType(const TRetainedRef<const FChannelTypeFamily> InType, const int32 InNumFrames, FSimpleAllocBase* InAllocator)
+		: FChannelAgnosticType(InType, InNumFrames, InType.Get().NumChannels(), InAllocator)
+	{
+	}
+	
+	FChannelAgnosticType::FChannelAgnosticType(const TRetainedRef<const FChannelTypeFamily> InType, const int32 InNumFrames, const int32 InNumChannels, FSimpleAllocBase* InAllocator)
+		: Buffer(InNumFrames * InNumChannels, InAllocator)
+		, Type(&InType.Get())
+		, NumFramesPrivate(InNumFrames)
+		, NumChannelsPrivate(InNumChannels)
+	{
+		check(Type);
+		check(NumChannelsPrivate > 0);
+		check(Type->NumChannels() == InNumChannels || (Type->NumChannels() == 0 && InNumChannels > 0));
+	}
+
+	bool FChannelAgnosticType::IsA(const FChannelAgnosticType& InOther) const
+	{
+		check(Type);
+		check(InOther.Type);
+		return Type->IsA(InOther.Type);
+	}
+
+	bool FChannelAgnosticType::IsA(const FName& InTypeName) const
+	{
+		check(Type);
+		return Type->IsA(InTypeName);
+	}
+	
+	FName FChannelAgnosticType::GetTypeName() const
+	{
+		check(Type);
+		return Type->GetName();
+	}
+
+	void FChannelAgnosticType::Zero()
+	{
+		FMemory::Memzero(Buffer.GetView().GetData(), Buffer.GetView().NumBytes());
+	}
+}
diff --git a/Engine/Source/Runtime/AudioChannelAgnosticCore/Private/ChannelAgnostic/ChannelAgnosticTypeUtils.cpp b/Engine/Source/Runtime/AudioChannelAgnosticCore/Private/ChannelAgnostic/ChannelAgnosticTypeUtils.cpp
new file mode 100644
--- /dev/null
+++ b/Engine/Source/Runtime/AudioChannelAgnosticCore/Private/ChannelAgnostic/ChannelAgnosticTypeUtils.cpp
@@ -0,0 +1,59 @@
+﻿// Copyright Epic Games, Inc. All Rights Reserved.
+
+
+#include "ChannelAgnostic/ChannelAgnosticTypeUtils.h"
+
+#include "Algo/Transform.h"
+#include "ChannelAgnostic/ChannelAgnosticType.h"
+#include "DSP/WavFormatUtils.h"
+
+namespace Audio
+{
+	// TODO: MOVE to DSP.
+	static void Interleave(const TArrayView<const float> InMultiMono, const int32 InNumChannels, const TArrayView<float> OutInterleaved)
+	{
+		check(InNumChannels > 0);
+		checkSlow(InMultiMono.Num() <= OutInterleaved.Num());
+
+		float* Dst = OutInterleaved.GetData();
+		const float* Src = InMultiMono.GetData();
+		const int32 NumFrames = InMultiMono.Num() / InNumChannels;
+		checkSlow(InMultiMono.Num() % InNumChannels == 0);
+					
+		for(int32 Frame = 0; Frame < NumFrames; ++Frame)
+		{
+			for (int32 Channel = 0; Channel < InNumChannels; ++Channel)
+			{
+				*Dst++ = Src[NumFrames * Channel + Frame];
+			}
+		}
+	}
+
+	TOptional<uint32> FChannelAgnosticUtils::FindChannelMaskFromAudioMixerNumChannels(const int32 InNumChannels)
+	{
+		return WavFormatUtils::NumChannelsToCommonChannelMask(InNumChannels);
+	}
+
+	bool FChannelAgnosticUtils::ChannelIdToShortMixerChannels(const TArray<FName>& InIds, TArray<ESpeakerShortNames>& OutShortNames)
+	{
+		bool bResult = true;
+		OutShortNames.Reset();
+		Algo::Transform(InIds, OutShortNames, [&bResult](const FName InName) -> ESpeakerShortNames 
+		{
+			if (TOptional<ESpeakerShortNames> Enum = NameToShortSpeakerName(InName); Enum.IsSet())
+			{
+				return *Enum;
+			}
+			bResult = false;
+			return ESpeakerShortNames::Unknown;
+		});
+		return bResult;	
+	}
+	
+
+	void FChannelAgnosticUtils::Interleave(const FChannelAgnosticType& In, const TArrayView<float> Out)
+	{
+		Audio::Interleave(In.Buffer.GetView(), In.NumChannels(), Out);
+	}
+
+}
\ No newline at end of file
diff --git a/Engine/Source/Runtime/AudioChannelAgnosticCore/Private/TypeFamily/ChannelTypeFamily.cpp b/Engine/Source/Runtime/AudioChannelAgnosticCore/Private/TypeFamily/ChannelTypeFamily.cpp
new file mode 100644
--- /dev/null
+++ b/Engine/Source/Runtime/AudioChannelAgnosticCore/Private/TypeFamily/ChannelTypeFamily.cpp
@@ -0,0 +1,385 @@
+﻿// Copyright Epic Games, Inc. All Rights Reserved.
+
+#include "TypeFamily/ChannelTypeFamily.h"
+
+#include "Algo/Copy.h"
+#include "HAL/FileManager.h"
+#include "ChannelAgnostic/ChannelAgnosticTypeUtils.h"
+#include "Containers/Map.h"
+#include "DSP/SphericalHarmonicCalculator.h"
+#include "DSP/Vbap.h"
+
+namespace Audio
+{
+	namespace ChannelTypeFamilyPrivate
+	{
+		static FString MakePrettyString(const TArray<FDiscreteChannelTypeFamily::FSpeaker>& InSpeakers)
+		{
+			return FString::JoinBy(InSpeakers,TEXT(", "), [](FDiscreteChannelTypeFamily::FSpeaker InSpeaker) -> FString
+			{
+				return FString::Printf(TEXT("[Name=%s,Az=%2.2f, El=%2.2f]"), 
+					*InSpeaker.ID.ToString(), InSpeaker.AzimuthDegrees, InSpeaker.ElevationDegrees);
+			});
+		}
+	}
+	
+	
+		
+	class FChannelRegistryImpl final : public IChannelTypeRegistry 
+	{
+		TMap<FName, TSharedRef<const FChannelTypeFamily>> Types;
+		TMap<uint32, TSharedRef<const FDiscreteChannelTypeFamily>> ChannelBitMaskToDiscreteLookupMap;		// Look up for discrete types from BitMask.
+	public:
+		UE_NONCOPYABLE(FChannelRegistryImpl);
+		FChannelRegistryImpl() = default;
+		virtual ~FChannelRegistryImpl() override = default;
+	protected:
+		virtual bool RegisterType(const FName& InUniqueName, TUniquePtr<FTypeFamily>&& InType) override
+		{
+			static const FName ChannelTypeName = TEXT("Cat");
+			const FTypeFamily* Type = FindTypeInternal(ChannelTypeName);
+			if (InUniqueName != ChannelTypeName)  // Protect the initial base type registration.
+			{
+				if (!Type)
+				{
+					// CAT not a registered type. Shouldn't happen.
+					checkNoEntry();
+					return false;
+				}
+			}
+			
+			if (Type && !InType->IsA(Type))
+			{
+				// Not a cat. (maybe a dog?)
+				return false;	
+			}
+			
+			// TODO: thread safety here.
+			if (Types.Find(InUniqueName) != nullptr)
+			{
+				return false;
+			}
+			
+			// Safe to cast it and add.
+			TSharedRef<FChannelTypeFamily> Ref(static_cast<FChannelTypeFamily*>(InType.Release()));
+			
+			// Add to list of registered types.
+			check(Types.Find(InUniqueName) == nullptr);
+			Types.Emplace(InUniqueName,MoveTemp(Ref));
+			
+			// Rebuild lookups.
+			BuildDiscreteBitmaskLookupMap();
+			
+			return true;
+		}	
+		virtual bool UnregisterType(const FName& InUniqueName) override
+		{
+			if (TSharedRef<const FChannelTypeFamily>* Found = Types.Find(InUniqueName))
+			{
+				const TSharedPtr<const FChannelTypeFamily> Removed(*Found);
+				Types.Remove(InUniqueName);
+				TSet<FName> Dependencies;
+				for (auto i : Types)
+				{
+					if (i.Value->GetParent() && i.Value->GetParent() == Removed.Get())
+					{
+						Dependencies.Add(i.Key);
+					}
+				}
+				
+				// Clear deps.
+				for (const FName i : Dependencies)
+				{
+					UnregisterType(i);
+				}
+				
+				// Rebuild lookups.
+				BuildDiscreteBitmaskLookupMap();
+			}
+			return false;
+		}
+		
+		virtual const FTypeFamily* FindTypeInternal(const FName InUniqueName) const override
+		{
+			if (const TSharedRef<const FChannelTypeFamily>* Found = Types.Find(InUniqueName))
+			{
+				return &(Found->Get());
+			}
+			return nullptr;
+		}
+
+		virtual TArray<TSharedRef<const FChannelTypeFamily>> GetAllChannelFormats() const override
+		{
+			TArray<TSharedRef<const FChannelTypeFamily>> AllFormats;
+			Types.GenerateValueArray(AllFormats);
+			return AllFormats;
+		};
+		
+		virtual TArray<TSharedRef<const FChannelTypeFamily>> FindAllFamily(const FName InFamilyName) const override
+		{
+			TArray<TSharedRef<const FChannelTypeFamily>> Matches;
+			const TArray<TSharedRef<const FChannelTypeFamily>> AllFormats = GetAllChannelFormats();
+			Algo::CopyIf(AllFormats, Matches, [&InFamilyName](const TSharedRef<const FChannelTypeFamily>& i) -> bool
+				{
+					return i->GetFamilyName() == InFamilyName;
+				});
+			return Matches;
+		}
+		
+		void BuildDiscreteBitmaskLookupMap()
+		{
+			ChannelBitMaskToDiscreteLookupMap.Reset();
+			const TArray<TSharedRef<const FChannelTypeFamily>> AllDiscrete = FindAllFamily(FDiscreteChannelTypeFamily::GetFamilyTypeName());
+			for (const TSharedRef<const FChannelTypeFamily>& i : AllDiscrete)
+			{
+				TSharedRef<const FDiscreteChannelTypeFamily> Discrete = StaticCastSharedRef<const FDiscreteChannelTypeFamily>(i);
+				if (Discrete->IsAbstract())
+				{
+					continue;
+				}
+
+				// Assuming we can cleanly Convert FNames to Enum 
+				if (TArray<ESpeakerShortNames> Channels; FChannelAgnosticUtils::ChannelIdToShortMixerChannels(Discrete->GetSpeakerIds(), Channels))
+				{
+					if (const uint32 BitMask = ShortNamesToChannelMask(Channels); BitMask != 0)
+					{
+						ensure(!ChannelBitMaskToDiscreteLookupMap.Contains(BitMask));
+						ChannelBitMaskToDiscreteLookupMap.Emplace(BitMask, Discrete);
+					}
+				}
+			}
+		}
+		virtual const TSharedPtr<const FDiscreteChannelTypeFamily> FindDiscreteChannelFromBitMask(const uint32 InChannelMask) const override
+		{
+			// Look up bitmask.
+			if (const TSharedRef<const FDiscreteChannelTypeFamily>* Found = ChannelBitMaskToDiscreteLookupMap.Find(InChannelMask))
+			{
+				return Found->ToSharedPtr();
+			}
+	
+			// Not found.
+			return {};				
+		}
+	};
+
+	int32 FSoundfieldChannelTypeFamily::OrderToNumChannels(const int32 InOrder)
+	{
+		return FSphericalHarmonicCalculator::OrderToNumChannels(InOrder);
+	}
+
+	TOptional<FChannelTypeFamily::FChannelName> FSoundfieldChannelTypeFamily::GetChannelName(const int32 InChannelIndex) const
+	{
+		if (InChannelIndex < 0)
+		{
+			return {};
+		}
+
+		int32 OrderNum=0;
+		int32 DegreeNum=0;
+		FSphericalHarmonicCalculator::IndexToOrderAndDegree(InChannelIndex, OrderNum, DegreeNum);
+		const FName ChannelName = *FString::Printf(TEXT("ACN %d"), InChannelIndex); // *ChannelName.ToString(),	 (<-- Indexed' FName appears to fail for Node config UX, so use Sprintf'd one).ChannelName(TEXT("ACN"),InChannelIndex; 
+		return FChannelName
+		{
+			.Name = ChannelName,  
+			.FriendlyName = FString::Printf(TEXT("ACN=%d, Order=%d, Degree=%d"), InChannelIndex, Order, DegreeNum)
+		};
+	}
+	
+
+	IChannelTypeRegistry& GetChannelRegistry()
+	{
+		static FChannelRegistryImpl Registry;
+		return Registry;
+	}
+
+   FChannelTypeFamily::FChannelTypeFamily(
+	   const FName& InUniqueName,
+	   const FName& InFamilyTypeName,
+	   const int32 InNumChannels,
+	   FChannelTypeFamily* InParentType,
+	   const FString& InFriendlyName,
+	   const bool InbIsParentsDefault,
+	   const bool InbIsAbstract)
+	   : Super(InUniqueName, InParentType, InFriendlyName)
+	   , bIsAbstract(InbIsAbstract)
+	   , bIsParentsDefault(InbIsParentsDefault)
+	   , NumChannelsPrivate(InNumChannels)
+	   , FamilyType(InFamilyTypeName)
+	{
+		check(InNumChannels >= 0);
+		check(!InUniqueName.IsNone());
+		if (bIsParentsDefault)
+		{
+			check(InParentType);
+			check(InParentType->DefaultChild == nullptr); 
+			InParentType->DefaultChild = this;
+		}
+	}
+
+   FDiscreteChannelTypeFamily::FDiscreteChannelTypeFamily(const FName& InUniqueName, FChannelTypeFamily* InParentType, const FString& InFriendlyName, const TArray<FSpeaker>& InOrder, const bool bIsParentsDefault, const bool bIsAbstract)
+		: Super(InUniqueName, GetFamilyTypeName(), InOrder.Num(), InParentType, InFriendlyName, bIsParentsDefault, bIsAbstract)
+		, Order(InOrder)
+		, Panner(MakePanner(InOrder))
+	{
+		checkf(InParentType, TEXT("Type=%s, Has a Null Parent"), *GetName().ToString());
+	}
+
+    FDiscreteChannelTypeFamily::~FDiscreteChannelTypeFamily()
+	{
+		// Needed for TUnique deleter, as panner is forward declared in header.
+	}
+
+   TOptional<FChannelTypeFamily::FChannelName> FDiscreteChannelTypeFamily::GetChannelName(const int32 InChannelIndex) const
+	{
+		check(InChannelIndex >= 0 && InChannelIndex < NumChannels());
+		if (Order.IsValidIndex(InChannelIndex))
+		{
+			const FName SpeakerId = Order[InChannelIndex].ID;
+			return
+			{ // Optional.
+					{ 
+						.Name = SpeakerId,
+						.FriendlyName = SpeakerId.ToString() 
+					}
+			};
+		}
+		return {};
+	}
+
+	TUniquePtr<IDiscretePanner> FDiscreteChannelTypeFamily::MakePanner(const TArray<FSpeaker>& InOrder)
+	{
+		// Just using 2D VBap for now.
+		TArray<IDiscretePanner::FSpeaker> Speakers;
+		for (int32 Channel = 0; Channel < InOrder.Num(); Channel++)
+		{
+			const FSpeaker& Speaker = InOrder[Channel];
+
+			// Add all speakers that are considered 'spatialized'
+			// This excludes channels that not, like LFE and possibly (Centre channel for dialog). 
+			if (Speaker.bIsSpatialized)
+			{
+				IDiscretePanner::FSpeaker PannerSpeaker
+				{
+					.AzimuthDegrees = Speaker.AzimuthDegrees,
+					.ElevationDegrees = Speaker.ElevationDegrees,
+					.ChannelID = Speaker.ID,
+					.ChannelIndex = Channel
+				};
+				Speakers.Add(PannerSpeaker);
+			}
+		}
+		
+		if (Speakers.Num() > 0)
+		{
+			return MakeUnique<FVBap2D>(Speakers);
+		}
+		return {};
+	}
+	
+	
+	void RegisterCatBaseClassesWithChannelRegistry(IChannelTypeRegistry& Registry)
+	{
+		// Register root type.
+		Registry.RegisterType(TEXT("Cat"), MakeUnique<FChannelTypeFamily>(TEXT("Cat"), TEXT("Cat"),  0, nullptr, TEXT("Base Cat"), false, true));
+		FChannelTypeFamily* BaseCat = Registry.FindChannel(TEXT("Cat"));
+		check(BaseCat);
+
+		Registry.RegisterType(
+			 TEXT("Cat:Discrete"),
+			MakeUnique<FDiscreteChannelTypeFamily>(
+						TEXT("Cat:Discrete"),
+						BaseCat,
+						TEXT("Discrete Channel Format"),
+						TArray<FDiscreteChannelTypeFamily::FSpeaker>(),
+						true,
+						true)
+			);
+		Registry.RegisterType(
+			TEXT("Cat:Soundfield"),
+			MakeUnique<FSoundfieldChannelTypeFamily>(
+				TEXT("Cat:Soundfield"),
+				0, BaseCat, TEXT("Soundfield Format"),
+				false,
+				true));
+	}
+	
+	void UnregisterCatBaseClassesWithChannelRegistry(IChannelTypeRegistry& Registry)
+	{
+		// Unregister C++ types and defaults. 
+		Registry.UnregisterType(TEXT("Cat"));
+		Registry.UnregisterType(TEXT("Cat:Discrete"));
+		Registry.UnregisterType(TEXT("Cat:Soundfield"));
+	}
+	
+}
+
+
+const TCHAR* LexToString(const ESpeakerShortNames InSpeaker)
+{
+	// No code gen support here, so role this manually.
+#define CASE_AND_STRING(X) case ESpeakerShortNames::X: return TEXT(#X)
+
+	using enum ESpeakerShortNames;
+	switch (InSpeaker)
+	{
+	CASE_AND_STRING(FL);   // Front Left
+	CASE_AND_STRING(FR);   // Front Right
+	CASE_AND_STRING(FC);   // Front Center
+	CASE_AND_STRING(LFE);  // Low-Frequency Effects (Subwoofer)
+	CASE_AND_STRING(FLC);  // Front Left Center
+	CASE_AND_STRING(FRC);  // Front Right Center
+	CASE_AND_STRING(SL);   // Side Left
+	CASE_AND_STRING(SR);   // Side Right
+	CASE_AND_STRING(BL);   // Back Left
+	CASE_AND_STRING(BR);   // Back Right
+	CASE_AND_STRING(BC);   // Back Center
+	CASE_AND_STRING(TFL);  // Top Front Left
+	CASE_AND_STRING(TFR);  // Top Front Right
+	CASE_AND_STRING(TBL);  // Top Back Left
+	CASE_AND_STRING(TBR);  // Top Back Right
+	default:
+		break;
+	}
+	return TEXT("Unknown");
+}
+FName LexToName(const ESpeakerShortNames InSpeaker)
+{
+#define CASE_TO_CONST_FNAME(X)\
+	case X:\
+	{\
+		static const FName Name = LexToString(InSpeaker);\
+		return Name;\
+	}
+	switch (InSpeaker)
+	{
+		FOREACH_ENUM_ESPEAKERSHORTNAMES(CASE_TO_CONST_FNAME)
+	default:
+		checkNoEntry();
+	}
+	return {};
+#undef CASE_TO_CONST_FNAME
+}
+
+TOptional<ESpeakerShortNames> NameToShortSpeakerName(const FName InName)
+{
+	#define IF_FNAME_RETURN(X)\
+	{\
+		static const FName Name = LexToString(X);\
+		if (Name == InName) return X;\
+	}
+	FOREACH_ENUM_ESPEAKERSHORTNAMES(IF_FNAME_RETURN)
+	#undef IF_FNAME_RETURN
+
+	// failed.
+	return {};
+}
+
+uint32 ShortNamesToChannelMask(const TArray<ESpeakerShortNames>& InChannels)
+{
+	uint32 ChannelMask = 0;
+	for (const ESpeakerShortNames i : InChannels)
+	{
+		ChannelMask |= 1 << static_cast<uint32>(i);
+	}
+	return ChannelMask;
+}
diff --git a/Engine/Source/Runtime/AudioChannelAgnosticCore/Public/ChannelAgnostic/ChannelAgnosticTranscoding.h b/Engine/Source/Runtime/AudioChannelAgnosticCore/Public/ChannelAgnostic/ChannelAgnosticTranscoding.h
new file mode 100644
--- /dev/null
+++ b/Engine/Source/Runtime/AudioChannelAgnosticCore/Public/ChannelAgnostic/ChannelAgnosticTranscoding.h
@@ -0,0 +1,108 @@
+﻿// Copyright Epic Games, Inc. All Rights Reserved.
+
+#pragma once
+
+#include "TypeFamily/ChannelTypeFamily.h"
+
+namespace Audio
+{   
+	enum class EChannelTranscodeMethod : uint8
+	{
+		ChannelDrop,
+		MixUpOrDown,
+	};
+	
+	// All transcoder (from -> to) combinations.
+	namespace ChannelAgnosticTranscoder 
+	{
+		struct FGetTranscoderParams
+		{
+			const FChannelTypeFamily& ToType;
+			EChannelTranscodeMethod TranscodeMethod = EChannelTranscodeMethod::ChannelDrop;
+			EChannelMapMonoUpmixMethod MixMethod = EChannelMapMonoUpmixMethod::EqualPower;
+		};
+				
+		using FTranscoder = TFunction<void(TArrayView<const float*>, TArrayView<float*>, const int32 NumFrames)>;
+		
+		using FDiscrete = FDiscreteChannelTypeFamily;
+		using FSoundfield = FSoundfieldChannelTypeFamily;
+		using FComposite = FCompositeChannelTypeFamily;
+		using FParams = FGetTranscoderParams;
+		
+		// Discrete-to-X
+		AUDIOCHANNELAGNOSTICCORE_API FTranscoder GetTranscoder(const FDiscrete&		InFrom,	const FDiscrete&	InTo,	const FParams&); // d->d
+		AUDIOCHANNELAGNOSTICCORE_API FTranscoder GetTranscoder(const FDiscrete&		InFrom,	const FSoundfield&	InTo,	const FParams&); // d->s
+		AUDIOCHANNELAGNOSTICCORE_API FTranscoder GetTranscoder(const FDiscrete&		InFrom,	const FComposite&	InTo,	const FParams&); // d->c
+		
+		// Soundfield-to-X
+		AUDIOCHANNELAGNOSTICCORE_API FTranscoder GetTranscoder(const FSoundfield&	InFrom,	const FSoundfield&	InTo,	const FParams&); // s->s
+		AUDIOCHANNELAGNOSTICCORE_API FTranscoder GetTranscoder(const FSoundfield&	InFrom,	const FDiscrete&	InTo,	const FParams&); // s->d
+		AUDIOCHANNELAGNOSTICCORE_API FTranscoder GetTranscoder(const FSoundfield&	InFrom,	const FComposite&	InTo,	const FParams&); // s->c
+				
+		// Composite-to-X
+		AUDIOCHANNELAGNOSTICCORE_API FTranscoder GetTranscoder(const FComposite&	InFrom,	const FComposite&	InTo,	const FParams&); // c->c
+		AUDIOCHANNELAGNOSTICCORE_API FTranscoder GetTranscoder(const FComposite&	InFrom,	const FDiscrete&	InTo,	const FParams&); // c->d
+		AUDIOCHANNELAGNOSTICCORE_API FTranscoder GetTranscoder(const FComposite&	InFrom,	const FSoundfield&	InTo,	const FParams&); // c->s
+	}
+		
+	template<typename TFromType>
+	UE_EXPERIMENTAL(5.8,"AudioChannelAgnosticCore is experimental and in flux")
+	struct TFromVisitor final : IChannelTypeVisitor
+	{
+		const TFromType& FromType;
+		const ChannelAgnosticTranscoder::FGetTranscoderParams& Params;
+		ChannelAgnosticTranscoder::FTranscoder& Transcoder;
+			
+		TFromVisitor(const TFromType& InFromType, const ChannelAgnosticTranscoder::FGetTranscoderParams& InParams, ChannelAgnosticTranscoder::FTranscoder& InTranscoder)
+			: FromType(InFromType), Params(InParams), Transcoder(InTranscoder)
+		{}
+			
+		virtual void Visit(const FDiscreteChannelTypeFamily& InTo) override
+		{
+			Transcoder = ChannelAgnosticTranscoder::GetTranscoder(FromType, InTo, Params);	
+		}
+		virtual void Visit(const FSoundfieldChannelTypeFamily& InTo) override
+		{
+			Transcoder = ChannelAgnosticTranscoder::GetTranscoder(FromType, InTo, Params);
+		}
+		virtual void Visit(const FCompositeChannelTypeFamily& InTo) override
+		{
+			Transcoder = ChannelAgnosticTranscoder::GetTranscoder(FromType, InTo, Params);
+		}
+	};
+		
+	struct FTranscoderResolver final : public IChannelTypeVisitor
+	{
+		PRAGMA_DISABLE_EXPERIMENTAL_WARNINGS
+		
+		explicit FTranscoderResolver(const ChannelAgnosticTranscoder::FGetTranscoderParams& InParams) : Params(InParams) {}
+		virtual void Visit(const FDiscreteChannelTypeFamily& InFrom) override
+		{
+			TFromVisitor Visitor(InFrom, Params, Result);
+			Params.ToType.Accept(Visitor);
+		}
+		virtual void Visit(const FSoundfieldChannelTypeFamily& InFrom) override
+		{
+			TFromVisitor Visitor(InFrom, Params, Result);
+			Params.ToType.Accept(Visitor);
+		}
+		virtual void Visit(const FCompositeChannelTypeFamily& InFrom) override
+		{
+			TFromVisitor Visitor(InFrom, Params, Result);
+			Params.ToType.Accept(Visitor);	
+		}
+	
+		UE_EXPERIMENTAL(5.8,"AudioChannelAgnosticCore is experimental and in flux")
+		static ChannelAgnosticTranscoder::FTranscoder Resolve(const FChannelTypeFamily& FromType, const ChannelAgnosticTranscoder::FGetTranscoderParams& Params)
+		{
+			FTranscoderResolver Resolver(Params);
+			FromType.Accept(Resolver);
+			return MoveTemp(Resolver.Result);	
+		}
+		
+		ChannelAgnosticTranscoder::FGetTranscoderParams Params;
+		ChannelAgnosticTranscoder::FTranscoder Result;
+		
+		PRAGMA_ENABLE_EXPERIMENTAL_WARNINGS		
+	};
+}
\ No newline at end of file
diff --git a/Engine/Source/Runtime/AudioChannelAgnosticCore/Public/ChannelAgnostic/ChannelAgnosticType.h b/Engine/Source/Runtime/AudioChannelAgnosticCore/Public/ChannelAgnostic/ChannelAgnosticType.h
new file mode 100644
--- /dev/null
+++ b/Engine/Source/Runtime/AudioChannelAgnosticCore/Public/ChannelAgnostic/ChannelAgnosticType.h
@@ -0,0 +1,113 @@
+﻿// Copyright Epic Games, Inc. All Rights Reserved.
+
+#pragma once
+
+#include "Containers/ArrayView.h"
+#include "HAL/Platform.h"
+#include "ScratchBuffer.h"
+#include "Templates/RetainedRef.h"
+#include "TypeFamily/ChannelTypeFamily.h"
+
+namespace Audio
+{
+	class FChannelTypeFamily;
+	class FSimpleAllocBase;
+
+	class FChannelAgnosticType
+	{
+	public:
+		UE_EXPERIMENTAL(5.8,"AudioChannelAgnosticCore is experimental and in flux")
+		AUDIOCHANNELAGNOSTICCORE_API static FSimpleAllocBase& GetDefaultAllocator();
+
+		UE_EXPERIMENTAL(5.8,"AudioChannelAgnosticCore is experimental and in flux")
+		AUDIOCHANNELAGNOSTICCORE_API explicit FChannelAgnosticType(const TRetainedRef<const FChannelTypeFamily> InType, const int32 InNumFrames, FSimpleAllocBase* InAllocator = &GetDefaultAllocator());
+
+		/**
+		 * 
+		 * @param InType 
+		 * @param InNumFrames 
+		 * @param InNumChannels 
+		 * @param InAllocator 
+		 */
+		UE_EXPERIMENTAL(5.8,"AudioChannelAgnosticCore is experimental and in flux")
+		AUDIOCHANNELAGNOSTICCORE_API explicit FChannelAgnosticType(const TRetainedRef<const FChannelTypeFamily> InType, const int32 InNumFrames, const int32 InNumChannels, FSimpleAllocBase* InAllocator = &GetDefaultAllocator());
+	
+		// To be a Variable it *must* be copyable.
+		FChannelAgnosticType(const FChannelAgnosticType& Other) = default;
+		FChannelAgnosticType(FChannelAgnosticType&& Other) = default;
+		FChannelAgnosticType& operator=(FChannelAgnosticType&& Other) = default;
+		FChannelAgnosticType& operator=(const FChannelAgnosticType& Other) = default;
+
+		UE_EXPERIMENTAL(5.8,"AudioChannelAgnosticCore is experimental and in flux")
+		[[nodiscard]]
+		int32 NumFrames() const { return NumFramesPrivate; }
+
+		UE_EXPERIMENTAL(5.8,"AudioChannelAgnosticCore is experimental and in flux")
+		[[nodiscard]]
+		int32 NumChannels() const { return NumChannelsPrivate; }
+		
+		UE_EXPERIMENTAL(5.8,"AudioChannelAgnosticCore is experimental and in flux")
+		[[nodiscard]]
+		TArrayView<const float> GetChannel(const int32 InChannelIndex) const
+		{
+			PRAGMA_DISABLE_EXPERIMENTAL_WARNINGS					
+			check(InChannelIndex >= 0 && InChannelIndex < NumChannelsPrivate);
+			const TArrayView<const float> AllChannels = Buffer.GetView();
+			return MakeArrayView(AllChannels.GetData() + (NumFramesPrivate * InChannelIndex), NumFramesPrivate);
+			PRAGMA_ENABLE_EXPERIMENTAL_WARNINGS					
+		}
+		
+		UE_EXPERIMENTAL(5.8,"AudioChannelAgnosticCore is experimental and in flux")
+		[[nodiscard]]
+		TArrayView<float> GetChannel(const int32 InChannelIndex)
+		{
+			PRAGMA_DISABLE_EXPERIMENTAL_WARNINGS			
+			const TArrayView<const float> ConstView = std::as_const(*this).GetChannel(InChannelIndex);
+			return MakeArrayView<float>(const_cast<float*>(ConstView.GetData()), ConstView.Num());
+			PRAGMA_ENABLE_EXPERIMENTAL_WARNINGS			
+		}
+
+		UE_EXPERIMENTAL(5.8,"AudioChannelAgnosticCore is experimental and in flux")
+		AUDIOCHANNELAGNOSTICCORE_API bool IsA(const FChannelAgnosticType& InOther) const;
+		UE_EXPERIMENTAL(5.8,"AudioChannelAgnosticCore is experimental and in flux")
+		AUDIOCHANNELAGNOSTICCORE_API bool IsA(const FName& InTypeName) const;
+		UE_EXPERIMENTAL(5.8,"AudioChannelAgnosticCore is experimental and in flux")
+		AUDIOCHANNELAGNOSTICCORE_API FName GetTypeName() const;
+		UE_EXPERIMENTAL(5.8,"AudioChannelAgnosticCore is experimental and in flux")
+		const FChannelTypeFamily& GetType() const { return *Type; }
+
+		/*
+		 * Zero the contents of the buffer.
+		 */
+		UE_EXPERIMENTAL(5.8,"AudioChannelAgnosticCore is experimental and in flux")
+		AUDIOCHANNELAGNOSTICCORE_API void Zero();
+
+		/**
+		 * For fast raw DSP access to the buffer.
+		 * @return View of entire multi-mono buffer.
+		 */
+		UE_EXPERIMENTAL(5.8,"AudioChannelAgnosticCore is experimental and in flux")
+		[[nodiscard]]
+		TArrayView<float> GetRawMultiMono()
+		{
+			return Buffer.GetView();
+		}
+		/**
+		 * For fast raw DSP access to the buffer.  
+		 * @return View of entire multi-mono buffer.
+		 */
+		UE_EXPERIMENTAL(5.8,"AudioChannelAgnosticCore is experimental and in flux")
+		[[nodiscard]]
+		TArrayView<const float> GetRawMultiMono() const
+		{
+			return Buffer.GetView();	
+		}
+		
+	private:
+		friend class FChannelAgnosticUtils;
+		TScratchBuffer<float> Buffer;
+		const FChannelTypeFamily* Type = nullptr;	// Keep this as pointer internally for easy of copying.
+		int32 NumFramesPrivate = 0;
+		int32 NumChannelsPrivate = 0;
+	};
+}
diff --git a/Engine/Source/Runtime/AudioChannelAgnosticCore/Public/ChannelAgnostic/ChannelAgnosticTypeUtils.h b/Engine/Source/Runtime/AudioChannelAgnosticCore/Public/ChannelAgnostic/ChannelAgnosticTypeUtils.h
new file mode 100644
--- /dev/null
+++ b/Engine/Source/Runtime/AudioChannelAgnosticCore/Public/ChannelAgnostic/ChannelAgnosticTypeUtils.h
@@ -0,0 +1,30 @@
+﻿// Copyright Epic Games, Inc. All Rights Reserved.
+
+#pragma once
+
+#include "Templates/SharedPointer.h"
+#include "TypeFamily/ChannelTypeFamily.h"
+
+namespace Audio
+{
+	class FDiscreteChannelTypeFamily;
+	class IChannelTypeRegistry;
+	class FChannelAgnosticType;
+
+	class FChannelAgnosticUtils
+	{
+	public:
+		/**
+		 * Attempts to match a ChannelCount to a ChannelMask, for compatibility with existing AudioMixer layouts.
+		 * @return ChannelBitMask (Matching AudioMixerChannels enum shifted into a bitmask)
+		 */
+		UE_EXPERIMENTAL(5.8,"AudioChannelAgnosticCore is experimental and in flux")		
+		AUDIOCHANNELAGNOSTICCORE_API static TOptional<uint32> FindChannelMaskFromAudioMixerNumChannels(const int32 InNumChannels);
+		
+		UE_EXPERIMENTAL(5.8,"AudioChannelAgnosticCore is experimental and in flux")
+		AUDIOCHANNELAGNOSTICCORE_API static bool ChannelIdToShortMixerChannels(const TArray<FName>& InIds, TArray<ESpeakerShortNames>& OutShortNames);
+				
+		UE_EXPERIMENTAL(5.8,"AudioChannelAgnosticCore is experimental and in flux")		
+		static AUDIOCHANNELAGNOSTICCORE_API void Interleave(const FChannelAgnosticType& In, TArrayView<float> Out);
+	};
+}
diff --git a/Engine/Source/Runtime/AudioChannelAgnosticCore/Public/ScratchBuffer.h b/Engine/Source/Runtime/AudioChannelAgnosticCore/Public/ScratchBuffer.h
new file mode 100644
--- /dev/null
+++ b/Engine/Source/Runtime/AudioChannelAgnosticCore/Public/ScratchBuffer.h
@@ -0,0 +1,138 @@
+// Copyright Epic Games, Inc. All Rights Reserved.
+
+#pragma once
+
+#include "SimpleAlloc/SimpleAllocBase.h"
+
+namespace Audio
+{
+	template <typename T>
+	UE_EXPERIMENTAL(5.8,"AudioChannelAgnosticCore is experimental and in flux")
+	class TScratchBuffer
+	{
+	public:
+		PRAGMA_DISABLE_EXPERIMENTAL_WARNINGS
+		
+		using ValueType = T;
+
+		explicit TScratchBuffer(const int32 InNum, FSimpleAllocBase* InAllocator)
+			: Allocator(InAllocator)
+		{
+			DoAllocation(InNum);
+		}
+
+		TScratchBuffer(const TScratchBuffer& Other)
+			: Allocator(Other.Allocator)
+		{
+			*this = Other; // Call copy operator.
+		}
+
+		TScratchBuffer(TScratchBuffer&& Other)
+			: Allocator(Other.Allocator)
+		{
+			*this = MoveTemp(Other); // Call move operator.
+		}
+
+		TScratchBuffer& operator=(const TScratchBuffer& Other)
+		{
+			if (this != &Other)
+			{
+				// Free our state.
+				FreeAllocation();
+				
+				// Make new allocation.
+				Allocator = Other.Allocator;
+				DoAllocation(Other.Allocation.Num());
+
+				// Copy other into new allocated memory 
+				if (Allocation.GetData() && Allocation.Num() > 0)
+				{
+					FMemory::Memcpy(Allocation.GetData(), Other.Allocation.GetData(), Allocation.GetTypeSize() * Allocation.Num());
+				}
+			}
+			return *this;
+		}
+
+
+		// Move operator.
+		TScratchBuffer& operator=(TScratchBuffer&& Other)
+		{
+			if (this == &Other)
+			{
+				return *this;
+			}
+			// Free our state.
+			FreeAllocation();
+			Allocator = Other.Allocator;
+			Allocation = MoveTemp(Other.Allocation);
+			LifetimeToken = Other.LifetimeToken;
+
+			// Reset just in case.
+			Other.Reset();
+			return *this;
+		}
+
+		~TScratchBuffer()
+		{
+			FreeAllocation();
+		}
+
+		bool IsValid() const
+		{
+			return Allocator->GetCurrentLifetime() == LifetimeToken;
+		}
+
+		int32 Num() const
+		{
+			check(IsValid());
+			return Allocation.Num();
+		}
+
+		TArrayView<ValueType> GetView()
+		{
+			check(IsValid());
+			return MakeArrayView(Allocation.GetData(), Num());
+		}
+
+		TArrayView<const ValueType> GetView() const
+		{
+			check(IsValid());
+			return MakeArrayView(Allocation.GetData(), Num());
+		}
+
+	private:
+		void Reset()
+		{
+			Allocation = {};
+			LifetimeToken = -1;
+		}
+
+		void DoAllocation(const int32 InNum)
+		{
+			if (InNum > 0)
+			{
+				check(Allocation.GetData() == nullptr);
+				Allocation = MakeArrayView<ValueType>(
+					static_cast<ValueType*>(Allocator->Malloc(InNum * sizeof(ValueType))),
+					InNum);
+				check(Allocation.GetData() != nullptr);
+				LifetimeToken = Allocator->GetCurrentLifetime();
+				FMemory::Memzero(Allocation.GetData(), Allocation.GetTypeSize() * Allocation.Num());
+			}
+		}
+		void FreeAllocation()
+		{
+			if (Allocation.GetData() != nullptr)
+			{
+				Allocator->Free(Allocation.GetData());
+				Allocation = {};
+			}
+		}
+
+		FSimpleAllocBase* Allocator = nullptr;
+		TArrayView<ValueType> Allocation;
+		int32 LifetimeToken = -1;
+		
+		PRAGMA_ENABLE_EXPERIMENTAL_WARNINGS
+	};
+}
\ No newline at end of file
diff --git a/Engine/Source/Runtime/AudioChannelAgnosticCore/Public/SimpleAlloc/SimpleAllocBase.h b/Engine/Source/Runtime/AudioChannelAgnosticCore/Public/SimpleAlloc/SimpleAllocBase.h
new file mode 100644
--- /dev/null
+++ b/Engine/Source/Runtime/AudioChannelAgnosticCore/Public/SimpleAlloc/SimpleAllocBase.h
@@ -0,0 +1,39 @@
+// Copyright Epic Games, Inc. All Rights Reserved.
+
+#pragma once
+
+#include "HAL/UnrealMemory.h"
+#include "Misc/CoreMiscDefines.h"
+#include "Math/UnrealMathUtility.h"
+
+namespace Audio
+{
+	class FSimpleAllocBase
+	{
+	public:
+		UE_NONCOPYABLE(FSimpleAllocBase);
+		UE_EXPERIMENTAL(5.8,"AudioChannelAgnosticCore is experimental and in flux")
+		FSimpleAllocBase() = default;
+		UE_EXPERIMENTAL(5.8,"AudioChannelAgnosticCore is experimental and in flux")
+		virtual ~FSimpleAllocBase() = default;
+	
+		UE_EXPERIMENTAL(5.8,"AudioChannelAgnosticCore is experimental and in flux")
+		[[nodiscard]] virtual void* Malloc(const SIZE_T, const uint32 InAlignment = DEFAULT_ALIGNMENT) = 0;
+		UE_EXPERIMENTAL(5.8,"AudioChannelAgnosticCore is experimental and in flux")
+		virtual void Free(void*) = 0;
+		UE_EXPERIMENTAL(5.8,"AudioChannelAgnosticCore is experimental and in flux")
+		virtual uint32 GetCurrentLifetime() const { return 0; };
+		UE_EXPERIMENTAL(5.8,"AudioChannelAgnosticCore is experimental and in flux")
+		virtual void Reset() {};
+	};
+
+	namespace SimpleAllocBasePrivate
+	{
+		// This is what's defined in MemoryBase, so follow it here. 
+		FORCEINLINE static uint32 GetDefaultSizeToAlignment(const uint32 InSize)
+		{
+			check(InSize > 0);
+			return InSize >= 16 ? 16 : 8;
+		};
+	}
+}
\ No newline at end of file
diff --git a/Engine/Source/Runtime/AudioChannelAgnosticCore/Public/SimpleAlloc/SimpleHeapAllocator.h b/Engine/Source/Runtime/AudioChannelAgnosticCore/Public/SimpleAlloc/SimpleHeapAllocator.h
new file mode 100644
--- /dev/null
+++ b/Engine/Source/Runtime/AudioChannelAgnosticCore/Public/SimpleAlloc/SimpleHeapAllocator.h
@@ -0,0 +1,27 @@
+﻿// Copyright Epic Games, Inc. All Rights Reserved.
+
+#pragma once
+
+#include "SimpleAllocBase.h"
+
+namespace Audio
+{
+	class FSimpleHeapAllocator final : public FSimpleAllocBase
+	{
+	public:
+		UE_NONCOPYABLE(FSimpleHeapAllocator)
+		
+		UE_EXPERIMENTAL(5.8,"AudioChannelAgnosticCore is experimental and in flux")
+		FSimpleHeapAllocator() = default;
+		UE_EXPERIMENTAL(5.8,"AudioChannelAgnosticCore is experimental and in flux")
+		virtual void* Malloc(const SIZE_T InSizeBytes, const uint32 InAlignment=DEFAULT_ALIGNMENT) override
+		{
+			return FMemory::Malloc(InSizeBytes, InAlignment);
+		}
+		UE_EXPERIMENTAL(5.8,"AudioChannelAgnosticCore is experimental and in flux")
+		virtual void Free(void* InPtr) override
+		{
+			return FMemory::Free(InPtr);
+		}
+	};
+}
diff --git a/Engine/Source/Runtime/AudioChannelAgnosticCore/Public/SimpleAlloc/SimpleLinearAllocator.h b/Engine/Source/Runtime/AudioChannelAgnosticCore/Public/SimpleAlloc/SimpleLinearAllocator.h
new file mode 100644
--- /dev/null
+++ b/Engine/Source/Runtime/AudioChannelAgnosticCore/Public/SimpleAlloc/SimpleLinearAllocator.h
@@ -0,0 +1,74 @@
+﻿// Copyright Epic Games, Inc. All Rights Reserved.
+
+#pragma once
+
+#include "SimpleAllocBase.h"
+#include "Containers/ArrayView.h"
+
+namespace Audio
+{
+	class FSimpleLinearAllocator : public FSimpleAllocBase
+	{
+	public:
+		UE_NONCOPYABLE(FSimpleLinearAllocator)	
+		
+		UE_EXPERIMENTAL(5.8,"AudioChannelAgnosticCore is experimental and in flux")
+		explicit FSimpleLinearAllocator(const TArrayView<uint8> InPage)
+			: Page(InPage)
+		{
+			check(InPage.Num() > 0);
+			check(IsAligned(InPage.GetData(), SimpleAllocBasePrivate::GetDefaultSizeToAlignment(InPage.Num())));
+		}
+	
+		UE_EXPERIMENTAL(5.8,"AudioChannelAgnosticCore is experimental and in flux")
+		virtual void* Malloc(const SIZE_T InSizeBytes, const uint32 InAlignment = DEFAULT_ALIGNMENT) override
+		{
+			check(InSizeBytes > 0);
+			const uint32 Alignment = InAlignment == DEFAULT_ALIGNMENT ? SimpleAllocBasePrivate::GetDefaultSizeToAlignment(InSizeBytes) : InAlignment;
+			const uint32 AlignedNewTop = Align(Top, Alignment);
+		
+			if (AlignedNewTop + InSizeBytes <= Page.Num())
+			{
+				Top = AlignedNewTop;
+				void* Ret = Page.GetData() + Top;
+				Top += InSizeBytes;
+				return Ret;
+			}
+			return nullptr;
+		}
+
+		UE_EXPERIMENTAL(5.8,"AudioChannelAgnosticCore is experimental and in flux")
+		virtual void Free(void*) override {}
+
+		UE_EXPERIMENTAL(5.8,"AudioChannelAgnosticCore is experimental and in flux")
+		virtual uint32 GetCurrentLifetime() const override
+		{
+			return CurrentLifetime;
+		}
+		UE_EXPERIMENTAL(5.8,"AudioChannelAgnosticCore is experimental and in flux")
+		virtual void Reset() override
+		{
+			Top = 0;
+			CurrentLifetime++;
+		}
+	
+	protected:	
+		uint32 Top = 0;					// Top of all allocations, climbs upwards towards end of page.
+		uint32 CurrentLifetime = 0;		// Each reset the lifetime is increased
+		TArrayView<uint8> Page;			// One page for now.
+	};
+
+	class FSimpleLinearAllocatorFromHeap final : public FSimpleLinearAllocator
+	{
+	public:
+		UE_EXPERIMENTAL(5.8,"AudioChannelAgnosticCore is experimental and in flux")
+		explicit FSimpleLinearAllocatorFromHeap(const SIZE_T InPageSize, const uint32 InPageAlignment = DEFAULT_ALIGNMENT)
+			: FSimpleLinearAllocator(MakeArrayView(static_cast<uint8*>(FMemory::Malloc(InPageSize,InPageAlignment)), IntCastChecked<int32>(InPageSize)))
+		{}
+		UE_EXPERIMENTAL(5.8,"AudioChannelAgnosticCore is experimental and in flux")
+		virtual ~FSimpleLinearAllocatorFromHeap() override
+		{
+			FMemory::Free(Page.GetData());
+		}
+	};
+}
diff --git a/Engine/Source/Runtime/AudioChannelAgnosticCore/Public/TypeFamily/ChannelTypeFamily.h b/Engine/Source/Runtime/AudioChannelAgnosticCore/Public/TypeFamily/ChannelTypeFamily.h
new file mode 100644
--- /dev/null
+++ b/Engine/Source/Runtime/AudioChannelAgnosticCore/Public/TypeFamily/ChannelTypeFamily.h
@@ -0,0 +1,507 @@
+// Copyright Epic Games, Inc. All Rights Reserved.
+
+#pragma once
+
+#include "AudioMixer.h"
+#include "TypeFamily.h"
+#include "Algo/Transform.h"
+#include "DSP/ChannelMap.h"
+#include "Templates/Function.h"
+#include "Misc/Optional.h"
+
+
+namespace Audio
+{
+	class FCompositeChannelTypeFamily;
+	struct IDiscretePanner;
+}
+
+/**
+ * A Comprehensive Short form Speaker enumeration
+ */
+enum class ESpeakerShortNames : uint32
+{
+	FL = EAudioMixerChannel::FrontLeft,				// Front Left
+	FR = EAudioMixerChannel::FrontRight,			// Front Right
+	FC = EAudioMixerChannel::FrontCenter,			// Front Center
+	LFE = EAudioMixerChannel::LowFrequency,			// Low-Frequency Effects (Subwoofer)
+	FLC = EAudioMixerChannel::FrontLeftOfCenter,	// Front Left Center
+	FRC = EAudioMixerChannel::FrontRightOfCenter,	// Front Right Center
+	SL = EAudioMixerChannel::SideLeft,				// Side Left
+	SR = EAudioMixerChannel::SideRight,				// Side Right
+	BL = EAudioMixerChannel::BackLeft,				// Back Left
+	BR = EAudioMixerChannel::BackRight,				// Back Right
+	BC = EAudioMixerChannel::BackCenter,			// Back Center
+	TFL = EAudioMixerChannel::TopFrontLeft,			// Top Front Left
+	TFR = EAudioMixerChannel::TopFrontRight,		// Top Front Right
+	TBL = EAudioMixerChannel::TopBackLeft,			// Top Back Left
+	TBR = EAudioMixerChannel::TopBackRight,			// Top Back Right
+
+	NumChannels = EAudioMixerChannel::ChannelTypeCount,
+	DefaultChannel = EAudioMixerChannel::DefaultChannel,
+	Unknown = EAudioMixerChannel::Unknown
+};
+
+// Pasted from UHT generated.h, but we can't use because of AudioUnitTests.
+#define FOREACH_ENUM_ESPEAKERSHORTNAMES(op) \
+	op(ESpeakerShortNames::FL) \
+	op(ESpeakerShortNames::FR) \
+	op(ESpeakerShortNames::FC) \
+	op(ESpeakerShortNames::LFE) \
+	op(ESpeakerShortNames::FLC) \
+	op(ESpeakerShortNames::FRC) \
+	op(ESpeakerShortNames::SL) \
+	op(ESpeakerShortNames::SR) \
+	op(ESpeakerShortNames::BL) \
+	op(ESpeakerShortNames::BR) \
+	op(ESpeakerShortNames::BC) \
+	op(ESpeakerShortNames::TFL) \
+	op(ESpeakerShortNames::TFR) \
+	op(ESpeakerShortNames::TBL) \
+	op(ESpeakerShortNames::TBR)
+
+AUDIOCHANNELAGNOSTICCORE_API const TCHAR* LexToString(const ESpeakerShortNames InSpeaker);
+AUDIOCHANNELAGNOSTICCORE_API FName LexToName(const ESpeakerShortNames InSpeaker);
+AUDIOCHANNELAGNOSTICCORE_API TOptional<ESpeakerShortNames> NameToShortSpeakerName(const FName InName);
+AUDIOCHANNELAGNOSTICCORE_API uint32 ShortNamesToChannelMask(const TArray<ESpeakerShortNames>& InChannels);
+
+namespace Audio
+{
+	// Forward declare these for visitor.
+	class FDiscreteChannelTypeFamily;
+	class FSoundfieldChannelTypeFamily;
+	class FPackedChannelTypeFamily;
+	
+	class IChannelTypeVisitor
+	{
+	public:
+		virtual ~IChannelTypeVisitor() = default;
+		UE_EXPERIMENTAL(5.8,"AudioChannelAgnosticCore is experimental and in flux")
+		virtual void Visit(const FDiscreteChannelTypeFamily&) = 0;
+		UE_EXPERIMENTAL(5.8,"AudioChannelAgnosticCore is experimental and in flux")
+		virtual void Visit(const FSoundfieldChannelTypeFamily&) = 0;
+		UE_EXPERIMENTAL(5.8,"AudioChannelAgnosticCore is experimental and in flux")
+		virtual void Visit(const FCompositeChannelTypeFamily&) = 0;
+	};
+	
+	// Base type for all "channel" based types.
+	class FChannelTypeFamily : public FTypeFamily
+	{
+	public:
+		UE_NONCOPYABLE(FChannelTypeFamily);
+		using Super = FTypeFamily;
+
+		/**
+		 * Constructor. 
+		 * @param InUniqueName The unique name that will be used for look up in registry.
+		 * @param InFamilyTypeName The name of the concrete type (c++) that's defining this. (for safe casting). 
+		 * @param InNumChannels Num of channels in this type. (pure categorical, organisational entries will be 0).
+		 * @param InParentType Pointer to this parents type in the tree. Can be null.
+		 * @param InFriendlyName Friendly name to display to the user e.g. "Dolby Stereo (2.0)"
+		 * @param InbIsParentsDefault Marks if this entry is the default child of its parent. 
+		 */
+		UE_EXPERIMENTAL(5.8,"AudioChannelAgnosticCore is experimental and in flux")
+		AUDIOCHANNELAGNOSTICCORE_API explicit FChannelTypeFamily(
+			const FName& InUniqueName,
+			const FName& InFamilyTypeName,
+			const int32 InNumChannels,
+			FChannelTypeFamily* InParentType,
+			const FString& InFriendlyName,
+			const bool InbIsParentsDefault,
+			const bool InbIsAbstract);
+
+		UE_EXPERIMENTAL(5.8,"AudioChannelAgnosticCore is experimental and in flux")
+		virtual ~FChannelTypeFamily() = default;
+
+		/**
+		 * If this type can be instantiated or is just organisational.
+		 * @return true if abstract, false otherwise
+		 */
+		UE_EXPERIMENTAL(5.8,"AudioChannelAgnosticCore is experimental and in flux")
+		bool IsAbstract() const
+		{
+			return bIsAbstract;
+		}
+
+		/**
+		 * If this is marked as being the default on its parent
+		 * @return true if true, false otherwise.
+		 */
+		UE_EXPERIMENTAL(5.8,"AudioChannelAgnosticCore is experimental and in flux")
+		bool IsParentsDefault() const
+		{
+			return bIsParentsDefault;
+		}
+
+		/**
+		 * Returns the default child if one exists. Example. "Stereo" would return "Stereo_2_0"
+		 * @return Pointer to type, null otherwise.
+		 */
+		UE_EXPERIMENTAL(5.8,"AudioChannelAgnosticCore is experimental and in flux")
+		const FChannelTypeFamily* GetDefaultChild() const { return DefaultChild; }
+
+		/**
+		 * Num of channels  
+		 * @return >= 0
+		 */
+		UE_EXPERIMENTAL(5.8,"AudioChannelAgnosticCore is experimental and in flux")
+		int32 NumChannels() const { return NumChannelsPrivate; }
+
+		using IPanner = IDiscretePanner;
+		
+		struct FGetPannerParams
+		{
+			// Control which algorithm to use for this format.
+			// Not currently used.
+			// FName PannerType = TEXT("VBAP 2D");
+		};
+		UE_EXPERIMENTAL(5.8,"AudioChannelAgnosticCore is experimental and in flux")
+		virtual const IPanner* GetPanner(const FGetPannerParams& InParams = {}) const { return nullptr; }
+		
+		/**
+		 * Returns the family name (C++ type) of this ChannelType. This allows us to safely downcast (if necessary).
+		 * @return The name of this type Family. (i.e. Discrete/Ambisoncs). 
+		 */
+		UE_EXPERIMENTAL(5.8,"AudioChannelAgnosticCore is experimental and in flux")
+		FName GetFamilyName() const { return FamilyType; }
+
+		/**
+		 * Returns the name of the channel. For discrete this would be what speaker its mapped to etc.
+		 * @param InChannelIndex Index of channel >= 0 and < NumChannels.
+		 * @return Optional containing the Channel name and friendly name.
+		 */
+		struct FChannelName
+		{
+			FName Name;
+			FString FriendlyName;
+		};
+
+		UE_EXPERIMENTAL(5.8,"AudioChannelAgnosticCore is experimental and in flux")
+		virtual TOptional<FChannelName> GetChannelName(const int32 InChannelIndex) const { return {}; }
+
+		template<typename T>
+		UE_EXPERIMENTAL(5.8,"AudioChannelAgnosticCore is experimental and in flux")
+		T* Cast()
+		{
+			PRAGMA_DISABLE_EXPERIMENTAL_WARNINGS
+			if (T::GetFamilyTypeName() == GetFamilyName())
+			{
+				return static_cast<T*>(this);
+			}
+			PRAGMA_ENABLE_EXPERIMENTAL_WARNINGS
+			return nullptr;
+		}
+
+		template<typename T>
+		UE_EXPERIMENTAL(5.8,"AudioChannelAgnosticCore is experimental and in flux")
+		const T* Cast() const { return const_cast<FChannelTypeFamily*>(this)->Cast<T>(); }
+
+		UE_EXPERIMENTAL(5.8,"AudioChannelAgnosticCore is experimental and in flux")
+		virtual void Accept(IChannelTypeVisitor& InVisitor) const
+		{
+			checkNoEntry();
+		};
+
+	protected:
+		// Both are friends so we can hide the api.
+		friend class FDiscreteChannelTypeFamily;
+		friend class FSoundfieldChannelTypeFamily;
+
+		UE_EXPERIMENTAL(5.8,"AudioChannelAgnosticCore is experimental and in flux")
+		void SetNumChannels(const int32 InNumChannels)
+		{
+			NumChannelsPrivate = InNumChannels;
+		}
+		
+	private:
+		bool bIsAbstract = false;
+		bool bIsParentsDefault = false;
+		const FChannelTypeFamily* DefaultChild = nullptr;
+		int32 NumChannelsPrivate = 0;
+		FName FamilyType; 
+	};
+	
+	class FDiscreteChannelTypeFamily : public FChannelTypeFamily 
+	{
+	public:
+		UE_NONCOPYABLE(FDiscreteChannelTypeFamily);
+		using Super = FChannelTypeFamily;
+
+		struct FSpeaker
+		{
+			FName ID;								// ID of speaker.
+			float AzimuthDegrees = 0.f;				// -180 (left), 0 (front), 180 (right).
+			float ElevationDegrees = 0.f;			// -180 (down), 0 (ground), 180 (up). 
+			bool bIsSpatialized = true;				// If false, azimuth+elevation are ignored.
+		};
+		
+		static FName GetFamilyTypeName()
+		{
+			static const FName Name = TEXT("Discrete");
+			return Name;
+		}
+
+		UE_EXPERIMENTAL(5.8,"AudioChannelAgnosticCore is experimental and in flux")
+		AUDIOCHANNELAGNOSTICCORE_API FDiscreteChannelTypeFamily(const FName& InUniqueName, FChannelTypeFamily* InParentType, const FString& InFriendlyName,
+			const TArray<FSpeaker>& InOrder, const bool bIsParentsDefault, const bool bIsAbstract);
+
+		UE_EXPERIMENTAL(5.8,"AudioChannelAgnosticCore is experimental and in flux")
+		AUDIOCHANNELAGNOSTICCORE_API virtual ~FDiscreteChannelTypeFamily() override;
+		
+		/**
+		 * Returns true if this format defines any height data. 
+		 * @return true if heights are found.
+		 */
+		UE_EXPERIMENTAL(5.8,"AudioChannelAgnosticCore is experimental and in flux")
+		bool HasHeightChannels() const
+		{
+			return Order.FindByPredicate([](const FSpeaker& i) -> bool { return !FMath::IsNearlyZero(i.ElevationDegrees); }) != nullptr;
+		}
+		
+		/**
+		 * Find the index in the list of Channels for speaker.
+		 * @param InSpeaker 
+		 * @return >= 0 (index) INDEX_NONE (failed to find).
+		 */
+		UE_EXPERIMENTAL(5.8,"AudioChannelAgnosticCore is experimental and in flux")
+		int32 FindSpeakerIndex(const ESpeakerShortNames InSpeaker) const
+		{
+			const FName SpeakerName = LexToString(InSpeaker);
+			return Order.IndexOfByPredicate([SpeakerName](const FSpeaker& i) -> bool { return SpeakerName == i.ID; });
+		}
+
+		UE_EXPERIMENTAL(5.8,"AudioChannelAgnosticCore is experimental and in flux")
+		int32 FindSpeakerIndex(const FName InSpeaker) const
+		{
+			return Order.IndexOfByPredicate([InSpeaker](const FSpeaker& i) -> bool { return InSpeaker == i.ID; });
+		}
+		
+
+		/**
+		 * Checks if a paticular speaker is present in the Channels for this Format.
+		 * @param InSpeaker 
+		 * @return true (exists), false otherwise.
+		 */
+		UE_EXPERIMENTAL(5.8,"AudioChannelAgnosticCore is experimental and in flux")
+		bool HasSpeaker(const ESpeakerShortNames InSpeaker) const
+		{
+			PRAGMA_DISABLE_EXPERIMENTAL_WARNINGS
+			return FindSpeakerIndex(InSpeaker) != INDEX_NONE;
+			PRAGMA_ENABLE_EXPERIMENTAL_WARNINGS
+		}
+
+		/**
+		 * Get the function used for translating/transcoding between this and the "To" type.
+		 * @param InToType The Type we are transcoding into.
+		 * @return The Function Object. (this will be not be set if a function wasn't returned)
+		 */	
+		UE_EXPERIMENTAL(5.8,"AudioChannelAgnosticCore is experimental and in flux")
+		AUDIOCHANNELAGNOSTICCORE_API virtual TOptional<FChannelName> GetChannelName(const int32 InChannelIndex) const override;
+
+		/**
+		 * Get the order of the speakers defined on this discrete format.
+		 * @return Array of speakers in order.
+		 */
+		UE_EXPERIMENTAL(5.8,"AudioChannelAgnosticCore is experimental and in flux")
+		const TArray<FSpeaker>& GetSpeakerOrder() const { return Order; }
+	
+		UE_EXPERIMENTAL(5.8,"AudioChannelAgnosticCore is experimental and in flux")
+		TArray<FName> GetSpeakerIds() const
+		{
+			TArray<FName> Ids;
+			Algo::Transform(Order, Ids, [](const FSpeaker& i) -> FName { return i.ID; });
+			return Ids;
+		}
+
+		UE_EXPERIMENTAL(5.8,"AudioChannelAgnosticCore is experimental and in flux")
+		virtual void Accept(IChannelTypeVisitor& InVisitor) const override
+		{
+			PRAGMA_DISABLE_EXPERIMENTAL_WARNINGS
+			InVisitor.Visit(*this);
+			PRAGMA_ENABLE_EXPERIMENTAL_WARNINGS
+		}
+
+		UE_EXPERIMENTAL(5.8,"AudioChannelAgnosticCore is experimental and in flux")
+		virtual const IPanner* GetPanner(const FGetPannerParams& InParams ={}) const override
+		{
+			PRAGMA_DISABLE_EXPERIMENTAL_WARNINGS
+			return Panner.Get();	
+			PRAGMA_ENABLE_EXPERIMENTAL_WARNINGS
+		}
+
+	protected:
+		static TUniquePtr<IDiscretePanner> MakePanner(const TArray<FSpeaker>& Order);
+		TArray<FSpeaker> Order;
+		TUniquePtr<IPanner> Panner;	// Just one for now.
+	};
+
+	// Soundfield 
+	class FSoundfieldChannelTypeFamily : public FChannelTypeFamily 
+	{
+	public:
+		UE_NONCOPYABLE(FSoundfieldChannelTypeFamily);
+		using Super = FChannelTypeFamily;
+		
+		PRAGMA_DISABLE_EXPERIMENTAL_WARNINGS
+		
+		UE_EXPERIMENTAL(5.8,"AudioChannelAgnosticCore is experimental and in flux")
+		static FName GetFamilyTypeName()
+		{
+			static const FName Name = TEXT("Soundfield");
+			return Name;
+		}
+
+		UE_EXPERIMENTAL(5.8,"AudioChannelAgnosticCore is experimental and in flux")
+		AUDIOCHANNELAGNOSTICCORE_API static int32 OrderToNumChannels(const int32 InOrder);
+
+		UE_EXPERIMENTAL(5.8,"AudioChannelAgnosticCore is experimental and in flux")
+		explicit FSoundfieldChannelTypeFamily(const FName& InUniqueName, const int32 InOrder, FChannelTypeFamily* InParentType, const FString& InFriendlyName, const bool bIsParentsDefault, const bool InbIsAbstract)
+			: Super(InUniqueName, GetFamilyTypeName(), OrderToNumChannels(InOrder), InParentType, InFriendlyName, bIsParentsDefault, InbIsAbstract)
+			, Order(InOrder)
+		{}
+
+		UE_EXPERIMENTAL(5.8,"AudioChannelAgnosticCore is experimental and in flux")
+		int32 GetAmbisonicsOrder() const { return Order; }
+		
+		UE_EXPERIMENTAL(5.8,"AudioChannelAgnosticCore is experimental and in flux")
+		AUDIOCHANNELAGNOSTICCORE_API virtual TOptional<FChannelName> GetChannelName(const int32 InChannelIndex) const override;
+
+		UE_EXPERIMENTAL(5.8,"AudioChannelAgnosticCore is experimental and in flux")
+		virtual void Accept(IChannelTypeVisitor& InVisitor) const override
+		{
+			InVisitor.Visit(*this);
+		}
+		
+		PRAGMA_ENABLE_EXPERIMENTAL_WARNINGS		
+	private:
+		int32 Order = 0;
+	};
+
+	
+	class FCompositeChannelTypeFamily : public FChannelTypeFamily
+	{
+	public:
+		UE_NONCOPYABLE(FCompositeChannelTypeFamily);
+
+		using Super = FChannelTypeFamily;
+		
+		PRAGMA_DISABLE_EXPERIMENTAL_WARNINGS		
+
+		UE_EXPERIMENTAL(5.8,"AudioChannelAgnosticCore is experimental and in flux")
+		static FName GetFamilyTypeName()
+		{
+			static const FName Name = TEXT("Composite");
+			return Name;
+		}
+		UE_EXPERIMENTAL(5.8,"AudioChannelAgnosticCore is experimental and in flux")
+		explicit FCompositeChannelTypeFamily(const FName& InUniqueName, FChannelTypeFamily* InParentType, const FString& InFriendlyName, const bool bIsParentsDefault)
+			: Super(InUniqueName, GetFamilyTypeName(), 0, InParentType, InFriendlyName, bIsParentsDefault, false)
+		{}
+	
+		UE_EXPERIMENTAL(5.8,"AudioChannelAgnosticCore is experimental and in flux")
+		int32 NumTypes() const
+		{
+			return Types.Num();
+		}
+
+		struct FContainedType
+		{
+			FName Type;					// Type name of embedded format.
+			uint32 Offset = 0;			// Index in channels.	
+			uint32 Count = 0;			// Count in channels.
+		};		
+		
+		UE_EXPERIMENTAL(5.8,"AudioChannelAgnosticCore is experimental and in flux")		
+		const FContainedType* GetType(const int32 InTypeIndex) const
+		{
+			if (Types.IsValidIndex(InTypeIndex))
+			{
+				return &Types[InTypeIndex];
+			}
+			return nullptr;
+		}
+		
+		UE_EXPERIMENTAL(5.8,"AudioChannelAgnosticCore is experimental and in flux")
+		virtual void Accept(IChannelTypeVisitor& InVisitor) const override
+		{
+			InVisitor.Visit(*this);
+		}
+		
+		PRAGMA_ENABLE_EXPERIMENTAL_WARNINGS
+		
+	protected:
+		TArray<FContainedType> Types;
+	};
+	
+	
+	// Channel type registry.
+	class IChannelTypeRegistry : public ITypeFamilyRegistry
+	{
+	public:
+		PRAGMA_DISABLE_EXPERIMENTAL_WARNINGS				
+		
+		UE_EXPERIMENTAL(5.8,"AudioChannelAgnosticCore is experimental and in flux")			
+		const FChannelTypeFamily* FindChannel(const FName InName) const
+		{
+			return Find<FChannelTypeFamily>(InName);
+		}
+		UE_EXPERIMENTAL(5.8,"AudioChannelAgnosticCore is experimental and in flux")			
+		FChannelTypeFamily* FindChannel(const FName InName)
+		{
+			return const_cast<FChannelTypeFamily*>(const_cast<const IChannelTypeRegistry*>(this)->FindChannel(InName));
+		}
+	
+		UE_EXPERIMENTAL(5.8,"AudioChannelAgnosticCore is experimental and in flux")			
+		virtual const TSharedPtr<const FDiscreteChannelTypeFamily> FindDiscreteChannelFromBitMask(const uint32 InChannelMask) const = 0;
+
+		/**
+		 * Find the first concrete (non-abstract, > 0 channels) channel. (e.g. Stereo -> 'Stereo_2_0')
+		 * Each Abstract type should have a designated default child.
+		 *
+		 * @param InName Name of format.
+		 * @return Pointer to type (on success), nullptr on failure
+		 */
+		UE_EXPERIMENTAL(5.8,"AudioChannelAgnosticCore is experimental and in flux")			
+		const FChannelTypeFamily* FindConcreteChannel(const FName InName) const
+		{
+			check(!InName.IsNone());
+
+			// Walk looking for non-abstract default child. (e.g. Surround->Surround5X->Surround 5.1)
+			const FChannelTypeFamily* Channel = FindChannel(InName);
+			while (Channel && Channel->IsAbstract())
+			{
+				Channel = Channel->GetDefaultChild();
+			}
+			return Channel;
+		}
+		UE_EXPERIMENTAL(5.8,"AudioChannelAgnosticCore is experimental and in flux")			
+		const FChannelTypeFamily& FindConcreteChannelChecked(const FName InName) const
+		{
+			const FChannelTypeFamily* Concrete = FindConcreteChannel(InName);
+			checkf(Concrete, TEXT("Concrete Channel not Found: %s"), *InName.ToString());
+			return *Concrete;
+		}
+
+		/**
+		 * Returns every registered format as an array.
+		 * @return Array of format types.
+		 */
+		UE_EXPERIMENTAL(5.8,"AudioChannelAgnosticCore is experimental and in flux")			
+		virtual TArray<TSharedRef<const FChannelTypeFamily>> GetAllChannelFormats() const = 0;
+		
+		UE_EXPERIMENTAL(5.8,"AudioChannelAgnosticCore is experimental and in flux")	
+		virtual TArray<TSharedRef<const FChannelTypeFamily>> FindAllFamily(const FName InFamilyName) const = 0;		
+
+		PRAGMA_ENABLE_EXPERIMENTAL_WARNINGS
+	};
+
+	UE_EXPERIMENTAL(5.8,"AudioChannelAgnosticCore is experimental and in flux")		
+	AUDIOCHANNELAGNOSTICCORE_API IChannelTypeRegistry& GetChannelRegistry();
+	
+	UE_EXPERIMENTAL(5.8,"AudioChannelAgnosticCore is experimental and in flux")	
+	AUDIOCHANNELAGNOSTICCORE_API void RegisterCatBaseClassesWithChannelRegistry(IChannelTypeRegistry& Registry = GetChannelRegistry());
+	
+	UE_EXPERIMENTAL(5.8,"AudioChannelAgnosticCore is experimental and in flux")		
+	AUDIOCHANNELAGNOSTICCORE_API void UnregisterCatBaseClassesWithChannelRegistry(IChannelTypeRegistry& Registry = GetChannelRegistry());
+
+}//namespace Audio
diff --git a/Engine/Source/Runtime/AudioChannelAgnosticCore/Public/TypeFamily/TypeFamily.h b/Engine/Source/Runtime/AudioChannelAgnosticCore/Public/TypeFamily/TypeFamily.h
new file mode 100644
--- /dev/null
+++ b/Engine/Source/Runtime/AudioChannelAgnosticCore/Public/TypeFamily/TypeFamily.h
@@ -0,0 +1,116 @@
+﻿// Copyright Epic Games, Inc. All Rights Reserved.
+
+#pragma once
+
+#include "Templates/RetainedRef.h"
+#include "Templates/UniquePtr.h"
+#include "UObject/NameTypes.h"
+
+namespace Audio
+{
+	class FTypeFamily
+	{
+		UE_NONCOPYABLE(FTypeFamily); 
+		
+		FName Name;
+		const FTypeFamily* Parent = nullptr;
+		FString FriendlyName;
+		
+		template<typename Predicate> 
+		UE_EXPERIMENTAL(5.8,"AudioChannelAgnosticCore is experimental and in flux")
+		const FTypeFamily* FindAncestor(Predicate Pred) const
+		{
+			for (const FTypeFamily* i=this; i; i=i->Parent)
+			{
+				if (::Invoke(Pred,i))
+				{
+					return i;
+				}
+			}
+			return nullptr;
+		}
+	public:
+		UE_EXPERIMENTAL(5.8,"AudioChannelAgnosticCore is experimental and in flux")
+		FTypeFamily(const FName InName, const FTypeFamily* InParent, const FString& InFriendlyName)
+			: Name(InName)
+			, Parent(InParent)
+			, FriendlyName(InFriendlyName)
+		{}
+		
+		virtual ~FTypeFamily() = default;
+
+		UE_EXPERIMENTAL(5.8,"AudioChannelAgnosticCore is experimental and in flux")
+		bool IsA(const FTypeFamily* Other) const
+		{
+			PRAGMA_DISABLE_EXPERIMENTAL_WARNINGS			
+			return nullptr != FindAncestor([Other](const FTypeFamily* i) -> bool { return i==Other; });
+			PRAGMA_ENABLE_EXPERIMENTAL_WARNINGS			
+		}
+		UE_EXPERIMENTAL(5.8,"AudioChannelAgnosticCore is experimental and in flux")
+		bool IsA(const FName& Other) const
+		{
+			PRAGMA_DISABLE_EXPERIMENTAL_WARNINGS			
+			return nullptr != FindAncestor([&](const FTypeFamily* i) -> bool { return i->Name==Other; });
+			PRAGMA_ENABLE_EXPERIMENTAL_WARNINGS			
+		}
+		UE_EXPERIMENTAL(5.8,"AudioChannelAgnosticCore is experimental and in flux")
+		FName GetName() const { return Name; }
+		UE_EXPERIMENTAL(5.8,"AudioChannelAgnosticCore is experimental and in flux")
+		FString GetFriendlyName() const { return FriendlyName; }
+		UE_EXPERIMENTAL(5.8,"AudioChannelAgnosticCore is experimental and in flux")
+		const FTypeFamily* GetParent() const { return Parent; }
+	};
+
+	// Simple type registry.
+	class ITypeFamilyRegistry
+	{
+	public:
+		UE_EXPERIMENTAL(5.8,"AudioChannelAgnosticCore is experimental and in flux")
+		virtual ~ITypeFamilyRegistry() = default;
+
+		UE_EXPERIMENTAL(5.8,"AudioChannelAgnosticCore is experimental and in flux")
+		virtual bool RegisterType(const FName& InUniqueName, TUniquePtr<FTypeFamily>&& InType) = 0;
+		
+		UE_EXPERIMENTAL(5.8,"AudioChannelAgnosticCore is experimental and in flux")
+		virtual bool UnregisterType(const FName& InUniqueName) = 0;
+
+		template<typename T> 
+		UE_EXPERIMENTAL(5.8,"AudioChannelAgnosticCore is experimental and in flux")
+		const T* Find(const FName& InUniqueName) const
+		{
+			// TODO: test if this is valid by using FName Comparison.	
+			PRAGMA_DISABLE_EXPERIMENTAL_WARNINGS			
+			return static_cast<const T*>(FindTypeInternal(InUniqueName));
+			PRAGMA_ENABLE_EXPERIMENTAL_WARNINGS			
+		}
+	protected:
+		UE_EXPERIMENTAL(5.8,"AudioChannelAgnosticCore is experimental and in flux")
+		virtual const FTypeFamily* FindTypeInternal(const FName InUniqueName) const = 0;
+	};
+
+	// Convenience template so we don't have to specialize each Find. 
+	template<typename T> struct TFamilyRegistry
+	{
+		ITypeFamilyRegistry& FamilyRegistry;
+		UE_EXPERIMENTAL(5.8,"AudioChannelAgnosticCore is experimental and in flux")
+		explicit TFamilyRegistry(ITypeFamilyRegistry& InFamilyRegistry)
+			: FamilyRegistry(InFamilyRegistry)
+		{}
+		UE_EXPERIMENTAL(5.8,"AudioChannelAgnosticCore is experimental and in flux")
+		const T* Find(const FName& InName) const
+		{
+			PRAGMA_DISABLE_EXPERIMENTAL_WARNINGS			
+			return FamilyRegistry.Find<T>(InName);
+			PRAGMA_ENABLE_EXPERIMENTAL_WARNINGS			
+		}
+		UE_EXPERIMENTAL(5.8,"AudioChannelAgnosticCore is experimental and in flux")
+		const T& FindChecked(const FName& InName) const
+		{
+			PRAGMA_DISABLE_EXPERIMENTAL_WARNINGS			
+			const T* Found = Find(InName);
+			check(Found);
+			return *Found;
+			PRAGMA_ENABLE_EXPERIMENTAL_WARNINGS			
+		}
+	};
+}
\ No newline at end of file
diff --git a/Engine/Source/Runtime/SignalProcessing/Private/Vbap.cpp b/Engine/Source/Runtime/SignalProcessing/Private/Vbap.cpp
--- a/Engine/Source/Runtime/SignalProcessing/Private/Vbap.cpp
+++ b/Engine/Source/Runtime/SignalProcessing/Private/Vbap.cpp
@@ -156,7 +156,7 @@ namespace Audio
 		// Stabilize results by always sorting on Channel ID
 		OutResults.Results.Sort([](const FPanResult& A, const FPanResult& B)
 		{
-			return A.ChannelID < B.ChannelID;
+			return A.ChannelID.FastLess(B.ChannelID);
 		});
 		
 		return true;
diff --git a/Engine/Source/Runtime/SignalProcessing/Public/DSP/Vbap.h b/Engine/Source/Runtime/SignalProcessing/Public/DSP/Vbap.h
--- a/Engine/Source/Runtime/SignalProcessing/Public/DSP/Vbap.h
+++ b/Engine/Source/Runtime/SignalProcessing/Public/DSP/Vbap.h
@@ -16,7 +16,7 @@ namespace Audio
 												// Listener relative in degrees
 			float AzimuthDegrees = 0.0f;		// -180 - 180.  (0 forward).
 			float ElevationDegrees = 0.0f;		// -180 - 180.  (0 ground).
-			int32 ChannelID = INDEX_NONE;		// Unique id of speaker.
+			FName ChannelID = {};				// Unique id of speaker.
 			int32 ChannelIndex = INDEX_NONE;	// Index in the format. 
 		};
 		struct FInputParams
@@ -29,7 +29,7 @@ namespace Audio
 		};
 		struct FPanResult
 		{
-			int32 ChannelID=INDEX_NONE;			// Unique ID (enum) or Hash of the Speaker Name.
+			FName ChannelID={};					// Unique ID of the Speaker Name.
 			int32 ChannelIndex=INDEX_NONE;		// Index in the speaker in format. 
 			float Gain=0.0f;
 		};
