speculative fix for metasound trigger delay node rounding issue on arm64
#rnx

#Synced-CL 50668864

diff --git a/Engine/Plugins/Runtime/Metasound/Source/MetasoundStandardNodes/Private/MetasoundTriggerDelayNode.cpp b/Engine/Plugins/Runtime/Metasound/Source/MetasoundStandardNodes/Private/MetasoundTriggerDelayNode.cpp
--- a/Engine/Plugins/Runtime/Metasound/Source/MetasoundStandardNodes/Private/MetasoundTriggerDelayNode.cpp
+++ b/Engine/Plugins/Runtime/Metasound/Source/MetasoundStandardNodes/Private/MetasoundTriggerDelayNode.cpp
@@ -85,7 +85,12 @@ namespace Metasound
 			},
 			[this](int32 StartFrame, int32 EndFrame)
 			{
-				const int32 FrameToTrigger = FMath::Max(0, FMath::RoundToInt(DelayTime->GetSeconds() * SampleRate)) + StartFrame;
+				// Static cast to int32 to ensure that values are cast to int32 
+				// before going into FMath::Max(...). This protects us from overflow issues 
+				// when converting between int64->int32.  On some platforms, int64->int32 
+				// can result in surprising outcomes when the int64 is out of range
+				// of the int32. 
+				const int32 FrameToTrigger = FMath::Max(static_cast<int32>(0), static_cast<int32>(FMath::RoundToInt(DelayTime->GetSeconds() * SampleRate))) + StartFrame;
 				TriggerOut->RemoveAfter(StartFrame);
 				TriggerOut->TriggerFrame(FrameToTrigger);
 			}
