Variant inputs on metasound nodes

#rnx

#Synced-CL 50668864

diff --git a/Engine/Plugins/Experimental/MetasoundExperimental/Source/MetasoundExperimentalRuntime/Private/MetasoundExampleWeightedSumNode.cpp b/Engine/Plugins/Experimental/MetasoundExperimental/Source/MetasoundExperimentalRuntime/Private/MetasoundExampleWeightedSumNode.cpp
new file mode 100644
--- /dev/null
+++ b/Engine/Plugins/Experimental/MetasoundExperimental/Source/MetasoundExperimentalRuntime/Private/MetasoundExampleWeightedSumNode.cpp
@@ -0,0 +1,314 @@
+// Copyright Epic Games, Inc. All Rights Reserved.
+
+#include "MetasoundExampleWeightedSumNode.h" 
+
+#include "Internationalization/Text.h"
+#include "Math/UnrealMathUtility.h"
+
+#include "MetasoundDataFactory.h"
+#include "MetasoundFacade.h"
+#include "MetasoundOperatorData.h"
+#include "MetasoundNodeRegistrationMacro.h"
+#include "MetasoundExecutableOperator.h"
+#include "MetasoundParamHelper.h"
+#include "MetasoundPrimitives.h"
+
+
+#define LOCTEXT_NAMESPACE "MetasoundExperimentalRuntime"
+
+namespace Metasound::Experimental
+{
+	namespace ExampleWeightedSumNodePrivate
+	{
+		DEFINE_METASOUND_PARAM(InputValue, "Value", "A value to sum");
+		DEFINE_METASOUND_PARAM(InputWeight, "Weight", "Relative weight for value when doing weighted sum.");
+		DEFINE_METASOUND_PARAM(OutputSum, "WeightedSum", "Weightedum of all inputs");
+
+		const FLazyName VariantName{"DataType"};
+
+		FClassInterface MakeClassInterface()
+		{
+			return CreateClassInterface
+			(
+				FSubInterfaceDescription
+				{
+					"Inputs", // Name of sub interface
+					1, 		// Minimum number of instances
+					1000, 	// Maximum number of instances
+					2 		// Default number of instances
+				},
+
+				FSubInterface
+				{
+					"Inputs", // Name of sub interface
+
+					// Vertices in sub interface
+					TInputVariantVertex<bool, int32, float>{METASOUND_GET_PARAM_NAME(InputValue), VariantName, METASOUND_GET_PARAM_METADATA(InputValue)},
+					TInputDataVertex<float>{METASOUND_GET_PARAM_NAME(InputWeight), METASOUND_GET_PARAM_METADATA(InputWeight)}
+				},
+
+				TOutputVariantVertex<bool, int32, float>{METASOUND_GET_PARAM_NAME(OutputSum), VariantName, METASOUND_GET_PARAM_METADATA(OutputSum)}
+			);
+		}
+		// Create the node's vertex interface based upon the number of inputs and outputs
+		// desired. 
+		FVertexInterface GetVertexInterface(uint32 InNumInputs, const FName& InVariantDataType)
+		{
+			// In the future, the FClassInterface will be stored on the FNodeClassMetadata struct
+			static const FClassInterface ClassInterface = MakeClassInterface();
+
+			TArray<FSubInterfaceConfiguration> SubInterfaceConfigs;
+			SubInterfaceConfigs.Add(FSubInterfaceConfiguration{"Inputs", InNumInputs});
+
+			TArray<FVariantConfiguration> VariantConfigs;
+			if (InVariantDataType.IsValid() && (!InVariantDataType.IsNone()))
+			{
+				VariantConfigs.Add(FVariantConfiguration{VariantName, InVariantDataType});
+			}
+
+			return ClassInterface.CreateVertexInterface(SubInterfaceConfigs, VariantConfigs);
+		}
+	}
+
+	template<typename VariantType>
+	class TExampleWeightedSumOperator : public TExecutableOperator<TExampleWeightedSumOperator<VariantType>>
+	{
+	public:
+		TExampleWeightedSumOperator(
+			TArray<TDataReadReference<float>> InInputWeights,
+			TArray<TDataReadReference<VariantType>> InInputValues,
+			TDataWriteReference<VariantType> InOutSum)
+		: InputWeights(MoveTemp(InInputWeights))
+		, InputValues(MoveTemp(InInputValues))
+		, OutSum(MoveTemp(InOutSum))
+		{
+			check(InputWeights.Num() == InputValues.Num());
+			Execute(); // Initialize OutSum
+		}
+
+		virtual void BindInputs(FInputVertexInterfaceData& InOutVertexData) override
+		{
+			using namespace ExampleWeightedSumNodePrivate;
+			InOutVertexData.BindSubInterfaceVertices("Inputs", METASOUND_GET_PARAM_NAME(InputValue), InputValues);
+			InOutVertexData.BindSubInterfaceVertices("Inputs", METASOUND_GET_PARAM_NAME(InputWeight), InputWeights);
+		}
+
+		virtual void BindOutputs(FOutputVertexInterfaceData& InOutVertexData) override
+		{
+			using namespace ExampleWeightedSumNodePrivate;
+			InOutVertexData.BindVertex(METASOUND_GET_PARAM_NAME(OutputSum), OutSum);
+		}
+
+		void Execute()
+		{
+			float Sum = 0.f;
+			const int32 NumInputs = InputWeights.Num();
+
+			if constexpr (std::is_same_v<VariantType, bool>)
+			{
+				// Special handling for bool to treat it as 1=true and 0=false
+				for (int32 i = 0; i < NumInputs; i++)
+				{
+					if (*InputValues[i])
+					{
+						Sum += *InputWeights[i];
+					}
+				}
+
+				// If the output is a bool, make it true if the sum is greater
+				// than or equal 0.5, false otherwise. 
+				if (Sum >= 0.5f)
+				{
+					*OutSum = true;
+				}
+				else
+				{
+					*OutSum = false;
+				}
+			}
+			else 
+			{
+				// Handle ints or floats
+				for (int32 i = 0; i < NumInputs; i++)
+				{
+					Sum += *InputValues[i] * *InputWeights[i];
+				}
+
+				if constexpr (std::is_same_v<VariantType, int32>)
+				{
+					// Round if the output is an int
+					*OutSum = FMath::RoundToInt(Sum);
+				}
+				else
+				{
+					// Assign if the output is a float
+					*OutSum = Sum;
+				}
+			}
+
+		}
+
+		void Reset(const IOperator::FResetParams& InParams)
+		{
+			Execute(); // Reinitialize output value
+		}
+
+	private:
+		TArray<TDataReadReference<float>> InputWeights;
+		TArray<TDataReadReference<VariantType>> InputValues;
+		TDataWriteReference<VariantType> OutSum;
+	};
+
+	// FNodeFacade doesn't handle the case of 1 node returning multiple different 
+	// operators very well. We could create some convenience FNodeFacade implementations
+	// to make it easier for folks to write nodes which utilize variants. 
+	class FExampleWeightedSumNode : public FBasicNode
+	{
+		class FFactory : public IOperatorFactory
+		{
+		public:
+
+			virtual TUniquePtr<IOperator> CreateOperator(const FBuildOperatorParams& InParams, FBuildResults& OutResults) override
+			{
+				using namespace ExampleWeightedSumNodePrivate;
+
+				const FVertexInterface& NodeInterface = InParams.Node.GetVertexInterface();
+
+				// Determine the variant data type by inspecting the node interface. We could look
+				// at putting this information in the IOperatorData accessible through the Node.GetOperatorData(),
+				// but this seems just as convenient. 
+				const FName VariantTypeName = NodeInterface.GetOutputVertex(METASOUND_GET_PARAM_NAME(OutputSum)).DataTypeName;
+				if (VariantTypeName.IsNone() || VariantTypeName == GVariantDataTypeName)
+				{
+					// TODO: Should this be an error? Should we somehow skip nodes that
+					// don't have an explicit type set for a variant?
+					return {};
+				}
+
+				// These method gets/creates all the instances of the input depending on 
+				// the number of subinterfaces that exist on the interface. 
+				TArray<TDataReadReference<float>> InputWeights = InParams.InputData.GetOrCreateDefaultSubInterfaceDataReadReferences<float>(
+					"Inputs", 								// Name of sub interface
+					METASOUND_GET_PARAM_NAME(InputWeight), 	// Name of vertex on sub interface
+					InParams.OperatorSettings);
+
+				// We need to handle the various possible datatypes of the inputs. This
+				// is probably the most complicated scenario we'll encounter because
+				// we are utilize both variants AND subinterfaces. The templatizing of the
+				// operator is better than what we used to do because we are no longer
+				// sharing lots of node registration data or repeating FNodeClassMetadata
+				// codegen for each datatype and subinterface.
+				//
+				// There may be a way where we rewrite this to be more like TDataReadReference<bool, int32, float>
+				// or where we templatize the minimum bits of the operator which would
+				// still over reasonable memory savings. 
+				if (GetMetasoundDataTypeName<bool>() == VariantTypeName)
+				{
+					return MakeUnique<TExampleWeightedSumOperator<bool>>(
+							MoveTemp(InputWeights),
+							InParams.InputData.GetOrCreateDefaultSubInterfaceDataReadReferences<bool>(
+								"Inputs", 								// Name of sub interface
+								METASOUND_GET_PARAM_NAME(InputValue), 	// Name of vertex on sub interface
+								InParams.OperatorSettings),
+							TDataWriteReference<bool>::CreateNew());
+				}
+				else if (GetMetasoundDataTypeName<int32>() == VariantTypeName)
+				{
+					return MakeUnique<TExampleWeightedSumOperator<int32>>(
+							MoveTemp(InputWeights),
+							InParams.InputData.GetOrCreateDefaultSubInterfaceDataReadReferences<int32>(
+								"Inputs", 								// Name of sub interface
+								METASOUND_GET_PARAM_NAME(InputValue), 	// Name of vertex on sub interface
+								InParams.OperatorSettings),
+							TDataWriteReference<int32>::CreateNew());
+				}
+				else if (GetMetasoundDataTypeName<float>() == VariantTypeName)
+				{
+					return MakeUnique<TExampleWeightedSumOperator<float>>(
+							MoveTemp(InputWeights),
+							InParams.InputData.GetOrCreateDefaultSubInterfaceDataReadReferences<float>(
+								"Inputs", 								// Name of sub interface
+								METASOUND_GET_PARAM_NAME(InputValue), 	// Name of vertex on sub interface
+								InParams.OperatorSettings),
+							TDataWriteReference<float>::CreateNew());
+				}
+				else
+				{
+					// We shouldn't ever end up here. It means we have a varianttype which
+					// we are not handling. 
+					checkNoEntry();
+					return {};
+				}
+			}
+		};
+
+	public:
+		FExampleWeightedSumNode(const FNodeData& InNodeData, const TSharedRef<const FNodeClassMetadata>& InClassMetadata)
+		: FBasicNode(InNodeData, InClassMetadata)
+		{
+		}
+
+
+		static FNodeClassMetadata CreateNodeClassMetadata()
+		{
+			using namespace ExampleWeightedSumNodePrivate;
+			return FNodeClassMetadata
+			{
+				FNodeClassName{ "Experimental", "ConfigurableWeightedSumOperator", "" },
+				1, // Major version
+				0, // Minor version
+				METASOUND_LOCTEXT("ExampleConfigurableWeightedSumNodeName", "Weighted Sum (Example Node)"),	
+				METASOUND_LOCTEXT("ExampleConfigurableWeightedSumNodeDescription", "A Node which shows how to make a configurable node for yourself using subinterfaces and variant vertices."),	
+				TEXT("UE"), // Author
+				METASOUND_LOCTEXT("ExampleConfigurableWeightedSumPromptIfMissing", "Enable the MetaSoundExperimental Plugin"), // Prompt if missing
+				ExampleWeightedSumNodePrivate::GetVertexInterface(1 /* default num inputs */, FName()), // TODO: Update FNodeClassMetadata to store an FClassInterface instead of a FVertexInterface
+				{}
+			};
+		}
+
+		virtual FOperatorFactorySharedRef GetDefaultOperatorFactory() const override
+		{
+			return MakeShared<FFactory>();
+		}
+	private:
+
+
+	};
+
+
+	// The node extension must be registered along with the node. 
+	METASOUND_REGISTER_NODE_AND_CONFIGURATION(FExampleWeightedSumNode, FMetaSoundExperimentalExampleWeightedSumNodeConfiguration);
+}
+
+
+FMetaSoundExperimentalExampleWeightedSumNodeConfiguration::FMetaSoundExperimentalExampleWeightedSumNodeConfiguration()
+: Type() 
+, NumInputs(1)
+{
+}
+
+TArray<FString> UMetaSoundExperimentalExampleWeightedSumNodeConfigurationOptions::GetTypeOptions()
+{
+	using namespace Metasound;
+	using namespace Metasound::Experimental::ExampleWeightedSumNodePrivate;
+
+	static const FClassInterface ClassInterface = MakeClassInterface();
+	check(ClassInterface.GetVariantDescriptions().Num() == 1); // This example just has one variant. Otherwise we'd have to filter it by the variants name.
+	
+	TArray<FString> Options;
+	Options.Add(FName().ToString()); // Add empty option in case none is selected. 
+	for (const FName& TypeName: ClassInterface.GetVariantDescriptions()[0].DataTypes)
+	{
+		Options.Add(TypeName.ToString());
+	}
+	return Options;
+}
+
+TInstancedStruct<FMetasoundFrontendClassInterface> FMetaSoundExperimentalExampleWeightedSumNodeConfiguration::OverrideDefaultInterface(const FMetasoundFrontendClass& InNodeClass) const
+{
+	using namespace Metasound::Experimental::ExampleWeightedSumNodePrivate;
+	// Override the interface based upon the number of inputs and outputs desired. 
+	return TInstancedStruct<FMetasoundFrontendClassInterface>::Make(FMetasoundFrontendClassInterface::GenerateClassInterface(GetVertexInterface(NumInputs, Type)));
+}
+
+#undef LOCTEXT_NAMESPACE // MetasoundExperimentalRuntime
diff --git a/Engine/Plugins/Experimental/MetasoundExperimental/Source/MetasoundExperimentalRuntime/Private/MetasoundExampleWeightedSumNode.h b/Engine/Plugins/Experimental/MetasoundExperimental/Source/MetasoundExperimentalRuntime/Private/MetasoundExampleWeightedSumNode.h
new file mode 100644
--- /dev/null
+++ b/Engine/Plugins/Experimental/MetasoundExperimental/Source/MetasoundExperimentalRuntime/Private/MetasoundExampleWeightedSumNode.h
@@ -0,0 +1,49 @@
+// Copyright Epic Games, Inc. All Rights Reserved.
+
+#pragma once
+
+#include "MetasoundOperatorData.h"
+#include "MetasoundFrontendDocument.h"
+
+#include "MetasoundExampleWeightedSumNode.generated.h"
+
+// This is a workaround because we cannot put a UFUNCTION on the NodeConfiguration
+// because the node configuration is a USTRUCT and UFUNCTIONS are not allowed on 
+// USTRUCTs. 
+//
+// In the future, this will be unnecessary since we will have a built-in NodeConfiguration
+// for handling the sub-interface and variant options exposed via the FClassInterface.
+// That implementation will likely utilize a details customization to provide the 
+// options dropdown for variant type selection.
+UCLASS()
+class UMetaSoundExperimentalExampleWeightedSumNodeConfigurationOptions : public UObject
+{
+	GENERATED_BODY()
+	
+	public:
+
+	UFUNCTION(CallInEditor)
+	static TArray<FString> GetTypeOptions();
+};
+
+// In the future, this configuration should be not needed and will be replaced
+// by a built-in node configuration that handles any node utilizing sub-interfaces
+// or variants on a FClassInterface.
+USTRUCT()
+struct FMetaSoundExperimentalExampleWeightedSumNodeConfiguration : public FMetaSoundFrontendNodeConfiguration
+{
+	GENERATED_BODY()
+
+	FMetaSoundExperimentalExampleWeightedSumNodeConfiguration();
+
+	UPROPERTY(EditAnywhere, Category = General, meta=(GetOptions="MetasoundExperimentalRuntime.MetaSoundExperimentalExampleWeightedSumNodeCOnfigurationOptions.GetTypeOptions"))
+	FName Type;
+
+	UPROPERTY(EditAnywhere, Category = General, meta = (ClampMin = "1", ClampMax = "1000"))
+	uint32 NumInputs;
+
+
+	/* Get the current interface for the class based upon the node extension */
+	virtual TInstancedStruct<FMetasoundFrontendClassInterface> OverrideDefaultInterface(const FMetasoundFrontendClass& InNodeClass) const;
+};
+
diff --git a/Engine/Plugins/Runtime/Metasound/Source/MetasoundFrontend/Private/MetasoundFrontendDocument.cpp b/Engine/Plugins/Runtime/Metasound/Source/MetasoundFrontend/Private/MetasoundFrontendDocument.cpp
--- a/Engine/Plugins/Runtime/Metasound/Source/MetasoundFrontend/Private/MetasoundFrontendDocument.cpp
+++ b/Engine/Plugins/Runtime/Metasound/Source/MetasoundFrontend/Private/MetasoundFrontendDocument.cpp
@@ -597,6 +597,45 @@ bool FMetasoundFrontendNodeInterface::Update(const FMetasoundFrontendClassInterf
 		}	
 		bInterfaceUpdated |= true;
 	}
+
+	// If the interface was updated, make sure the node's interface is sorted
+	// the same as the class's interface
+	if (bInterfaceUpdated)
+	{
+		// Sort the input interface
+		check(InClassInterface.Inputs.Num() == Inputs.Num());
+		for (int32 ClassInputIndex = 0; ClassInputIndex < InClassInterface.Inputs.Num(); ClassInputIndex++)
+		{
+			int32 NodeInputIndex = Inputs.IndexOfByPredicate(FFindMatchingVertex{InClassInterface.Inputs[ClassInputIndex]});
+			check(NodeInputIndex != INDEX_NONE); // The matching node inputs should exist by now.
+			if (NodeInputIndex > ClassInputIndex)
+			{
+				Inputs.Swap(NodeInputIndex, ClassInputIndex);
+			}
+			else
+			{
+				// This sort algorithm should never have NodeInputIndex < ClassInputIndex
+				check(NodeInputIndex == ClassInputIndex);
+			}
+		}
+
+		// Sort the output interface
+		check(InClassInterface.Outputs.Num() == Outputs.Num());
+		for (int32 ClassOutputIndex = 0; ClassOutputIndex < InClassInterface.Outputs.Num(); ClassOutputIndex++)
+		{
+			int32 NodeOutputIndex = Outputs.IndexOfByPredicate(FFindMatchingVertex{InClassInterface.Outputs[ClassOutputIndex]});
+			check(NodeOutputIndex != INDEX_NONE); // The matching node inputs should exist by now.
+			if (NodeOutputIndex > ClassOutputIndex)
+			{
+				Outputs.Swap(NodeOutputIndex, ClassOutputIndex);
+			}
+			else
+			{
+				// This sort algorithm should never have NodeOutputIndex < ClassOutputIndex
+				check(NodeOutputIndex == ClassOutputIndex);
+			}
+		}
+	}
 	
 	return bInterfaceUpdated;
 }
diff --git a/Engine/Plugins/Runtime/Metasound/Source/MetasoundGraphCore/Private/MetasoundVertex.cpp b/Engine/Plugins/Runtime/Metasound/Source/MetasoundGraphCore/Private/MetasoundVertex.cpp
--- a/Engine/Plugins/Runtime/Metasound/Source/MetasoundGraphCore/Private/MetasoundVertex.cpp
+++ b/Engine/Plugins/Runtime/Metasound/Source/MetasoundGraphCore/Private/MetasoundVertex.cpp
@@ -1,8 +1,14 @@
 // Copyright Epic Games, Inc. All Rights Reserved.
 
 #include "MetasoundVertex.h"
+
+#include "MetasoundLog.h"
 #include "MetasoundVertexPrivate.h"
 
+#include "Algo/AllOf.h"
+#include "Algo/AnyOf.h"
+#include "Algo/Compare.h"
+#include "Algo/Count.h"
 #include "Algo/Find.h"
 #include "Containers/Array.h"
 
@@ -10,6 +16,14 @@ namespace Metasound
 {
 	namespace VertexPrivate
 	{
+		template<typename ArgType>
+		uint16 CastToSubInterfaceInstancePositionChecked(ArgType InArg)
+		{
+			static_assert(std::is_integral_v<ArgType>, "ArgType must be integral type");
+			checkf((InArg <= SubInterface::MaxInstancePosition) && (InArg >= SubInterface::MinInstancePosition), TEXT("SubInterface instance position (%d) out of range [%d, %d]"), (int32)InArg, SubInterface::MinInstancePosition, SubInterface::MaxInstancePosition);
+			return static_cast<uint16>(InArg);
+		}
+
 		// Functor used for finding FSubInterfaceLayouts by name. 
 		struct FEqualSubInterfaceName
 		{
@@ -26,70 +40,11 @@ namespace Metasound
 			}
 		};
 
-		FSubInterfaceDeclarationBuilder::FSubInterfaceDeclarationBuilder(TArray<FSubInterfaceLayout>& OutSubInterfaceLayouts)
-		: SubInterfaces(&OutSubInterfaceLayouts)
-		{
-		}
-
-		FSubInterfaceDeclarationBuilder::~FSubInterfaceDeclarationBuilder()
-		{
-			checkf(CurrentSubInterfaceIndex == INDEX_NONE, TEXT("Failed to beging/end all sub interfaces"));
-		}
-
-		void FSubInterfaceDeclarationBuilder::ReserveSubInterfaces(int32 Num)
-		{
-			SubInterfaces->Reserve(Num);
-		}
-
-PRAGMA_DISABLE_EXPERIMENTAL_WARNINGS
-		void FSubInterfaceDeclarationBuilder::Add(FBeginSubInterface&& InSubInterface)
-		{
-			PushSubInterfaceDeclaration(InSubInterface.Name);
-		}
-
-		void FSubInterfaceDeclarationBuilder::Add(FEndSubInterface&& InSubInterface)
-		{
-			PopSubInterfaceDeclaration();
-		}
-PRAGMA_ENABLE_EXPERIMENTAL_WARNINGS
-
-		void FSubInterfaceDeclarationBuilder::OnVertexAdded(const FVertexName& InName)
-		{
-			checkf((CurrentSubInterfaceIndex == INDEX_NONE) || (InName.GetNumber() == NAME_NO_NUMBER_INTERNAL), TEXT("Vertex %s in sub interface %s cannot have a trailing number because it is in a sub interface."), *InName.ToString(), *(*SubInterfaces)[CurrentSubInterfaceIndex].SubInterfaceName.ToString());
-			CurrentNumVertices++;
-		}
-
-		void FSubInterfaceDeclarationBuilder::PushSubInterfaceDeclaration(const FName& InName)
-		{
-			// Sub interface declarations should only happen once because the declaration
-			// needs to be consistent. Relying on developers declaring it identically
-			// is error prone so it is not allowed. The default number of times a
-			// sub interface is replicated can be controlled via the FSubInterfaceDescription.
-			checkf(INDEX_NONE == SubInterfaces->IndexOfByPredicate(FEqualSubInterfaceName{InName}), TEXT("Sub interface %s is already declared."), *InName.ToString());
-
-			SubInterfaces->Add(FSubInterfaceLayout{InName});
-			int32 NextSubInterfaceIndex = SubInterfaces->Num() - 1;
-
-			checkf(CurrentSubInterfaceIndex == INDEX_NONE, TEXT("Sub interface %s cannot be embedded inside sub interface %s"), *InName.ToString(), *(*SubInterfaces)[CurrentSubInterfaceIndex].SubInterfaceName.ToString());
-
-			CurrentSubInterfaceIndex = NextSubInterfaceIndex;
-			(*SubInterfaces)[CurrentSubInterfaceIndex].Instances.Add(FSubInterfaceLayout::FInstance{CurrentNumVertices, INDEX_NONE});
-		}
-
-		void FSubInterfaceDeclarationBuilder::PopSubInterfaceDeclaration()
+		FInputVertexInterfaceDeclarationBuilder::FInputVertexInterfaceDeclarationBuilder(TArray<FInputDataVertex>& OutVertices)
+		: Vertices(OutVertices)
 		{
-			if (ensureMsgf(SubInterfaces->IsValidIndex(CurrentSubInterfaceIndex), TEXT("Sub interface has incorrect Begin/End declaration")))
-			{
-				(*SubInterfaces)[CurrentSubInterfaceIndex].Instances.Last().End = CurrentNumVertices;
-				CurrentSubInterfaceIndex = INDEX_NONE;
-			}
 		}
 
-		FInputVertexInterfaceDeclarationBuilder::FInputVertexInterfaceDeclarationBuilder(TArray<FInputDataVertex>& OutVertices, TArray<FSubInterfaceLayout>& OutInstances)
-		: FSubInterfaceDeclarationBuilder(OutInstances)
-		, Vertices(OutVertices)
-		{
-		}
 
 		void FInputVertexInterfaceDeclarationBuilder::Add(FInputDataVertex&& InVertex)
 		{
@@ -98,12 +53,10 @@ PRAGMA_ENABLE_EXPERIMENTAL_WARNINGS
 			checkf(!Vertices.ContainsByPredicate(FEqualVertexName(InVertex.VertexName)), TEXT("Duplicate vertex name %s. Vertex names must be unique"), *InVertex.VertexName.ToString());
 
 			Vertices.Emplace(MoveTemp(InVertex));
-			OnVertexAdded(Vertices.Last().VertexName);
 		}
 
-		FOutputVertexInterfaceDeclarationBuilder::FOutputVertexInterfaceDeclarationBuilder(TArray<FOutputDataVertex>& OutVertices, TArray<FSubInterfaceLayout>& OutInstances)
-		: FSubInterfaceDeclarationBuilder(OutInstances)
-		, Vertices(OutVertices)
+		FOutputVertexInterfaceDeclarationBuilder::FOutputVertexInterfaceDeclarationBuilder(TArray<FOutputDataVertex>& OutVertices)
+		: Vertices(OutVertices)
 		{
 		}
 
@@ -114,7 +67,6 @@ PRAGMA_ENABLE_EXPERIMENTAL_WARNINGS
 			checkf(!Vertices.ContainsByPredicate(FEqualVertexName(InVertex.VertexName)), TEXT("Duplicate vertex name %s. Vertex names must be unique"), *InVertex.VertexName.ToString());
 
 			Vertices.Emplace(MoveTemp(InVertex));
-			OnVertexAdded(Vertices.Last().VertexName);
 		}
 
 		FEnvironmentDeclarationBuilder::FEnvironmentDeclarationBuilder(TArray<FEnvironmentVertex>& OutVertices)
@@ -123,38 +75,214 @@ PRAGMA_ENABLE_EXPERIMENTAL_WARNINGS
 		}
 
 		/** Builds an interface from a sub-interface configuration */
-		class FSubInterfaceConfigurationBuilder
+		template<typename VertexType>
+		class TInterfaceConfigurationBuilder
 		{
 		public:
-			FSubInterfaceConfigurationBuilder(TArray<FSubInterfaceLayout>& OutSubInterfaceLayouts)
-			: SubInterfaces(OutSubInterfaceLayouts)
+
+			TInterfaceConfigurationBuilder(TArray<VertexType>& OutVertices, TArray<FSubInterfaceLayout>& OutSubInterfaceLayouts)
+			: Vertices(OutVertices)
+			, SubInterfaces(OutSubInterfaceLayouts)
 			{
 			}
 
-			virtual ~FSubInterfaceConfigurationBuilder() = default;
+			virtual ~TInterfaceConfigurationBuilder() = default;
 
+			struct FBuildParams
+			{
+				EVertexDirection VertexDirection;
+				TConstArrayView<VertexType> ClassVertices;
+
+				const FVariantSchema* VariantSchema = nullptr;
+				TConstArrayView<FVariantConfiguration> VariantConfigs;
+
+				const FSubInterfaceSchema* SubInterfaceSchema = nullptr;
+				TConstArrayView<FSubInterfaceConfiguration> SubInterfaceConfigs;
+			};
+
+			/** In order to create the vertex interface, the subinterface and variant
+			 * configurations need to be incorporated. 
+			 *
+			 * Variant configurations are applied first since instances of the same 
+			 * subinterface are expected to share variant configurations. If there 
+			 * is an subinterface containing a variant, then all the instances of 
+			 * variant have to be the same type.
+			 *
+			 * For example
+			 * FClassInterface
+			 * {
+			 * 		FSubInterface
+			 * 		{
+			 * 			"SubInterface1",
+			 * 			TVariantInputVertex<bool, int32>("In1", ...)
+			 * 		}
+			 * 	}
+			 * 	If there are 5 instances of "SubInterface1", all 5 must have the same
+			 * 	variant configuration for "In1". All 5 vertices must be bool or int32. 
+			 */
+			void Build(const FBuildParams& InBuildParams)
+			{
+				// Initialize with a copy of the class vertexes
+				Vertices = InBuildParams.ClassVertices;
 
-			virtual void Build(TConstArrayView<FSubInterfaceDescription> InSubInterfaceDescriptions, TConstArrayView<FSubInterfaceConfiguration> InSubInterfaceConfigs)
+				// If there are subinterfaces, initialize the data structures for
+				// tracking the locations of subinterfaces
+				if (InBuildParams.SubInterfaceSchema)
+				{
+					InitializeSubInterfaceLayouts(InBuildParams.VertexDirection, *InBuildParams.SubInterfaceSchema);
+				}
+
+				// If there are variants, initialize variants before expanding
+				// subinterfaces instances.
+				if (InBuildParams.VariantSchema)
+				{
+					ApplyVariantConfigurations(InBuildParams.VertexDirection, *InBuildParams.VariantSchema, InBuildParams.VariantConfigs);
+				}
+
+				// Build sub interfaces 
+				if (InBuildParams.SubInterfaceSchema)
+				{
+					BuildSubInterfaces(*InBuildParams.SubInterfaceSchema, InBuildParams.SubInterfaceConfigs);
+				}
+			}
+
+		private:
+
+			void InitializeSubInterfaceLayouts(EVertexDirection InVertexDirection, const FSubInterfaceSchema& InSubInterfaceSchema)
 			{
-				/** Do nothing if there are no sub interfaces. */
-				if (SubInterfaces.IsEmpty())
+				TConstArrayView<FSubInterfaceClassLayout> ClassLayouts;
+				switch (InVertexDirection)
 				{
-					return;
+					case EVertexDirection::Input:
+						ClassLayouts = InSubInterfaceSchema.InputSubInterfaceLayouts;
+						break;
+					case EVertexDirection::Output:
+						ClassLayouts = InSubInterfaceSchema.OutputSubInterfaceLayouts;
+						break;
+					default:
+						{
+							// Unhandled vertex type
+							checkNoEntry();
+						}
+
 				}
 
+				/** Initialize subinterface layouts to 1 instance of each subinterface 
+				 * since we always have 1 instance of each subinterface in the 
+				 * class definition. */
+				for (const FSubInterfaceClassLayout& ClassLayout : ClassLayouts)
+				{
+					SubInterfaces.Add(FSubInterfaceLayout{ClassLayout.SubInterfaceName, {FSubInterfaceLayout::FInstance{ClassLayout.Begin, ClassLayout.End}}});
+				}
+			}
+
+			void ApplyVariantConfigurations(EVertexDirection InVertexDirection, const FVariantSchema& InVariantSchema, TConstArrayView<FVariantConfiguration> InVariantConfigs)
+			{
+				for (const FVariantConfiguration& VariantConfig : InVariantConfigs)
+				{
+					auto HasSameVariantName = [&VariantConfig](const FVariantDescription& InVariant)
+					{ 
+						return InVariant.VariantName == VariantConfig.VariantName; 
+					};
+
+					// Find the index of variant. The index can be used to filter
+					// the values in VertexKeyToVariantIndex in order to find the 
+					// actual vertices which are variants. 
+					const int32 VariantIndex = InVariantSchema.Variants.IndexOfByPredicate(HasSameVariantName);
+					if (INDEX_NONE == VariantIndex)
+					{
+						UE_LOG(LogMetaSound, Error, TEXT("Failed to find variant description in class interface for variant configuration %s"), *VariantConfig.VariantName.ToString());
+						continue;
+					}
+
+					// Make sure the configured data type is actually part of the
+					// supported set of datatypes
+					const FVariantDescription& VariantDescription = InVariantSchema.Variants[VariantIndex];
+					const int32 DataTypeIndex = VariantDescription.DataTypes.IndexOfByKey(VariantConfig.DataType);
+					if (INDEX_NONE == DataTypeIndex)
+					{
+						UE_LOG(LogMetaSound, Error, TEXT("Cannot configure variant %s with data type %s because it is not a supported data type of the variant"), *VariantConfig.VariantName.ToString(), *VariantConfig.DataType.ToString());
+						continue;
+					}
+
+					// Iterate through all variant vertices and apply configuration
+					// to the ones that are part of this particular variant.
+					for (const auto& VertexKeyIndexPair : InVariantSchema.VertexKeyToVariantIndex)
+					{
+						if (VertexKeyIndexPair.Value == VariantIndex)
+						{
+							if (VertexKeyIndexPair.Key.Direction == InVertexDirection)
+							{
+								auto HasSameVertexName = [&VertexKeyIndexPair](const VertexType& InVertex)
+								{
+									return InVertex.VertexName == VertexKeyIndexPair.Key.Name; 
+								};
+								VertexType* Vertex = Vertices.FindByPredicate(HasSameVertexName);
+
+								if (nullptr == Vertex)
+								{
+									// Vertices in the VariantSchema should always exist by now. 
+									UE_LOG(LogMetaSound, Error, TEXT("Failed to find vertex %s when configuring variant %s"), *VertexKeyIndexPair.Key.Name.ToString(), *VariantConfig.VariantName.ToString());
+									continue;
+								}
+								else if (Vertex->DataTypeName != GVariantDataTypeName)
+								{
+									// We shouldn't be configuration a variant vertex twice. 
+									UE_LOG(LogMetaSound, Error, TEXT("While attempting to configure variant vertex %s, the vertex is already configured to datatype %s when configuring variant %s"), *Vertex->VertexName.ToString(), *Vertex->DataTypeName.ToString(), *VariantConfig.VariantName.ToString());
+									continue;
+								}
+
+								// Set the configured data type.
+								Vertex->DataTypeName = VariantConfig.DataType;
+
+								// Input vertices need to also have their default 
+								// literals updated along with the data type.
+								if constexpr (std::is_base_of_v<FInputDataVertex, VertexType>)
+								{
+									const TArray<FLiteral, TInlineAllocator<2>>* DefaultLiterals = InVariantSchema.VariantLiterals.Find(Vertex->VertexName);
+									if (nullptr != DefaultLiterals)
+									{
+										if (DefaultLiterals->IsValidIndex(DataTypeIndex))
+										{
+											Vertex->SetDefaultLiteral((*DefaultLiterals)[DataTypeIndex]);
+										}
+									}
+								}
+							}
+						}
+					}
+				}
+			}
+
+			void BuildSubInterfaces(const FSubInterfaceSchema& InSubInterfaceSchema, TConstArrayView<FSubInterfaceConfiguration> InSubInterfaceConfigs)
+			{
+				if (SubInterfaces.Num() == 0)
+				{
+					return;
+				}
 
 				for (int32 LayoutIndex = 0; LayoutIndex < SubInterfaces.Num(); LayoutIndex++)
 				{
 					FSubInterfaceLayout& Layout = SubInterfaces[LayoutIndex];
 
-					const FSubInterfaceDescription* Description = Algo::FindBy(InSubInterfaceDescriptions, Layout.SubInterfaceName, [](const FSubInterfaceDescription& InDesc) { return InDesc.SubInterfaceName; });
+					// Find the description associated with the layout
+					const FSubInterfaceDescription* Description = Algo::FindBy(
+							InSubInterfaceSchema.SubInterfaces, 
+							Layout.SubInterfaceName, 
+							[](const FSubInterfaceDescription& InDesc) { return InDesc.SubInterfaceName; });
 
 					if (ensureMsgf(Description, TEXT("Missing sub interface description %s"), *Layout.SubInterfaceName.ToString()))
 					{
 						// Determine number of sub interface instances. 
 						uint32 Num = Description->NumDefault;
 
-						const FSubInterfaceConfiguration* Config = Algo::FindBy(InSubInterfaceConfigs, Layout.SubInterfaceName, [](const FSubInterfaceConfiguration& InConfig) { return InConfig.SubInterfaceName; });
+						const FSubInterfaceConfiguration* Config = Algo::FindBy(
+								InSubInterfaceConfigs, 
+								Layout.SubInterfaceName, 
+								[](const FSubInterfaceConfiguration& InConfig) { return InConfig.SubInterfaceName; });
+
+						// It's not required to have a configuration because the
+						// sub interface description provides NumDefault.
 						if (Config)
 						{
 							Num = FMath::Clamp(Config->Num, Description->Min, Description->Max);
@@ -164,11 +292,21 @@ PRAGMA_ENABLE_EXPERIMENTAL_WARNINGS
 						BuildSubInterfaceInstances(Layout, LayoutIndex, Num);
 					}
 				}
+				
+#if DO_CHECK
+				// Check that there are no duplicate names in the interface
+				for (int32 IndexA = 0; IndexA < Vertices.Num(); IndexA++)
+				{
+					const FVertexName& VertexNameA = Vertices[IndexA].VertexName;
+					for (int32 IndexB = IndexA + 1; IndexB < Vertices.Num(); IndexB++)
+					{
+						checkf(VertexNameA != Vertices[IndexB].VertexName, TEXT("Found duplicate names (%s) in interface"), *VertexNameA.ToString());
+					}
+				}
+#endif
 			}
 
 
-		private:
-
 			void BuildSubInterfaceInstances(FSubInterfaceLayout& InLayout, uint32 InLayoutIndex, uint32 InNum)
 			{
 				// We should always be beginning from a declaration of a sub interface
@@ -221,15 +359,21 @@ PRAGMA_ENABLE_EXPERIMENTAL_WARNINGS
 				Instances.Reserve(InNum);
 				Instances.SetNum(InNum);
 
-				for (uint32 Cntr = 1; Cntr < InNum; Cntr++)
+				if (InNum > 1)
 				{
-					// Initialize the position of the sub interface instance
-					const uint32 InsertIndex = ProtoBegin + (Cntr * ProtoNum);
-					Instances[Cntr].Begin = InsertIndex;
-					Instances[Cntr].End = InsertIndex + ProtoNum;
+					check((ProtoBegin + ProtoNum) <= Vertices.Num());
+					const TArray<VertexType> Prototype{Vertices.GetData() + ProtoBegin, ProtoNum};
+
+					for (uint32 Cntr = 1; Cntr < InNum; Cntr++)
+					{
+						// Initialize the position of the sub interface instance
+						const uint32 InsertIndex = ProtoBegin + (Cntr * ProtoNum);
+						Instances[Cntr].Begin = InsertIndex;
+						Instances[Cntr].End = InsertIndex + ProtoNum;
 
-					// Create the vertices of the sub interface instance. 
-					InsertSubInterfaceInstance(ProtoBegin, ProtoNum, InsertIndex, Cntr);
+						// Create the vertices of the sub interface instance. 
+						InsertSubInterfaceInstance(Prototype, InsertIndex, Cntr);
+					}
 				}
 			}
 
@@ -267,57 +411,19 @@ PRAGMA_ENABLE_EXPERIMENTAL_WARNINGS
 				}
 			}
 
-			virtual void InsertSubInterfaceInstance(uint32 InProtoBegin, uint32 InProtoNum, uint32 InInsertPos, uint32 InSubInterfaceInstanceIndex) = 0;
-			virtual void RemoveVerticesAt(uint32 InVertexIndexBegin, uint32 InNum) = 0;
-
-			TArray<FSubInterfaceLayout>& SubInterfaces;
-		};
-
-		/** TInterfaceConfigurationBuilder builds vertex interfaces from a sub interface
-		 * configuration.*/
-		template<typename VertexType>
-		class TInterfaceConfigurationBuilder : public FSubInterfaceConfigurationBuilder
-		{
-		public:
-			TInterfaceConfigurationBuilder(TArray<VertexType>& OutVertices, TArray<FSubInterfaceLayout>& OutSubInterfaceLayouts)
-			: FSubInterfaceConfigurationBuilder(OutSubInterfaceLayouts)
-			, Vertices(OutVertices)
-			{
-			}
 
-			virtual void Build(TConstArrayView<FSubInterfaceDescription> InSubInterfaceDescriptions, TConstArrayView<FSubInterfaceConfiguration> InSubInterfaceConfigs) override
+			void InsertSubInterfaceInstance(const TArray<VertexType>& InPrototype, uint32 InInsertPos, uint32 InSubInterfaceInstanceIndex)
 			{
-				// Build sub interfaces using parent class.
-				FSubInterfaceConfigurationBuilder::Build(InSubInterfaceDescriptions, InSubInterfaceConfigs);
-
-#if DO_CHECK
-				// Check that there are no duplicate names in the interface
-				for (int32 IndexA = 0; IndexA < Vertices.Num(); IndexA++)
+				if (InPrototype.Num() > 0)
 				{
-					const FVertexName& VertexNameA = Vertices[IndexA].VertexName;
-					for (int32 IndexB = IndexA + 1; IndexB < Vertices.Num(); IndexB++)
-					{
-						checkf(VertexNameA != Vertices[IndexB].VertexName, TEXT("Found duplicate names (%s) in interface"), *VertexNameA.ToString());
-					}
-				}
-#endif
-			}
-		private:
+					// check that the insert position of the sub interface instance is valid. 
+					check(InInsertPos <= static_cast<uint32>(Vertices.Num()));
 
-			virtual void InsertSubInterfaceInstance(uint32 InProtoBegin, uint32 InProtoNum, uint32 InInsertPos, uint32 InSubInterfaceInstanceIndex) override
-			{
-				// check that the prototype exists in the vertex array.
-				check((InProtoBegin + InProtoNum) <= static_cast<uint32>(Vertices.Num()));
-				// check that the insert position of the sub interface instance is valid. 
-				check(InInsertPos <= static_cast<uint32>(Vertices.Num()));
-
-				if (InProtoNum > 0)
-				{
 					// Copy vertices from prototype vertices
-					Vertices.Insert(Vertices.GetData() + InProtoBegin, InProtoNum, InInsertPos);
+					Vertices.Insert(InPrototype.GetData(), InPrototype.Num(), InInsertPos);
 
 					// Rename vertices
-					const uint32 NewEnd = InInsertPos + InProtoNum;
+					const uint32 NewEnd = InInsertPos + InPrototype.Num();
 					for (uint32 VertexIndex = InInsertPos;  VertexIndex < NewEnd; VertexIndex++)
 					{
 						checkf(
@@ -331,15 +437,18 @@ PRAGMA_ENABLE_EXPERIMENTAL_WARNINGS
 				}
 			}
 
-			virtual void RemoveVerticesAt(uint32 InVertexIndexBegin, uint32 InNum) override
+			void RemoveVerticesAt(uint32 InVertexIndexBegin, uint32 InNum) 
 			{
 				check((InVertexIndexBegin + InNum) <= static_cast<unsigned>(Vertices.Num()));
 				Vertices.RemoveAt(InVertexIndexBegin, InNum);
 			}
 
+
 			TArray<VertexType>& Vertices;
+			TArray<FSubInterfaceLayout>& SubInterfaces;
 		};
 
+
 		class FInputInterfaceConfigurationBuilder : public TInterfaceConfigurationBuilder<FInputDataVertex>
 		{
 		public:
@@ -357,7 +466,66 @@ PRAGMA_ENABLE_EXPERIMENTAL_WARNINGS
 			{
 			}
 		};
-	}
+
+		bool ContainsSubInterfaceDescription(const FClassInterfaceData& InClassInterfaceData, const FName& InName)
+		{
+			return Algo::AnyOf(InClassInterfaceData.SubInterfaces, [&InName](const FSubInterfaceDescription& InDesc) { return InDesc.SubInterfaceName == InName; });
+		}
+
+		void AssertSubInterfacesAreWellFormed(const FClassInterfaceData& InClassInterfaceData)
+		{
+#if DO_CHECK
+			for (const FSubInterfaceClassLayout& Layout : InClassInterfaceData.InputSubInterfaceLayouts)
+			{
+				checkf(ContainsSubInterfaceDescription(InClassInterfaceData, Layout.SubInterfaceName), TEXT("Missing sub interface description for sub interface (%s)"), *Layout.SubInterfaceName.ToString());
+				checkf(1 == Algo::CountBy(InClassInterfaceData.InputSubInterfaceLayouts, Layout.SubInterfaceName, [](const FSubInterfaceClassLayout& OtherLayout) { return OtherLayout.SubInterfaceName; }), TEXT("Sub interface names must be unique. Found duplicate (%s)"), *Layout.SubInterfaceName.ToString());
+			}
+
+			for (const FSubInterfaceClassLayout& Layout : InClassInterfaceData.OutputSubInterfaceLayouts)
+			{
+				checkf(ContainsSubInterfaceDescription(InClassInterfaceData, Layout.SubInterfaceName), TEXT("Missing sub interface description for sub interface (%s)"), *Layout.SubInterfaceName.ToString());
+				checkf(1 == Algo::CountBy(InClassInterfaceData.OutputSubInterfaceLayouts, Layout.SubInterfaceName, [](const FSubInterfaceClassLayout& OtherLayout) { return OtherLayout.SubInterfaceName; }), TEXT("Sub interface names must be unique. Found duplicate (%s)"), *Layout.SubInterfaceName.ToString());
+			}
+#endif // if DO_CHECK
+		}
+
+		void AssertVariantsAreWellFormed(const FClassInterfaceData& InClassInterfaceData)
+		{
+#if DO_CHECK
+			for (const FVariantDescription& Variant : InClassInterfaceData.Variants)
+			{
+				checkf(1 == Algo::CountBy(InClassInterfaceData.Variants, Variant.VariantName, [](const FVariantDescription& OtherVariant) { return OtherVariant.VariantName; }), TEXT("Variant names must be unique. Found duplicate (%s)"), *Variant.VariantName.ToString());
+			}
+
+			for (const auto& Pair : InClassInterfaceData.VertexKeyToVariantIndex)
+			{
+				checkf(Pair.Value < InClassInterfaceData.Variants.Num(), TEXT("Missing Variant for VariantVertex (%s)"), *Pair.Key.Name.ToString());
+			}
+
+			for (const FInputDataVertex& Input : InClassInterfaceData.Inputs)
+			{
+				if (Input.DataTypeName == GVariantDataTypeName)
+				{
+					checkf(InClassInterfaceData.VariantLiterals.Contains(Input.VertexName), TEXT("Missing input defaults for input variant vertex (%s)"), *Input.VertexName.ToString());
+					checkf(InClassInterfaceData.VertexKeyToVariantIndex.Contains(FVertexKey{Input.VertexName, EVertexDirection::Input}), TEXT("Missing variant for input variant vertex (%s)"), *Input.VertexName.ToString());
+				}
+			}
+
+			for (const FOutputDataVertex& Output : InClassInterfaceData.Outputs)
+			{
+				if (Output.DataTypeName == GVariantDataTypeName)
+				{
+					checkf(InClassInterfaceData.VertexKeyToVariantIndex.Contains(FVertexKey{Output.VertexName, EVertexDirection::Output}), TEXT("Missing variant for output variant vertex (%s)"), *Output.VertexName.ToString());
+				}
+			}
+#endif // if DO_CHECK
+		}
+
+		const VertexPrivate::FClassInterfaceDataBuilder& FBuildableInterface::GetClassInterfaceDataBuilder() const
+		{
+			return Builder;
+		}
+	} // namespace VertexPrivate
 
 	FDataVertexMetadata FDataVertexMetadata::EmptyBasic { FText{}, FText{}, false };
 	FDataVertexMetadata FDataVertexMetadata::EmptyAdvanced { FText{}, FText{}, true };
@@ -421,6 +589,73 @@ PRAGMA_ENABLE_EXPERIMENTAL_WARNINGS
 		return InLHS.VertexName.FastLess(InRHS.VertexName);
 	}
 
+	const FName FSubInterface::GetName() const
+	{
+		return SubInterfaceName;
+	}
+
+	const FLazyName GVariantDataTypeName{FName()}; // NAME_None
+
+	FVariantVertex::FVariantVertex(const FVertexName& InVertexName, const FVertexName& InVariantName, const FDataVertexMetadata& InMetadata, EVertexAccessType InAccessType)
+	: VertexName(InVertexName)
+	, VariantName(InVariantName)
+#if WITH_EDITORONLY_DATA
+	, Metadata(InMetadata)
+#endif // WITH_EDITORONLY_DATA
+	, AccessType(InAccessType)
+	{
+	}
+
+	FVariantVertex::FVariantVertex(const FVertexName& InVertexName, const FName& InVariantName, const TArray<FName>& InDataTypeNames, const FDataVertexMetadata& InMetadata, EVertexAccessType InAccessType)
+	: VertexName(InVertexName)
+	, VariantName(InVariantName)
+	, DataTypeNames(InDataTypeNames)
+#if WITH_EDITORONLY_DATA
+	, Metadata(InMetadata)
+#endif // WITH_EDITORONLY_DATA
+	, AccessType(InAccessType)
+	{
+	}
+
+	FInputVariantVertex::FInputVariantVertex(const FVertexName& InVertexName, const FName& InVariantName, const FDataVertexMetadata& InMetadata, EVertexAccessType InAccessType)
+	: FVariantVertex(InVertexName, InVariantName, InMetadata, InAccessType)
+	{
+	}
+
+	FInputVariantVertex::FInputVariantVertex(const FVertexName& InVertexName, const FName& InVariantName, const TArray<FName>& InDataTypeNames, const FDataVertexMetadata& InMetadata, EVertexAccessType InAccessType, const TArray<FLiteral>& InLiterals)
+	: FVariantVertex(InVertexName, InVariantName, InDataTypeNames, InMetadata, InAccessType)
+	, Literals(InLiterals)
+	{
+		checkf(Literals.Num() == DataTypeNames.Num(), TEXT("Mismatch in number of default literals (%d) vs data types (%d). These must be equal"), Literals.Num(), DataTypeNames.Num());
+	}
+
+
+	namespace VertexPrivate
+	{
+		FVariantSchema FVariantSchema::FromClassInterfaceData(const VertexPrivate::FClassInterfaceData& InClassInterfaceData)
+		{
+			FVariantSchema Schema;
+			Schema.VertexKeyToVariantIndex = InClassInterfaceData.VertexKeyToVariantIndex;
+			Schema.Variants = InClassInterfaceData.Variants;
+			// This custom copy is here because of the different allocators used in the TSortedMap<FVertexName, TArray<FLiteral, ALLOCATOR>>. 
+			Schema.VariantLiterals.Empty(InClassInterfaceData.VariantLiterals.Num());
+			for (const auto& Pair : InClassInterfaceData.VariantLiterals)
+			{
+				Schema.VariantLiterals.Add(Pair.Key, TArray<FLiteral, TInlineAllocator<2>>(Pair.Value));
+			}
+			return Schema;
+		}
+
+		FSubInterfaceSchema FSubInterfaceSchema::FromClassInterfaceData(const VertexPrivate::FClassInterfaceData& InClassInterfaceData)
+		{
+			FSubInterfaceSchema Schema;
+			Schema.SubInterfaces = InClassInterfaceData.SubInterfaces;
+			Schema.InputSubInterfaceLayouts = InClassInterfaceData.InputSubInterfaceLayouts;
+			Schema.OutputSubInterfaceLayouts = InClassInterfaceData.OutputSubInterfaceLayouts;
+			return Schema;
+		}
+	}
+
 	FInputVertexInterface::FInputVertexInterface() = default;
 
 	FInputVertexInterface::FInputVertexInterface(TArray<FInputDataVertex> InVertices, TArray<VertexPrivate::FSubInterfaceLayout> InSubInterfaces)
@@ -611,39 +846,338 @@ PRAGMA_ENABLE_EXPERIMENTAL_WARNINGS
 		return !(InLHS == InRHS);
 	}
 
+
+
 PRAGMA_DISABLE_EXPERIMENTAL_WARNINGS
-	FClassVertexInterface::FClassVertexInterface(FVertexInterface InInterface)
-	: FClassVertexInterface(MoveTemp(InInterface.GetInputInterface()), MoveTemp(InInterface.GetOutputInterface()), MoveTemp(InInterface.GetEnvironmentInterface()))
+	FClassInterface::FClassInterface(const VertexPrivate::FClassInterfaceData& InClassInterfaceData)
+	: Inputs(InClassInterfaceData.Inputs)
+	, Outputs(InClassInterfaceData.Outputs)
+	, EnvironmentVariables(InClassInterfaceData.EnvironmentVariables)
 	{
+		using namespace VertexPrivate;
+
+		AssertSubInterfacesAreWellFormed(InClassInterfaceData);
+		AssertVariantsAreWellFormed(InClassInterfaceData);
+		
+		if (InClassInterfaceData.SubInterfaces.Num())
+		{
+			SubInterfaceSchema = MakeUnique<FSubInterfaceSchema>(FSubInterfaceSchema::FromClassInterfaceData(InClassInterfaceData));
+		}
+
+		if (InClassInterfaceData.Variants.Num())
+		{
+			VariantSchema = MakeUnique<FVariantSchema>(FVariantSchema::FromClassInterfaceData(InClassInterfaceData));
+		}
 	}
 
-	FClassVertexInterface::FClassVertexInterface(FInputVertexInterface InInputs, FOutputVertexInterface InOutputs, FEnvironmentVertexInterface InEnvironment)
-	: FClassVertexInterface(TArray<FSubInterfaceDescription>{}, MoveTemp(InInputs), MoveTemp(InOutputs), MoveTemp(InEnvironment))
+	FVertexInterface FClassInterface::CreateVertexInterface(TConstArrayView<FSubInterfaceConfiguration> InSubInterfaceConfigurations, TConstArrayView<FVariantConfiguration> InVariantConfigurations) const
 	{
+		using namespace VertexPrivate;
+
+		FVertexInterface Interface{
+			FInputVertexInterface{},
+			FOutputVertexInterface{},
+			FEnvironmentVertexInterface{EnvironmentVariables}};
+
+		{
+			// Build the inputs 
+			FInputInterfaceConfigurationBuilder::FBuildParams Params
+			{
+				.VertexDirection = EVertexDirection::Input,
+				.ClassVertices  = Inputs,
+				.VariantSchema = VariantSchema.Get(),
+				.VariantConfigs = InVariantConfigurations,
+				.SubInterfaceSchema = SubInterfaceSchema.Get(),
+				.SubInterfaceConfigs = InSubInterfaceConfigurations
+			};
+
+			// Build the interface in-place
+			FInputInterfaceConfigurationBuilder(Interface.GetInputInterface()).Build(Params);
+		}
+
+		{
+			// Build the outputs 
+			FOutputInterfaceConfigurationBuilder::FBuildParams Params
+			{
+				.VertexDirection = EVertexDirection::Output,
+				.ClassVertices=Outputs,
+				.VariantSchema = VariantSchema.Get(),
+				.VariantConfigs = InVariantConfigurations,
+				.SubInterfaceSchema = SubInterfaceSchema.Get(),
+				.SubInterfaceConfigs = InSubInterfaceConfigurations
+			};
+
+			// Build the interface in-place
+			FOutputInterfaceConfigurationBuilder(Interface.GetOutputInterface()).Build(Params);
+		}
+
+		return Interface;
 	}
 
-	FClassVertexInterface::FClassVertexInterface(TArray<FSubInterfaceDescription> InSubInterfaceDescriptions, FInputVertexInterface InInputs, FOutputVertexInterface InOutputs, FEnvironmentVertexInterface InEnvironment)
-	: SubInterfaces(MoveTemp(InSubInterfaceDescriptions))
-	, Inputs(MoveTemp(InInputs))
-	, Outputs(MoveTemp(InOutputs))
-	, Environment(MoveTemp(InEnvironment))
+	TConstArrayView<FSubInterfaceDescription> FClassInterface::GetSubInterfaceDescriptions() const
 	{
-
+		if (SubInterfaceSchema)
+		{
+			return SubInterfaceSchema->SubInterfaces;
+		}
+		return {};
 	}
 
-	FVertexInterface FClassVertexInterface::CreateVertexInterface(TConstArrayView<FSubInterfaceConfiguration> InSubInterfaceConfigurations) const
+	TConstArrayView<FVariantDescription> FClassInterface::GetVariantDescriptions() const
 	{
-		using namespace VertexPrivate;
+		if (VariantSchema)
+		{
+			return VariantSchema->Variants;
+		}
+		return {};
+	}
 
-		FInputVertexInterface NewInputs(Inputs);
-		FInputInterfaceConfigurationBuilder(NewInputs).Build(SubInterfaces, InSubInterfaceConfigurations);
-		
-		FOutputVertexInterface NewOutputs(Outputs);
-		FOutputInterfaceConfigurationBuilder(NewOutputs).Build(SubInterfaces, InSubInterfaceConfigurations);
+	const FVariantDescription* FClassInterface::FindVariantDescriptionForInput(const FVertexName& InVertexName) const
+	{
+		if (VariantSchema)
+		{
+			if (const uint8* Idx = VariantSchema->VertexKeyToVariantIndex.Find(VertexPrivate::FVertexKey{InVertexName, VertexPrivate::EVertexDirection::Input}))
+			{
+				return &VariantSchema->Variants[*Idx];
+			}
+		}
+		return nullptr;
+	}
 
-		return FVertexInterface(MoveTemp(NewInputs), MoveTemp(NewOutputs), Environment);
+	const FVariantDescription* FClassInterface::FindVariantDescriptionForOutput(const FVertexName& InVertexName) const
+	{
+		if (VariantSchema)
+		{
+			if (const uint8* Idx = VariantSchema->VertexKeyToVariantIndex.Find(VertexPrivate::FVertexKey{InVertexName, VertexPrivate::EVertexDirection::Output}))
+			{
+				return &VariantSchema->Variants[*Idx];
+			}
+		}
+		return nullptr;
 	}
+
 PRAGMA_ENABLE_EXPERIMENTAL_WARNINGS
+
+	namespace VertexPrivate
+	{
+		template<typename ElementType, typename ArrayType>
+		struct TAssertOnDuplicateName
+		{
+			TAssertOnDuplicateName(const ElementType& InNewElement, const ArrayType& InArray, const TCHAR* InArrayTypeName)
+			{
+				auto HasSameName = [&InNewElement](const ArrayType::ElementType& InElement) -> bool 
+				{ 
+					return InElement.VertexName == InNewElement.VertexName; 
+				};
+				checkf(Algo::NoneOf<ArrayType>(InArray, HasSameName), TEXT("Cannot have duplicate %s with same name (%s)"), InArrayTypeName, *InNewElement.VertexName.ToString());
+			}
+		};
+
+		const FClassInterfaceData& FClassInterfaceDataBuilder::GetClassInterfaceData() const
+		{
+			return ClassInterfaceData;
+		}
+
+		void FClassInterfaceDataBuilder::Add(const FInputDataVertex* InInputVertex)
+		{
+			check(InInputVertex);
+			TAssertOnDuplicateName{*InInputVertex, ClassInterfaceData.Inputs, TEXT("inputs")};
+			ClassInterfaceData.Inputs.Add(*InInputVertex);
+		}
+
+		void FClassInterfaceDataBuilder::Add(const FOutputDataVertex* InOutputVertex)
+		{
+			check(InOutputVertex);
+			TAssertOnDuplicateName{*InOutputVertex, ClassInterfaceData.Outputs, TEXT("outputs")};
+			ClassInterfaceData.Outputs.Add(*InOutputVertex);
+		}
+
+		void FClassInterfaceDataBuilder::Add(const FEnvironmentVertex* InEnvironmentVertex)
+		{
+			check(InEnvironmentVertex);
+			TAssertOnDuplicateName{*InEnvironmentVertex, ClassInterfaceData.EnvironmentVariables, TEXT("environment variables")};
+			ClassInterfaceData.EnvironmentVariables.Add(*InEnvironmentVertex);
+		}
+
+		void FClassInterfaceDataBuilder::Add(const FSubInterface* InSubInterface)
+		{
+			check(InSubInterface);
+
+			const FClassInterfaceDataBuilder& SubInterfaceBuilder = InSubInterface->GetClassInterfaceDataBuilder();
+
+			checkf(!SubInterfaceBuilder.ContainsAnySubInterfaces(), TEXT("Cannot have a sub interface nested in a sub interface (%s)"), *InSubInterface->GetName().ToString());
+
+			// Add all the vertices associated with the sub interface
+			AddSubInterfaceVertices(InSubInterface->GetName(), SubInterfaceBuilder.GetClassInterfaceData());
+
+			// Add any variants from the subinterface to this interface
+			MergeSubInterfaceVariantData(SubInterfaceBuilder.GetClassInterfaceData());
+		}
+
+		void FClassInterfaceDataBuilder::Add(const FSubInterfaceDescription* InSubInterfaceDescription)
+		{
+			check(InSubInterfaceDescription);
+			checkf(!ContainsSubInterfaceDescription(ClassInterfaceData, InSubInterfaceDescription->SubInterfaceName), TEXT("Sub interface descriptions must have unique names. Found duplicate name (%s)"), *InSubInterfaceDescription->SubInterfaceName.ToString());
+
+			ClassInterfaceData.SubInterfaces.Add(*InSubInterfaceDescription);
+		}
+
+		void FClassInterfaceDataBuilder::Add(const FInputVariantVertex* InInputVariantVertex)
+		{
+			check(InInputVariantVertex);
+			// Add to Inputs array but give it a special data type to denote that it's
+			// a variant input
+			TAssertOnDuplicateName{*InInputVariantVertex, ClassInterfaceData.Inputs, TEXT("inputs")};
+
+			ClassInterfaceData.Inputs.Add(FInputDataVertex{
+					InInputVariantVertex->VertexName, 
+					GVariantDataTypeName, 
+#if WITH_EDITORONLY_DATA
+					InInputVariantVertex->Metadata, 
+#else
+					FDataVertexMetadata{},
+#endif
+					InInputVariantVertex->AccessType
+				});
+			
+			// Add to VertexKeyToVariantIndex and VariantLiterals to store extra data associated with variant.
+			ClassInterfaceData.VariantLiterals.Emplace(InInputVariantVertex->VertexName, InInputVariantVertex->Literals);
+			uint8 VariantIndex = FindOrAddVariant(*InInputVariantVertex);
+			ClassInterfaceData.VertexKeyToVariantIndex.Add(FVertexKey{InInputVariantVertex->VertexName, EVertexDirection::Input}, VariantIndex);
+		}
+
+		void FClassInterfaceDataBuilder::Add(const FOutputVariantVertex* InOutputVariantVertex)
+		{
+			check(InOutputVariantVertex);
+			// Add to Outputs array but give it a special data type to denote that it's
+			// a variant input
+			TAssertOnDuplicateName{*InOutputVariantVertex, ClassInterfaceData.Outputs, TEXT("outputs")};
+			ClassInterfaceData.Outputs.Add(FOutputDataVertex{
+					InOutputVariantVertex->VertexName, 
+					GVariantDataTypeName, 
+#if WITH_EDITORONLY_DATA
+					InOutputVariantVertex->Metadata, 
+#else
+					FDataVertexMetadata{},
+#endif
+					InOutputVariantVertex->AccessType});
+			
+			// Add to VertexKeyToVariantIndex to store extra data associated with variant.
+			uint8 VariantIndex = FindOrAddVariant(*InOutputVariantVertex);
+			ClassInterfaceData.VertexKeyToVariantIndex.Add(FVertexKey{InOutputVariantVertex->VertexName, EVertexDirection::Output}, VariantIndex);
+		}
+
+		void FClassInterfaceDataBuilder::AddSubInterfaceVertices(const FName& InSubInterfaceName, const FClassInterfaceData& InSubInterfaceData)
+		{
+			if (InSubInterfaceData.Inputs.Num())
+			{
+				ClassInterfaceData.InputSubInterfaceLayouts.Add(FSubInterfaceClassLayout
+				{ 
+					.SubInterfaceName = InSubInterfaceName,
+					.Begin = CastToSubInterfaceInstancePositionChecked(ClassInterfaceData.Inputs.Num()),
+					.End = CastToSubInterfaceInstancePositionChecked(ClassInterfaceData.Inputs.Num() + InSubInterfaceData.Inputs.Num())
+				});
+
+				ClassInterfaceData.Inputs.Append(InSubInterfaceData.Inputs);
+
+
+				// This should not be possible because FSubInterface objects should only be nested inside of class inputs or outputs and therefore only contain inputs OR outputs
+				checkf(InSubInterfaceData.Outputs.Num() == 0, TEXT("FSubInterface declarations must only contain inputs Or outputs. Inputs AND outputs found in sub interface (%s)"), *InSubInterfaceName.ToString());
+
+			}
+			else if (InSubInterfaceData.Outputs.Num())
+			{
+				ClassInterfaceData.OutputSubInterfaceLayouts.Add(FSubInterfaceClassLayout
+				{ 
+					.SubInterfaceName = InSubInterfaceName, 
+					.Begin = CastToSubInterfaceInstancePositionChecked(ClassInterfaceData.Outputs.Num()), 
+					.End = CastToSubInterfaceInstancePositionChecked(ClassInterfaceData.Outputs.Num() + InSubInterfaceData.Outputs.Num())
+				});
+
+				ClassInterfaceData.Outputs.Append(InSubInterfaceData.Outputs);
+
+				// This should not be possible because FSubInterface objects should only be nested inside of class inputs or outputs and therefore only contain inputs OR outputs
+				checkf(InSubInterfaceData.Inputs.Num() == 0, TEXT("FSubInterface declarations must only contain inputs Or outputs. Inputs AND outputs found in sub interface (%s)"), *InSubInterfaceName.ToString());
+			}
+		}
+
+		void FClassInterfaceDataBuilder::MergeSubInterfaceVariantData(const FClassInterfaceData& InSubInterfaceData)
+		{
+			// Add in any variant data gathered from the sub interface
+			TSortedMap<uint8, uint8> SubInterfaceVariantRemap;
+			for (int32 VGIndex = 0; VGIndex < InSubInterfaceData.Variants.Num(); VGIndex++)
+			{
+				const FVariantDescription& Variant = InSubInterfaceData.Variants[VGIndex];
+
+				// Check to see if variant already exists. 
+				int32 NewIndex = IndexOfVariant(Variant.VariantName);
+				if (INDEX_NONE == NewIndex)
+				{
+					NewIndex = ClassInterfaceData.Variants.Num();
+					checkf((NewIndex >= TNumericLimits<uint8>::Min()) && (NewIndex <= TNumericLimits<uint8>::Max()), TEXT("There are too many variants (%d). Last variant (%s)"), NewIndex, *Variant.VariantName.ToString());
+
+					ClassInterfaceData.Variants.Add(Variant);
+
+				}
+				else
+				{
+					// Merge a variant 
+					FVariantDescription& ExistingVariant = ClassInterfaceData.Variants[NewIndex];
+					checkf(ExistingVariant.DataTypes == Variant.DataTypes, TEXT("Cannot have inconsistent data types in variant (%s)"), *ExistingVariant.VariantName.ToString());
+				}
+
+				// hold onto a remaping between the two variant indices 
+				SubInterfaceVariantRemap.Add(static_cast<uint8>(VGIndex), static_cast<uint8>(NewIndex));
+			}
+
+			for (const TPair<FVertexKey, uint8>& KeyToVariant : InSubInterfaceData.VertexKeyToVariantIndex)
+			{
+				ClassInterfaceData.VertexKeyToVariantIndex.Add(KeyToVariant.Key, SubInterfaceVariantRemap[KeyToVariant.Value]);
+			}
+
+			ClassInterfaceData.VariantLiterals.Append(InSubInterfaceData.VariantLiterals);
+		}
+
+		uint8 FClassInterfaceDataBuilder::FindOrAddVariant(const FVariantVertex& InVariantVertex)
+		{
+			int32 Index = IndexOfVariant(InVariantVertex.VariantName);
+			if (INDEX_NONE == Index)
+			{
+				Index = ClassInterfaceData.Variants.Num();
+				ClassInterfaceData.Variants.Add(
+					FVariantDescription
+					{
+						.VariantName=InVariantVertex.VariantName,
+						.DataTypes=InVariantVertex.DataTypeNames
+					}
+				);
+			}
+
+			FVariantDescription& Variant = ClassInterfaceData.Variants[Index];
+
+			checkf(Algo::Compare(Variant.DataTypes, InVariantVertex.DataTypeNames), TEXT("Cannot have inconsistent data types in variant (%s) on vertex (%s)"), *InVariantVertex.VariantName.ToString(), *InVariantVertex.VertexName.ToString());
+			
+			checkf((Index >= TNumericLimits<uint8>::Min()) && (Index <= TNumericLimits<uint8>::Max()), TEXT("There are too many variants (%d). Last variant (%s)"), Index, *InVariantVertex.VariantName.ToString());
+
+			return  static_cast<uint8>(Index);
+		}
+
+		bool FClassInterfaceDataBuilder::ContainsAnySubInterfaces() const
+		{
+			return !(ClassInterfaceData.SubInterfaces.IsEmpty() &&
+					ClassInterfaceData.InputSubInterfaceLayouts.IsEmpty() &&
+					ClassInterfaceData.OutputSubInterfaceLayouts.IsEmpty());
+		}
+
+		int32 FClassInterfaceDataBuilder::IndexOfVariant(const FName& InName) const
+		{
+			if (InName.IsNone())
+			{
+				return INDEX_NONE;
+			}
+			return ClassInterfaceData.Variants.IndexOfByPredicate([&InName](const FVariantDescription& InVariant) { return InVariant.VariantName == InName; });
+		}
+	}
 }
 
 FString LexToString(Metasound::EVertexAccessType InAccessType)
diff --git a/Engine/Plugins/Runtime/Metasound/Source/MetasoundGraphCore/Public/MetasoundVertex.h b/Engine/Plugins/Runtime/Metasound/Source/MetasoundGraphCore/Public/MetasoundVertex.h
--- a/Engine/Plugins/Runtime/Metasound/Source/MetasoundGraphCore/Public/MetasoundVertex.h
+++ b/Engine/Plugins/Runtime/Metasound/Source/MetasoundGraphCore/Public/MetasoundVertex.h
@@ -3,6 +3,7 @@
 #pragma once
 
 #include "Containers/Array.h"
+#include "Math/NumericLimits.h"
 #include "MetasoundDataFactory.h"
 #include "MetasoundDataReference.h"
 #include "MetasoundEnvironment.h"
@@ -17,6 +18,23 @@ namespace Metasound
 	/** Name of a given vertex.  Only unique for a given node interface. */
 	using FVertexName = FName;
 
+	
+	/** Describes a variant on a ClassInterface */
+	struct UE_EXPERIMENTAL(5.8, "MetaSound Variants are experimental") FVariantDescription;
+	struct FVariantDescription
+	{
+		FName VariantName; 
+		TArray<FName> DataTypes; //< Array of supported data types.
+	};
+
+	/** Configures a variant to be a specific data type. */
+	struct UE_EXPERIMENTAL(5.8, "MetaSound Variants are experimental") FVariantConfiguration;
+	struct FVariantConfiguration
+	{
+		FName VariantName;
+		FName DataType; 
+	};
+
 	namespace VertexPrivate
 	{
 		// Forward declare for friendship
@@ -25,20 +43,65 @@ namespace Metasound
 
 		struct FPrivateAccessTag;
 
+		namespace SubInterface
+		{
+			// These are utilized to make sure we don't overflow uint16 capacity
+			// in case someone declares a ridiculously large interface.
+			static constexpr uint16 InvalidInstancePosition = TNumericLimits<uint16>::Max();
+			static constexpr uint16 MaxInstancePosition = TNumericLimits<uint16>::Max() - 1;
+			static constexpr uint16 MinInstancePosition = 0;
+		}
+
 		/** Contains a list of sub-interface spans where the span indices
 		 * refer to vertex indices of an array containing vertices. */
 		struct FSubInterfaceLayout
 		{
 			struct FInstance
 			{
-				int32 Begin = 0;
-				int32 End = 0; //< Inclusive. 
+				uint16 Begin = 0;
+				uint16 End = 0; //< Exclusive
 			};
 
 			FName SubInterfaceName;
 			TArray<FInstance> Instances;
 		};
 
+		/** Contains the location of an individual sub interface in a FClassInterface */
+		struct FSubInterfaceClassLayout
+		{
+			FName SubInterfaceName;
+			uint16 Begin = 0;
+			uint16 End = 0; //< Exclusive. 
+		};
+
+		/** Denotes whether a vertex belongs to the input or output interface */
+		enum class EVertexDirection : uint8
+		{
+			Input,
+			Output
+		};
+
+		/** A key for placing input and output vertices in a shared container. */
+		struct FVertexKey
+		{
+			FVertexName Name;
+			EVertexDirection Direction;
+
+			friend bool operator==(const FVertexKey& InLHS, const FVertexKey& InRHS)
+			{
+				return (InLHS.Name == InRHS.Name) && (InLHS.Direction == InRHS.Direction);
+			}
+
+			friend bool operator<(const FVertexKey& InLHS, const FVertexKey& InRHS)
+			{
+				if (InLHS.Name == InRHS.Name)
+				{
+					return InLHS.Direction < InRHS.Direction;
+				}
+				return InLHS.Name.FastLess(InRHS.Name);
+			}
+		};
+
 		/** Functor for finding vertices with equal names */
 		template<typename VertexType>
 		struct TEqualVertexName
@@ -67,7 +130,24 @@ namespace Metasound
 				return Num;
 			}
 		};
-	}
+	} // namespace VertexPrivate
+
+	/** Convenience for using a TSortedMap with FVertexName Key type.
+	 *
+	 * This template makes it convenient to create a TSortedMap with an FVertexName 
+	 * while also avoiding compilation errors incurred from using the FName default
+	 * "less than" operator in the TSortedMap implementation. 
+	 *
+	 * - FVertexName is an alias to FName. 
+	 * - TSortedMap<FName, ValueType> fails to compile since the "less than" operator
+	 *   specific implementation needs to be chosen (FastLess vs LexicalLess)
+	 * - Due to the template argument order of TSortedMap this also forces you to
+	 *   choose the allocator. 
+	 * - This is all a bit of an annoyance to do every time we use a TSortedMap 
+	 *   with an FVertexName as the key.
+	 */
+	template<typename ValueType, typename AllocatorType=FDefaultAllocator>
+	using TSortedVertexNameMap = TSortedMap<FVertexName, ValueType, AllocatorType, FNameFastLess>;
 
 	// Vertex metadata
 	struct FDataVertexMetadata
@@ -358,7 +438,7 @@ namespace Metasound
 		{
 			static_assert(
 				(std::is_constructible_v<VertexType, VertexTypes&&> && ...),
-				"Vertex types must be move constructible from the group's base type");
+				"Vertex types must be move constructible from the base type");
 			
 			// Reserve array to hold exact number of vertices to avoid
 			// over allocation.
@@ -368,7 +448,7 @@ namespace Metasound
 			(AddOrUpdateVertex(Forward<VertexTypes>(InVertices)), ...);
 		}
 
-		/** Add a vertex to the group. */
+		/** Add a vertex. */
 		void Add(const VertexType& InVertex)
 		{
 			AddOrUpdateVertex(VertexType(InVertex));
@@ -394,7 +474,7 @@ namespace Metasound
 			return (NumRemoved > 0);
 		}
 
-		/** Returns true if the group contains a vertex with a matching key. */
+		/** Returns true if this contains a vertex with a matching key. */
 		bool Contains(const FVertexName& InKey) const
 		{
 			return Vertices.ContainsByPredicate(FEqualVertexName(InKey));
@@ -461,13 +541,13 @@ namespace Metasound
 			return Vertices[InIndex];
 		}
 
-		/** Compare whether two vertex groups are equal. */
+		/** Compare whether two vertex interfaces are equal. */
 		friend bool operator==(const TVertexInterfaceImpl<VertexType>& InLHS, const TVertexInterfaceImpl<VertexType>& InRHS)
 		{
 			return InLHS.Vertices == InRHS.Vertices;
 		}
 
-		/** Compare whether two vertex groups are unequal. */
+		/** Compare whether two vertex interfaces are unequal. */
 		friend bool operator!=(const TVertexInterfaceImpl<VertexType>& InLHS, const TVertexInterfaceImpl<VertexType>& InRHS)
 		{
 			return !(InLHS == InRHS);
@@ -481,19 +561,6 @@ namespace Metasound
 	template<typename VertexType>
 	using TVertexInterfaceGroup UE_DEPRECATED(5.6, "Use TVertexInterfaceImpl instead") = TVertexInterfaceImpl<VertexType>;
 
-	/** Declare the beginning of a sub interface in a vertex interface declaration. */
-	struct UE_EXPERIMENTAL(5.6, "Sub interfaces are experimental") FBeginSubInterface;
-	struct FBeginSubInterface
-	{
-		FName Name; //< Name of the sub interface. 
-	};
-
-	/** Declare the ending of a sub interface in a vertex interface declaration. */
-	struct UE_EXPERIMENTAL(5.6, "Sub interfaces are experimental") FEndSubInterface;
-	struct FEndSubInterface
-	{
-	};
-
 	namespace VertexPrivate
 	{
 		/** Interface Declaration Builders create vertex interfaces from a template
@@ -506,9 +573,7 @@ namespace Metasound
 		 *  FInputVertexInterface Interface
 		 *  {
 		 *  	TInputDataVertex<float>(...),
-		 *  	FBeginSubInterface(...),
-		 *  		TInputDataVertex<int32>(...),
-		 *  	FEndSubInterface(...),
+	 	 * 		TInputDataVertex<int32>(...),
 		 *  	TInputDataVertex<float>(...)
 		 *  };
 		 *
@@ -516,58 +581,20 @@ namespace Metasound
 		 *  FInputVertexInterface Interface;
 		 * 
 		 * 	Interface.Add(TInputDataVertex<float>(...));
-		 * 	Interface.Add(FBeginSubInterface(...));
 	     * 	Interface.Add(TInputDataVertex<int32>(...));
-		 * 	Interface.Add(FEndSubInterface(...));
 		 * 	Interface.Add(TInputDataVertex<float>(...);
 		 *
-		 * 
-		 * An essential element of the builders is that they use member function 
-		 * overloading of Add(...) to handle the variety of different objects which
-		 * might be used to describe an interface. 
-		 *
-		 * The builders also size memory allocations exactly to remove any addition
+		 * The builders size memory allocations exactly to remove any addition
 		 * slack. In general, there can be many interfaces in memory and minimizing
 		 * their memory footprint is important. 
 		 */
 
-		/** Base class for an interface builder which supports sub interfaces. */
-		class FSubInterfaceDeclarationBuilder
-		{
-		public:
-			UE_API FSubInterfaceDeclarationBuilder(TArray<FSubInterfaceLayout>& OutSubInterfaceLayouts);
-			UE_API virtual ~FSubInterfaceDeclarationBuilder();
-
-			// Handler for the beginning of a sub interface. 
-UE_API PRAGMA_DISABLE_EXPERIMENTAL_WARNINGS
-			void Add(FBeginSubInterface&& InSubInterface);
-
-			// Handler for the ending of a sub interface. 
-			UE_API void Add(FEndSubInterface&& InSubInterface);
-UE_API PRAGMA_ENABLE_EXPERIMENTAL_WARNINGS
-
-			// Allocate memory for a known number of sub interfaces.
-			void ReserveSubInterfaces(int32 Num);
-
-		protected:
-			// Called when a vertex is added to adjust internal counters. 
-			UE_API void OnVertexAdded(const FVertexName& InVertexName);
-
-		private:
-
-			void PushSubInterfaceDeclaration(const FName& InName);
-			void PopSubInterfaceDeclaration();
-
-			int32 CurrentNumVertices = 0;
-			TArray<FSubInterfaceLayout>* SubInterfaces;
-			int32 CurrentSubInterfaceIndex = INDEX_NONE;
-		};
-
 		/** Interface builder for FInputVertexInterface declarations. */
-		class FInputVertexInterfaceDeclarationBuilder : public FSubInterfaceDeclarationBuilder
+		class FInputVertexInterfaceDeclarationBuilder
 		{
 		public:
-			UE_API FInputVertexInterfaceDeclarationBuilder(TArray<FInputDataVertex>& OutVertices, TArray<FSubInterfaceLayout>& OutInstances);
+			UE_API FInputVertexInterfaceDeclarationBuilder(TArray<FInputDataVertex>& OutVertices);
+
 			template<typename... ArgTypes>
 			void Build(ArgTypes&&... InArgs)
 			{
@@ -578,15 +605,6 @@ UE_API PRAGMA_ENABLE_EXPERIMENTAL_WARNINGS
 					Vertices.Reserve(NumVertices);
 				}
 
-PRAGMA_DISABLE_EXPERIMENTAL_WARNINGS
-				// Size sub interfaces exactly to avoid wasted memory in array allocations
-				constexpr uint32 NumSubInterfaces = TNumOfTypeInPack<FBeginSubInterface>::Get<ArgTypes...>();
-PRAGMA_ENABLE_EXPERIMENTAL_WARNINGS
-				if constexpr (NumSubInterfaces > 0)
-				{
-					ReserveSubInterfaces(NumSubInterfaces);
-				}
-
 				// Add all the elements of the vertex interface. This uses a fold expression
 				// to call Add(...) on each input. The various overloads of the Add(...) method
 				// then assemble the appropriate structures.
@@ -595,17 +613,16 @@ PRAGMA_ENABLE_EXPERIMENTAL_WARNINGS
 
 		private:
 
-			using FSubInterfaceDeclarationBuilder::Add; // Unhide Add(...) function overloads in FSubInterfaceDeclarationBuilder
 			UE_API void Add(FInputDataVertex&& InVertex);
 
 			TArray<FInputDataVertex>& Vertices;
 		};
 
 		/** Interface builder for FOutputVertexInterface declarations. */
-		class FOutputVertexInterfaceDeclarationBuilder : public FSubInterfaceDeclarationBuilder
+		class FOutputVertexInterfaceDeclarationBuilder 
 		{
 		public:
-			UE_API FOutputVertexInterfaceDeclarationBuilder(TArray<FOutputDataVertex>& OutVertices, TArray<FSubInterfaceLayout>& OutInstances);
+			UE_API FOutputVertexInterfaceDeclarationBuilder(TArray<FOutputDataVertex>& OutVertices);
 
 			template<typename... ArgTypes>
 			void Build(ArgTypes&&... InArgs)
@@ -617,15 +634,6 @@ PRAGMA_ENABLE_EXPERIMENTAL_WARNINGS
 					Vertices.Reserve(NumVertices);
 				}
 
-PRAGMA_DISABLE_EXPERIMENTAL_WARNINGS
-				// Size sub interfaces exactly to avoid wasted memory in array allocations
-				constexpr uint32 NumSubInterfaces = TNumOfTypeInPack<FBeginSubInterface>::Get<ArgTypes...>();
-PRAGMA_ENABLE_EXPERIMENTAL_WARNINGS
-				if constexpr (NumSubInterfaces > 0)
-				{
-					ReserveSubInterfaces(NumSubInterfaces);
-				}
-
 				// Add all the elements of the vertex interface. This uses a fold expression
 				// to call Add(...) on each input. The various overloads of the Add(...) method
 				// then assemble the appropriate structures.
@@ -633,7 +641,7 @@ PRAGMA_ENABLE_EXPERIMENTAL_WARNINGS
 			}
 
 		private:
-			using FSubInterfaceDeclarationBuilder::Add; // Unhide Add(...) function overloads in FSubInterfaceDeclarationBuilder
+
 			UE_API void Add(FOutputDataVertex&& InVertex);
 
 			TArray<FOutputDataVertex>& Vertices;
@@ -678,16 +686,14 @@ PRAGMA_ENABLE_EXPERIMENTAL_WARNINGS
 		 *  FInputVertexInterface Interface
 		 *  {
 		 *  	TInputDataVertex<float>(...),
-		 *  	FBeginSubInterface(...),
-		 *  		TInputDataVertex<int32>(...),
-		 *  	FEndSubInterface(...),
+	 	 * 		TInputDataVertex<int32>(...),
 		 *  	TInputDataVertex<float>(...)
 		 *  };
 		 */
 		template<typename... ArgTypes>
 		explicit FInputVertexInterface(ArgTypes && ... InArgs)
 		{
-			VertexPrivate::FInputVertexInterfaceDeclarationBuilder InterfaceBuilder(Vertices, SubInterfaces);
+			VertexPrivate::FInputVertexInterfaceDeclarationBuilder InterfaceBuilder(Vertices);
 			InterfaceBuilder.Build(Forward<ArgTypes>(InArgs)...);
 		}
 
@@ -720,9 +726,7 @@ PRAGMA_ENABLE_EXPERIMENTAL_WARNINGS
 		 *  FOutputVertexInterface Interface
 		 *  {
 		 *  	TOutputDataVertex<float>(...),
-		 *  	FBeginSubInterface(...),
-		 *  		TOutputDataVertex<int32>(...),
-		 *  	FEndSubInterface(...),
+		 *  	TOutputDataVertex<int32>(...),
 		 *  	TOutputDataVertex<float>(...)
 		 *  };
 		 */
@@ -730,7 +734,7 @@ PRAGMA_ENABLE_EXPERIMENTAL_WARNINGS
 		explicit FOutputVertexInterface(ArgTypes && ... InArgs)
 		: FOutputVertexInterface()
 		{
-			VertexPrivate::FOutputVertexInterfaceDeclarationBuilder InterfaceBuilder(Vertices, SubInterfaces);
+			VertexPrivate::FOutputVertexInterfaceDeclarationBuilder InterfaceBuilder(Vertices);
 			InterfaceBuilder.Build(Forward<ArgTypes>(InArgs)...);
 		}
 
@@ -836,7 +840,7 @@ PRAGMA_ENABLE_EXPERIMENTAL_WARNINGS
 			FEnvironmentVertexInterface EnvironmentInterface;
 	};
 
-	/** A description of a sub interface which is used when declaring a FClassVertexInterface */
+	/** A description of a sub interface which is used when declaring a FClassInterface */
 	struct UE_EXPERIMENTAL(5.6, "Sub interfaces are experimental") FSubInterfaceDescription;
 	struct FSubInterfaceDescription
 	{
@@ -847,7 +851,7 @@ PRAGMA_ENABLE_EXPERIMENTAL_WARNINGS
 	};
 
 	/** A sub interface configuration is used to create a FVertexInterface from 
-	 * a FClassVertexInterface. */
+	 * a FClassInterface. */
 	struct UE_EXPERIMENTAL(5.6, "Sub interfaces are experimental") FSubInterfaceConfiguration;
 	struct FSubInterfaceConfiguration
 	{
@@ -855,34 +859,392 @@ PRAGMA_ENABLE_EXPERIMENTAL_WARNINGS
 		uint32 Num;
 	};
 
-	/** FClassVertexInterface describes the interface of a node class. It is an 
-	 * immutable factory for FVertexInterfaces. */
-	class UE_EXPERIMENTAL(5.6, "FClassVertexInterfaces are experimental") FClassVertexInterface;
 
-	class FClassVertexInterface
+	/** The data type name used for a variant vertex which has not been given
+	 * a concrete MetaSound data type. */
+	extern UE_API const FLazyName GVariantDataTypeName;
+
+	/** A variant vertex represents a vertex that has a variant data type. Vertices
+	 * are associated with a variant by using the VariantName. Multiple vertices
+	 * may share the same variant. */
+	class UE_EXPERIMENTAL(5.8, "MetaSound Variants are experimental") FVariantVertex;
+	class FVariantVertex
 	{
+	protected:
+		
+		// Derived classes using this constructor must remember to initialize the DataTypeNames array
+		UE_API FVariantVertex(const FVertexName& InVertexName, const FVertexName& InVariantName, const FDataVertexMetadata& InMetadata, EVertexAccessType InAccessType);
+
 	public:
+
+		UE_API FVariantVertex(const FVertexName& InVertexName, const FName& InVariantName, const TArray<FName>& InDataTypeNames, const FDataVertexMetadata& InMetadata, EVertexAccessType InAccessType);
+		 
+		/** Name of vertex. */
+		FVertexName VertexName;
+
+		/** Name of variant associated with the vertex. */
+		FName VariantName;
+
+		/** Possible data types supported by the variant. */
+		TArray<FName> DataTypeNames;
+
+#if WITH_EDITORONLY_DATA
+		/** Metadata associated with vertex. */
+		FDataVertexMetadata Metadata;
+#endif // WITH_EDITORONLY_DATA
+
+		/** Access type of the vertex. */
+		EVertexAccessType AccessType;
+	};
+
+	/** An input variant vertex represents an input vertex that has a variant data 
+	 * type. Vertices are associated with a variant by using the VariantName. 
+	 * Multiple vertices may share the same variant. */
+	class UE_EXPERIMENTAL(5.8, "MetaSound Variants are experimental") FInputVariantVertex;
 PRAGMA_DISABLE_EXPERIMENTAL_WARNINGS
-		FClassVertexInterface() = default;
-		~FClassVertexInterface() = default;
-		UE_API FClassVertexInterface(FVertexInterface InInterface);
-		UE_API FClassVertexInterface(FInputVertexInterface InInputs, FOutputVertexInterface InOutputs, FEnvironmentVertexInterface InEnvironment={});
-		UE_API FClassVertexInterface(TArray<FSubInterfaceDescription> InSubInterfactDescriptions, FInputVertexInterface InInputs, FOutputVertexInterface InOutputs, FEnvironmentVertexInterface InEnvironment={});
+	class FInputVariantVertex : public FVariantVertex
+	{
+	protected:
 
-		/** Create a FVertexInterface with the given configuration. */
-		UE_API FVertexInterface CreateVertexInterface(TConstArrayView<FSubInterfaceConfiguration> InSubInterfaceConfigurations) const;
+		// Derived classes using this constructor must remember to initialize the DataTypeNames array and Literals array.
+		UE_API FInputVariantVertex(const FVertexName& InVertexName, const FName& InVariantName, const FDataVertexMetadata& InMetadata, EVertexAccessType InAccessType);
+
+	public:
+		UE_API FInputVariantVertex(const FVertexName& InVertexName, const FName& InVariantName, const TArray<FName>& InDataTypeNames, const FDataVertexMetadata& InMetadata, EVertexAccessType InAccessType, const TArray<FLiteral>& InLiterals);
+
+		TArray<FLiteral> Literals;
+
+		friend bool METASOUNDGRAPHCORE_API operator==(const FInputVariantVertex& InLHS, const FInputVariantVertex& InRHS);
+		friend bool METASOUNDGRAPHCORE_API operator!=(const FInputVariantVertex& InLHS, const FInputVariantVertex& InRHS);
+		friend bool METASOUNDGRAPHCORE_API operator<(const FInputVariantVertex& InLHS, const FInputVariantVertex& InRHS);
+	};
+
+	/** An output variant vertex represents an output vertex that has a variant data 
+	 * type. Vertices are associated with a variant by using the VariantName. 
+	 * Multiple vertices may share the same variant. */
+	class UE_EXPERIMENTAL(5.8, "MetaSound Variants are experimental") FOutputVariantVertex;
+	class FOutputVariantVertex : public FVariantVertex
+	{
+	public:
+		using FVariantVertex::FVariantVertex;
+
+		friend bool METASOUNDGRAPHCORE_API operator==(const FOutputVariantVertex& InLHS, const FOutputVariantVertex& InRHS);
+		friend bool METASOUNDGRAPHCORE_API operator!=(const FOutputVariantVertex& InLHS, const FOutputVariantVertex& InRHS);
+		friend bool METASOUNDGRAPHCORE_API operator<(const FOutputVariantVertex& InLHS, const FOutputVariantVertex& InRHS);
+	};
 PRAGMA_ENABLE_EXPERIMENTAL_WARNINGS
 
-	private:
+	namespace VertexPrivate
+	{
+		template<typename ... DataTypes>
+		void BuildArrayOfTypeNames(TArray<FName>& InOutArray)
+		{
+			checkf(InOutArray.Num() == 0, TEXT("Initial array must be empty"));
+
+			InOutArray.Reserve(sizeof...(DataTypes));
+			( InOutArray.Add(GetMetasoundDataTypeName<DataTypes>()), ... );
+		}
+
+		template<typename ... ArgTypes>
+		void BuildArrayOfLiterals(int32 InNumDataTypes, TArray<FLiteral>& InOutArray, ArgTypes&&...Args)
+		{
+			constexpr int32 NumArgs = sizeof...(ArgTypes);
+
+			checkf(InOutArray.Num() == 0, TEXT("Initial array must be empty"));
+			checkf(InNumDataTypes >= NumArgs, TEXT("The number of default arguments (%d) must be less than or equal to the number of data types (%d)"), NumArgs, InNumDataTypes);
+
+			// Create literals for each supplied arg. 
+			InOutArray.Reserve(InNumDataTypes);
+			( InOutArray.Add(FLiteral{Args}), ... );
+
+			// If there were fewer args than data types, use a default FLiteral for remaining data types. 
+			int32 NumRemaining = InNumDataTypes - NumArgs;
+			if (NumRemaining > 0)
+			{
+				InOutArray.AddDefaulted(NumRemaining);
+			}
+		}
+	}
 
+	template<typename ... DataTypes>
+	class UE_EXPERIMENTAL(5.8, "MetaSound Variants are experimental") TInputVariantVertex;
+
+	/** An input variant vertex represents an input vertex that has a variant data 
+	 * type. Vertices are associated with a variant by using the VariantName. 
+	 * Multiple vertices may share the same variant. The supported data types of
+	 * the variant are declared by the template arguments to the vertex. 
+	 *
+	 * TInputVariantVertex<float, int32>(...) will produce a variant which supports
+	 * float or int32.
+	 */
+	template<typename ... DataTypes>
+	class TInputVariantVertex : public FInputVariantVertex
+	{
+	public:
+		TInputVariantVertex() = default;
+
+		template<typename... ArgTypes>
+		FORCENOINLINE TInputVariantVertex(const FVertexName& InVertexName, const FVertexName& InVariantName, const FDataVertexMetadata& InMetadata, ArgTypes&&... Args)
+		: FInputVariantVertex(InVertexName, InVariantName, InMetadata, EVertexAccessType::Reference)
+		{
+			using namespace VertexPrivate;
+
+			constexpr int32 NumDataTypes = sizeof...(DataTypes);
+			BuildArrayOfTypeNames<DataTypes...>(DataTypeNames);
+			BuildArrayOfLiterals(NumDataTypes, Literals, Forward<ArgTypes>(Args)...);
+		}
+
+		/** Vertex w/ variant name matching vertex name. */
+		template<typename... ArgTypes>
+		FORCENOINLINE TInputVariantVertex(const FVertexName& InVertexName, const FDataVertexMetadata& InMetadata, ArgTypes&&... Args)
+		: TInputVariantVertex(InVertexName, InVertexName, InMetadata, Forward<ArgTypes>(Args)...)
+		{
+		}
+	};
+
+	template<typename ... DataTypes>
+	class UE_EXPERIMENTAL(5.8, "MetaSound Variants are experimental") TOutputVariantVertex;
+
+
+	/** An output variant vertex represents an output vertex that has a variant data 
+	 * type. Vertices are associated with a variant by using the VariantName. 
+	 * Multiple vertices may share the same variant. The supported data types of
+	 * the variant are declared by the template arguments to the vertex. 
+	 *
+	 * TOutputVariantVertex<float, int32>(...) will produce a variant which supports
+	 * float or int32.
+	 */
+	template<typename ... DataTypes>
+	class TOutputVariantVertex : public FOutputVariantVertex
+	{
+	public:
+		TOutputVariantVertex() = default;
+
+		FORCENOINLINE TOutputVariantVertex(const FVertexName& InVertexName, const FVertexName& InVariantName, const FDataVertexMetadata& InMetadata)
+		: FOutputVariantVertex(InVertexName, InVariantName, InMetadata, EVertexAccessType::Reference)
+		{
+			using namespace VertexPrivate;
+
+			constexpr int32 NumDataTypes = sizeof...(DataTypes);
+			BuildArrayOfTypeNames<DataTypes...>(DataTypeNames);
+		}
+
+		/** Vertex with variant name matching vertex name. */
+		FORCENOINLINE TOutputVariantVertex(const FVertexName& InVertexName, const FDataVertexMetadata& InMetadata)
+		: TOutputVariantVertex(InVertexName, InVertexName, InMetadata)
+		{
+		}
+	};
+
+	/** FClassInterface describes the interface of a node class. It is an 
+	 * immutable factory for FVertexInterfaces. */
+	class UE_EXPERIMENTAL(5.6, "FClassInterfaces are experimental") FClassInterface;
+
+	class FSubInterface;
+
+	namespace VertexPrivate
+	{
+		// Data structure used to efficiently store class declarations w/o introducing
+		// unneeded allocations. FClassInterfaceData generally only lives for the
+		// duration of a function call. 
+		struct FClassInterfaceData
+		{
+PRAGMA_DISABLE_EXPERIMENTAL_WARNINGS
+			TArray<FSubInterfaceDescription, TInlineAllocator<4>> SubInterfaces;
+			TArray<FSubInterfaceClassLayout, TInlineAllocator<2>> InputSubInterfaceLayouts;
+			TArray<FSubInterfaceClassLayout, TInlineAllocator<2>> OutputSubInterfaceLayouts;
+
+			TSortedMap<FVertexKey, uint8, TInlineAllocator<4>> VertexKeyToVariantIndex;
+			TArray<FVariantDescription, TInlineAllocator<4>> Variants;
+			TSortedVertexNameMap<TArray<FLiteral, TInlineAllocator<4>>, TInlineAllocator<4>> VariantLiterals;
+PRAGMA_ENABLE_EXPERIMENTAL_WARNINGS
+
+			TArray<FInputDataVertex, TInlineAllocator<8>> Inputs;
+			TArray<FOutputDataVertex, TInlineAllocator<8>> Outputs;
+			TArray<FEnvironmentVertex, TInlineAllocator<4>> EnvironmentVariables;
+		};
+
+PRAGMA_DISABLE_EXPERIMENTAL_WARNINGS
+		// Stores information to describe variants on a class interface
+		struct FVariantSchema
+		{
+			TSortedMap<VertexPrivate::FVertexKey, uint8> VertexKeyToVariantIndex;
+			TArray<FVariantDescription> Variants;
+			TSortedVertexNameMap<TArray<FLiteral, TInlineAllocator<2>>> VariantLiterals;
+
+			UE_API static FVariantSchema FromClassInterfaceData(const VertexPrivate::FClassInterfaceData& InClassInterfaceData);
+		};
+
+		// Stores information to describe sub interfaces on a class interface. 
+		struct FSubInterfaceSchema
+		{
+			TArray<FSubInterfaceDescription> SubInterfaces;
+			TArray<VertexPrivate::FSubInterfaceClassLayout> InputSubInterfaceLayouts;
+			TArray<VertexPrivate::FSubInterfaceClassLayout> OutputSubInterfaceLayouts;
+
+			UE_API static FSubInterfaceSchema FromClassInterfaceData(const VertexPrivate::FClassInterfaceData& InClassInterfaceData);
+		};
+PRAGMA_ENABLE_EXPERIMENTAL_WARNINGS
+
+		/** The FClassInterfaceDataBuilder assembles declared inputs, outputs,
+		 * variants and subinterfaces into a FClassInterfaceData object. This is
+		 * largely to simply the task of node authors when they declare a nodes
+		 * interface. Instead of building an interface sequentially, they can
+		 * declare the interface and allow the builder to build the correct object.
+		 *
+		 * e.g:
+		 *
+		 * FClassInterface ClassInterface = CreateClassInterface(
+		 * 	TInputDataVertex<float>(...)
+		 * 	FSubInterface{
+		 * 		"MySub",
+		 * 		TInputVariantVertex<float, int32>(...)
+		 * 	}
+		 * 	TOutDataVertex<FAudioBuffer>(...)
+		 * );
+		 */
+		class FClassInterfaceDataBuilder
+		{
+		public:
+			template<typename ... Args>
+			FClassInterfaceDataBuilder& Build(Args&& ... InArgs)
+			{
+				(this->Add(&InArgs), ...);
+				return *this;
+			}
+
+			UE_API const FClassInterfaceData& GetClassInterfaceData() const;
+			
+		private:
+			UE_API void Add(const FInputDataVertex* InInputVertex);
+			UE_API void Add(const FOutputDataVertex* InOutputVertex);
+			UE_API void Add(const FEnvironmentVertex* InEnvironmentVertex);
+
+PRAGMA_DISABLE_EXPERIMENTAL_WARNINGS
+			UE_API void Add(const FSubInterfaceDescription* InSubInterfaceDescription);
+			UE_API void Add(const FSubInterface* InSubInterface);
+			UE_API void Add(const FInputVariantVertex* InInputVariantVertex);
+			UE_API void Add(const FOutputVariantVertex* InOutputVariantVertex);
+
+			bool ContainsAnySubInterfaces() const;
+			void AddSubInterfaceVertices(const FName& InSubInterfaceName, const FClassInterfaceData& InSubInterfaceData);
+			void MergeSubInterfaceVariantData(const FClassInterfaceData& InSubInterfaceData);
+
+			uint8 FindOrAddVariant(const FVariantVertex& InVariantVertex);
+			int32 IndexOfVariant(const FName& InName) const;
+PRAGMA_ENABLE_EXPERIMENTAL_WARNINGS
+
+			FClassInterfaceData ClassInterfaceData;
+		};
+
+		/** FBuildableInterface supports composition when declaring FSubInterfaces */
+		class FBuildableInterface
+		{
+		public:
+			FBuildableInterface() = default;
+			FBuildableInterface(const FBuildableInterface&) = default;
+			FBuildableInterface(FBuildableInterface&&) = default;
+
+			template<class FirstArgType, class... AdditionalArgTypes>
+			FBuildableInterface(FirstArgType&& InFirstArg, AdditionalArgTypes&&... InAdditionalArgs)
+			{
+				Builder.Build(Forward<FirstArgType>(InFirstArg), Forward<AdditionalArgTypes>(InAdditionalArgs)...);
+			}
+
+			UE_INTERNAL
+			UE_API const FClassInterfaceDataBuilder& GetClassInterfaceDataBuilder() const;
+
+		private:
+			FClassInterfaceDataBuilder Builder;
+		};
+	} // namespace VertexPrivate;
+
+	
+	class UE_EXPERIMENTAL(5.8, "MetaSound SubInterfaces are experimental") FSubInterface;
 PRAGMA_DISABLE_EXPERIMENTAL_WARNINGS
-		TArray<FSubInterfaceDescription> SubInterfaces;
+	class FSubInterface : public VertexPrivate::FBuildableInterface
+	{
+	public:
+		template<typename ... ArgTypes>
+		FSubInterface(FLazyName InName, ArgTypes&&... InArgs)
+		: FBuildableInterface(Forward<ArgTypes>(InArgs)...)
+		, SubInterfaceName(InName)
+		{
+		}
+
+		UE_API const FName GetName() const;
+	private:
+
+		FName SubInterfaceName;
+	};
 PRAGMA_ENABLE_EXPERIMENTAL_WARNINGS
-		FInputVertexInterface Inputs;
-		FOutputVertexInterface Outputs;
-		FEnvironmentVertexInterface Environment;
+
+	class UE_EXPERIMENTAL(5.8, "MetaSound ClassInterfaces are experimental") FClassInterface;
+PRAGMA_DISABLE_EXPERIMENTAL_WARNINGS
+	/** An FClassInterface defines the class interface of a MetaSound C++ node class. This
+	 * interface may contain basic vertices, but also subinterfaces and variants. 
+	 *
+	 * The FClassInterface is used to create the FVertexInterface for individual 
+	 * node instances. It does so by combining the FClassInterface definition with
+	 * subinterface and variant configurations. */
+	class FClassInterface
+	{
+	public:
+		FClassInterface() = default;
+		~FClassInterface() = default;
+
+		UE_INTERNAL
+		UE_API FClassInterface(const VertexPrivate::FClassInterfaceData& InClassInterfaceData);
+
+		/** Create a FVertexInterface with the given configuration. */
+		UE_API FVertexInterface CreateVertexInterface(TConstArrayView<FSubInterfaceConfiguration> InSubInterfaceConfigurations, TConstArrayView<FVariantConfiguration> InVariantConfigurations) const;
+
+		/** Return any sub interfaces on this class interface. */
+		UE_API TConstArrayView<FSubInterfaceDescription> GetSubInterfaceDescriptions() const;
+
+		/** Return any variants on this class interface. */
+		UE_API TConstArrayView<FVariantDescription> GetVariantDescriptions() const;
+
+		/** Return a variant description for a given input vertex if the vertex
+		 * is associated with a variant. */
+		UE_API const FVariantDescription* FindVariantDescriptionForInput(const FVertexName& InVertexName) const;
+
+		/** Return a variant description for a given output vertex if the vertex
+		 * is associated with a variant. */
+		UE_API const FVariantDescription* FindVariantDescriptionForOutput(const FVertexName& InVertexName) const;
+
+	private:
+		
+		// Many nodes do not need this extra configurable information. Save on memory
+		// storage by offloading this extra functionality behind a unique ptr. 
+		TUniquePtr<VertexPrivate::FVariantSchema> VariantSchema;
+		TUniquePtr<VertexPrivate::FSubInterfaceSchema> SubInterfaceSchema;
+
+		TArray<FInputDataVertex, TInlineAllocator<1>> Inputs;
+		TArray<FOutputDataVertex, TInlineAllocator<1>> Outputs;
+		TArray<FEnvironmentVertex> EnvironmentVariables;
 	};
 
+	/** Create a FClassInterface by declaring the vertices, subinterfaces and variants.
+	 * e.g:
+	 *
+	 * FClassInterface ClassInterface = CreateClassInterface(
+	 * 	TInputDataVertex<float>(...)
+	 * 	FSubInterface{
+	 * 		"MySub",
+	 * 		TInputVariantVertex<float, int32>(...)
+	 * 	}
+	 * 	TOutDataVertex<FAudioBuffer>(...)
+	 * );
+	 */
+	template<typename ... ArgTypes>
+	UE_EXPERIMENTAL(5.7, "FClassInterface is experimental")
+	FClassInterface CreateClassInterface(ArgTypes&& ... InArgs)
+	{
+		return FClassInterface{VertexPrivate::FClassInterfaceDataBuilder().Build(Forward<ArgTypes>(InArgs)...).GetClassInterfaceData()};
+	}
+PRAGMA_ENABLE_EXPERIMENTAL_WARNINGS
+
 	/**
 	 * This struct is used to pass in any arguments required for constructing a single node instance.
 	 * because of this, all FNode implementations have to implement a constructor that takes an FNodeInitData instance.
diff --git a/Engine/Plugins/Runtime/Metasound/Source/MetasoundGraphCore/Public/MetasoundVertexData.h b/Engine/Plugins/Runtime/Metasound/Source/MetasoundGraphCore/Public/MetasoundVertexData.h
--- a/Engine/Plugins/Runtime/Metasound/Source/MetasoundGraphCore/Public/MetasoundVertexData.h
+++ b/Engine/Plugins/Runtime/Metasound/Source/MetasoundGraphCore/Public/MetasoundVertexData.h
@@ -482,23 +482,6 @@ namespace Metasound
 		};
 	}
 
-	/** Convenience for using a TSortedMap with FVertexName Key type.
-	 *
-	 * This template makes it convenient to create a TSortedMap with an FVertexName 
-	 * while also avoiding compilation errors incurred from using the FName default
-	 * "less than" operator in the TSortedMap implementation. 
-	 *
-	 * - FVertexName is an alias to FName. 
-	 * - TSortedMap<FName, ValueType> fails to compile since the "less than" operator
-	 *   specific implementation needs to be chosen (FastLess vs LexicalLess)
-	 * - Due to the template argument order of TSortedMap this also forces you to
-	 *   choose the allocator. 
-	 * - This is all a bit of an annoyance to do every time we use a TSortedMap 
-	 *   with an FVertexName as the key.
-	 */
-	template<typename ValueType>
-	using TSortedVertexNameMap = TSortedMap<FVertexName, ValueType, FDefaultAllocator, FNameFastLess>;
-
 	/** MetaSound Binding 
 	 *
 	 * MetaSound IOperators read and write to shared data. For example: a TDataWriteReference<float> 
