Remove ensures hit when reloading the MetaSound Editor due to the builder going stale on rebuild
#rnx
#Synced-CL 50733540

diff --git a/Engine/Plugins/Runtime/Metasound/Source/MetasoundEditor/Private/MetasoundEditor.cpp b/Engine/Plugins/Runtime/Metasound/Source/MetasoundEditor/Private/MetasoundEditor.cpp
--- a/Engine/Plugins/Runtime/Metasound/Source/MetasoundEditor/Private/MetasoundEditor.cpp
+++ b/Engine/Plugins/Runtime/Metasound/Source/MetasoundEditor/Private/MetasoundEditor.cpp
@@ -3208,7 +3208,7 @@ namespace Metasound
 		{
 			if (Builder.IsValid())
 			{
-				if (const FMetaSoundFrontendDocumentBuilder& DocBuilder = Builder->GetBuilder(); ensure(DocBuilder.IsValid()))
+				if (const FMetaSoundFrontendDocumentBuilder& DocBuilder = Builder->GetBuilder(); DocBuilder.IsValid())
 				{
 					TConstStructView<FMetaSoundFrontendDocumentConfiguration> Config = DocBuilder.GetConstDocumentConfiguration();
 					return Config.IsValid() && !Config.GetPtr()->GetEditorOptions().bGraphEditorVisible;
@@ -3312,7 +3312,7 @@ namespace Metasound
 							.Text(LOCTEXT("MetasoundConvertFromConfigDialogAccept", "Accept"))
 							.OnClicked_Lambda([this, &DialogWindow]()
 							{
-								if (FMetaSoundFrontendDocumentBuilder& DocBuilder = Builder->GetBuilder(); ensure(DocBuilder.IsValid()))
+								if (FMetaSoundFrontendDocumentBuilder& DocBuilder = Builder->GetBuilder(); DocBuilder.IsValid())
 								{
 									const UScriptStruct* ScriptStruct = DocBuilder.GetConstDocumentConfiguration().GetScriptStruct();
 									const FText DisplayName = ScriptStruct ? ScriptStruct->GetDisplayNameText() : LOCTEXT("EmptyConfigName", "Unset");
