Misc Fix-ups from DocConfig feedback
- Fix DocConfig AssetAction tool tips to display correct class name
- Export ctor/dtor for FDocumentConfigurator/external plugin use
- Update MetaSoundDetailCustomizations to use configuration edit properties instead of graph editability if configuration specified (eg. interface, page customizations)

#rnx

#Synced-CL 50733540

diff --git a/Engine/Plugins/Runtime/Metasound/Source/MetasoundEditor/Private/MetasoundAssetDefinitions.cpp b/Engine/Plugins/Runtime/Metasound/Source/MetasoundEditor/Private/MetasoundAssetDefinitions.cpp
--- a/Engine/Plugins/Runtime/Metasound/Source/MetasoundEditor/Private/MetasoundAssetDefinitions.cpp
+++ b/Engine/Plugins/Runtime/Metasound/Source/MetasoundEditor/Private/MetasoundAssetDefinitions.cpp
@@ -398,13 +398,14 @@ namespace Metasound::Editor
 			FToolMenuSection* Section = CastChecked<UAssetDefinition_SoundAssetBase>(AssetDefinition)->FindSoundContextMenuSection(MetaSoundClassName);
 			check(Section);
 
-			Section->AddDynamicEntry(NAME_None, FNewToolMenuSectionDelegate::CreateLambda([CreateAction = MoveTemp(CreateFromConfigAction), MetaSoundClassName, ConfigName = ConfigDisplayName, ClassName, AssetIcon](FToolMenuSection& InSection)
+			const FText MetaSoundClassDisplayName = MetaSoundClass.GetDisplayNameText();
+			Section->AddDynamicEntry(NAME_None, FNewToolMenuSectionDelegate::CreateLambda([CreateAction = MoveTemp(CreateFromConfigAction), MetaSoundClassDisplayName, MetaSoundClassName, ConfigName = ConfigDisplayName, ClassName, AssetIcon](FToolMenuSection& InSection)
 			{
 				using namespace Metasound::Editor;
 
 				{
 					const TAttribute<FText> Label = FText::Format(LOCTEXT("MetaSound_CreateFromPreset", "Create {0}..."), ConfigName);
-					const TAttribute<FText> ToolTip = FText::Format(LOCTEXT("MetaSound_CreateFromPresetToolTipFormat", "Creates a {0} Preset using the selected MetaSound as a parent."), ClassName);
+					const TAttribute<FText> ToolTip = FText::Format(LOCTEXT("MetaSound_CreateFromPresetToolTipFormat", "Creates a {0} Preset using the selected MetaSound as a parent."), MetaSoundClassDisplayName);
 
 					const FName MenuEntryName = *FString::Format(TEXT("{0}_CreateFromConfig_{1}"), { MetaSoundClassName.ToString(), ConfigName.ToString()});
 					InSection.AddMenuEntry(MenuEntryName, Label, ToolTip, AssetIcon, CreateAction);
@@ -530,8 +531,9 @@ namespace Metasound::Editor
 
 			if (Section)
 			{
+				const FText MetaSoundClassDisplayName = MetaSoundClass.GetDisplayNameText();
 				const FText ConfigToolTip = Config.GetPtr()->GetEditorOptions().ToolTip;
-				auto AddEntryLambda = [CreateAction = MoveTemp(CreateFromConfigAction), ScriptStruct, ConfigToolTip, ClassName, AssetIcon](FToolMenuSection& InSection)
+				auto AddEntryLambda = [CreateAction = MoveTemp(CreateFromConfigAction), ScriptStruct, ConfigToolTip, ClassName, MetaSoundClassDisplayName, AssetIcon](FToolMenuSection& InSection)
 				{
 					using namespace Metasound::Editor;
 			
@@ -539,7 +541,7 @@ namespace Metasound::Editor
 					const TAttribute<FText> Label = FText::Format(LOCTEXT("MetaSound_CreateFromConfig", "Create {0}..."), ConfigName);
 					const TAttribute<FText> ToolTip = FText::Format(
 						LOCTEXT("MetaSound_CreateFromConfigToolTipFormat", "Creates a {0} with {1} configuration, initialized using the selected assets. {2}"),
-						ClassName, ConfigName, ConfigToolTip);
+						MetaSoundClassDisplayName, ConfigName, ConfigToolTip);
 			
 					const FName MenuEntryName = *FString::Format(TEXT("MetaSound_CreateFromConfig_{0}"), { ConfigName.ToString() });
 					InSection.AddMenuEntry(MenuEntryName, Label, ToolTip, AssetIcon, CreateAction);
diff --git a/Engine/Plugins/Runtime/Metasound/Source/MetasoundEditor/Private/MetasoundDetailCustomization.cpp b/Engine/Plugins/Runtime/Metasound/Source/MetasoundEditor/Private/MetasoundDetailCustomization.cpp
--- a/Engine/Plugins/Runtime/Metasound/Source/MetasoundEditor/Private/MetasoundDetailCustomization.cpp
+++ b/Engine/Plugins/Runtime/Metasound/Source/MetasoundEditor/Private/MetasoundDetailCustomization.cpp
@@ -89,6 +89,11 @@ namespace Metasound::Editor
 	}
 
 	bool FMetaSoundDetailCustomizationBase::IsGraphEditable() const
+	{
+		return IsEditable();
+	}
+
+	bool FMetaSoundDetailCustomizationBase::IsEditable() const
 	{
 		using namespace Engine;
 
@@ -280,7 +285,7 @@ namespace Metasound::Editor
 					}
 
 					TSharedRef<SWidget> OutputFormatValueWidget = OutputFormat->CreatePropertyValueWidget();
-					OutputFormatValueWidget->SetEnabled(TAttribute<bool>::Create(TAttribute<bool>::FGetter::CreateSP(this, &FMetaSoundDetailCustomizationBase::IsGraphEditable)));
+					OutputFormatValueWidget->SetEnabled(TAttribute<bool>::Create(TAttribute<bool>::FGetter::CreateSP(this, &FMetasoundDetailCustomization::IsEditable)));
 
 					static const FText OutputFormatName = LOCTEXT("MetasoundOutputFormatPropertyName", "Output Format");
 					GeneralCategoryBuilder.AddCustomRow(OutputFormatName)
@@ -419,6 +424,25 @@ namespace Metasound::Editor
 	{
 	}
 
+	bool FMetasoundPagesDetailCustomization::IsEditable() const
+	{
+		using namespace Metasound::Frontend;
+
+		if (Builder.IsValid())
+		{
+			FMetaSoundFrontendDocumentBuilder& DocBuilder = Builder->GetBuilder();
+			if (DocBuilder.IsValid())
+			{
+				if (const FDocumentConfiguration* Config = DocBuilder.GetDocumentConfiguration().GetPtr())
+				{
+					return Config->GetEditorOptions().bPageGraphEditingEnabled;
+				}
+			}
+		}
+
+		return FMetaSoundDetailCustomizationBase::IsEditable();
+	}
+
 	void FMetasoundPagesDetailCustomization::CustomizeDetails(IDetailLayoutBuilder& DetailLayout)
 	{
 		using namespace Engine;
@@ -487,7 +511,7 @@ namespace Metasound::Editor
 			[
 				SNew(STextBlock)
 				.Text(LOCTEXT("AddPageGraphAction", "Add Page Graph..."))
-				.IsEnabled(TAttribute<bool>::Create(TAttribute<bool>::FGetter::CreateSP(this, &FMetaSoundDetailCustomizationBase::IsGraphEditable)))
+				.IsEnabled(TAttribute<bool>::Create(TAttribute<bool>::FGetter::CreateSP(this, &FMetasoundPagesDetailCustomization::IsEditable)))
 			];
 
 		TSharedRef<SWidget> Utilities = SNew(SHorizontalBox)
@@ -525,7 +549,7 @@ namespace Metasound::Editor
 				}), LOCTEXT("ResetGraphPagesTooltip", "Removes all page graphs from the given MetaSound defined in the MetaSound project settings (does not remove the required 'Default' graph)."))
 			];
 
-		Utilities->SetEnabled(TAttribute<bool>::Create(TAttribute<bool>::FGetter::CreateSP(this, &FMetaSoundDetailCustomizationBase::IsGraphEditable)));
+		Utilities->SetEnabled(TAttribute<bool>::Create(TAttribute<bool>::FGetter::CreateSP(this, &FMetaSoundDetailCustomizationBase::IsEditable)));
 
 		{
 			const FText HeaderName = LOCTEXT("PageGraphsDisplayName", "Graphs");
@@ -684,7 +708,7 @@ namespace Metasound::Editor
 					.AutoWidth()[ ExecImageWidget ];
 			}
 
-			EntryWidget->SetEnabled(TAttribute<bool>::Create(TAttribute<bool>::FGetter::CreateSP(this, &FMetaSoundDetailCustomizationBase::IsGraphEditable)));
+			EntryWidget->SetEnabled(TAttribute<bool>::Create(TAttribute<bool>::FGetter::CreateSP(this, &FMetasoundPagesDetailCustomization::IsEditable)));
 			return EntryWidget;
 		};
 
@@ -862,6 +886,25 @@ namespace Metasound::Editor
 		}
 	}
 
+	bool FMetasoundInterfacesDetailCustomization::IsEditable() const
+	{
+		using namespace Metasound::Frontend;
+
+		if (Builder.IsValid())
+		{
+			FMetaSoundFrontendDocumentBuilder& DocBuilder = Builder->GetBuilder();
+			if (DocBuilder.IsValid())
+			{
+				if (const FDocumentConfiguration* Config = DocBuilder.GetDocumentConfiguration().GetPtr())
+				{
+					return Config->GetEditorOptions().bInterfaceEditingEnabled;
+				}
+			}
+		}
+
+		return FMetaSoundDetailCustomizationBase::IsEditable();
+	}
+
 	void FMetasoundInterfacesDetailCustomization::CustomizeDetails(IDetailLayoutBuilder& DetailLayout)
 	{
 		TArray<TWeakObjectPtr<UObject>> Objects;
@@ -882,9 +925,9 @@ namespace Metasound::Editor
 			}
 		}
 
-		TAttribute<bool> IsGraphEditableAttribute = TAttribute<bool>::Create(TAttribute<bool>::FGetter::CreateSP
+		TAttribute<bool> IsEditableAttribute = TAttribute<bool>::Create(TAttribute<bool>::FGetter::CreateSP
 		(
-			this, &FMetaSoundDetailCustomizationBase::IsGraphEditable
+			this, &FMetasoundInterfacesDetailCustomization::IsEditable
 		));
 
 		UpdateInterfaceNames();
@@ -996,7 +1039,7 @@ namespace Metasound::Editor
 			[
 				SNew(STextBlock)
 				.Text(LOCTEXT("UpdateInterfaceAction", "Add Interface..."))
-				.IsEnabled(IsGraphEditableAttribute)
+				.IsEnabled(IsEditableAttribute)
 			];
 
 		TSharedRef<SWidget> InterfaceUtilities = SNew(SHorizontalBox)
@@ -1059,7 +1102,7 @@ namespace Metasound::Editor
 
 				}), LOCTEXT("RemoveInterfaceTooltip1", "Removes all interfaces from the given MetaSound."))
 			];
-		InterfaceUtilities->SetEnabled(IsGraphEditableAttribute);
+		InterfaceUtilities->SetEnabled(IsEditableAttribute);
 
 		const FText HeaderName = LOCTEXT("InterfacesGroupDisplayName", "Interfaces");
 		IDetailCategoryBuilder& InterfaceCategory = DetailLayout.EditCategory("Interfaces", HeaderName);
@@ -1135,7 +1178,7 @@ namespace Metasound::Editor
 					RemoveButtonWidget
 				];
 
-			EntryWidget->SetEnabled(IsGraphEditableAttribute);
+			EntryWidget->SetEnabled(IsEditableAttribute);
 			return EntryWidget;
 		};
 
diff --git a/Engine/Plugins/Runtime/Metasound/Source/MetasoundEditor/Private/MetasoundDetailCustomization.h b/Engine/Plugins/Runtime/Metasound/Source/MetasoundEditor/Private/MetasoundDetailCustomization.h
--- a/Engine/Plugins/Runtime/Metasound/Source/MetasoundEditor/Private/MetasoundDetailCustomization.h
+++ b/Engine/Plugins/Runtime/Metasound/Source/MetasoundEditor/Private/MetasoundDetailCustomization.h
@@ -28,8 +28,11 @@ namespace Metasound
 		public:
 			virtual ~FMetaSoundDetailCustomizationBase() = default;
 
+			UE_DEPRECATED(5.8, "Use virtual IsEditable instead")
 			bool IsGraphEditable() const;
 
+			virtual bool IsEditable() const;
+
 		protected:
 			UObject* GetMetaSound() const;
 			void InitBuilder(UObject& MetaSound);
@@ -61,9 +64,10 @@ namespace Metasound
 		public:
 			FMetasoundPagesDetailCustomization();
 
-			// IDetailCustomization interface
+			// FMetaSoundDetailCustomizationBase interface
 			virtual void CustomizeDetails(IDetailLayoutBuilder& DetailLayout) override;
-			// End of IDetailCustomization interface
+			virtual bool IsEditable() const override;
+			// End of FMetaSoundDetailCustomizationBase interface
 
 		private:
 			const FMetasoundFrontendDocument& GetConstDocumentChecked() const;
@@ -109,6 +113,8 @@ namespace Metasound
 			virtual void CustomizeDetails(IDetailLayoutBuilder& DetailLayout) override;
 			// End of IDetailCustomization interface
 
+			virtual bool IsEditable() const override;
+
 		private:
 			void UpdateInterfaceNames();
 
diff --git a/Engine/Plugins/Runtime/Metasound/Source/MetasoundEngine/Public/DocumentConfigurations/MetasoundDocumentConfigurator.h b/Engine/Plugins/Runtime/Metasound/Source/MetasoundEngine/Public/DocumentConfigurations/MetasoundDocumentConfigurator.h
--- a/Engine/Plugins/Runtime/Metasound/Source/MetasoundEngine/Public/DocumentConfigurations/MetasoundDocumentConfigurator.h
+++ b/Engine/Plugins/Runtime/Metasound/Source/MetasoundEngine/Public/DocumentConfigurations/MetasoundDocumentConfigurator.h
@@ -30,8 +30,8 @@ namespace Metasound::Engine
 			bool bResetDoc = true;
 		};
 
-		FDocumentConfigurator(const FArgs& Args);
-		~FDocumentConfigurator();
+		UE_API FDocumentConfigurator(const FArgs& Args);
+		UE_API ~FDocumentConfigurator();
 
 #if WITH_EDITORONLY_DATA
 		UE_EXPERIMENTAL(5.8, "Configurator code reflection is in development and may change implementation or be removed in the future")
