Add context to MetaSound duplicate class key log directing users to asset class guid fixup asset action

#rnx

#Synced-CL 50934179

diff --git a/Engine/Plugins/Runtime/Metasound/Source/MetasoundFrontend/Private/MetasoundFrontendNodeClassRegistry.cpp b/Engine/Plugins/Runtime/Metasound/Source/MetasoundFrontend/Private/MetasoundFrontendNodeClassRegistry.cpp
--- a/Engine/Plugins/Runtime/Metasound/Source/MetasoundFrontend/Private/MetasoundFrontendNodeClassRegistry.cpp
+++ b/Engine/Plugins/Runtime/Metasound/Source/MetasoundFrontend/Private/MetasoundFrontendNodeClassRegistry.cpp
@@ -916,8 +916,8 @@ PRAGMA_ENABLE_DEPRECATION_WARNINGS
 					Algo::Transform(AssetPaths, ExistingPaths, [](const FTopLevelAssetPath& AssetPath) { return AssetPath.ToString(); });
 					const FString ExistingAssetPaths = FString::Join(ExistingPaths, TEXT("\n"));
 					UE_LOG(LogMetaSound, Error,
-						TEXT("Multiple node classes with key '%s' registered. Assets currently registered with class name:\n%s"),
-							* Key.ToString(),
+						TEXT("Multiple node classes with key '%s' registered. Assets currently registered with class name:\n%s\nReassign the class name of colliding assets by using the Reassign Asset Class Guid asset action in the right click -> Asset Actions menu of the MetaSound."),
+							*Key.ToString(),
 							*ExistingAssetPaths
 					);
 				}
