Add cancel/wait for MetaSound node registration tasks when unmounting plugins to avoid blocking GC of in-flight referenced MetaSounds

#rnx

#Synced-CL 51042711

diff --git a/Engine/Plugins/Runtime/Metasound/Source/MetasoundFrontend/Private/MetasoundFrontend.cpp b/Engine/Plugins/Runtime/Metasound/Source/MetasoundFrontend/Private/MetasoundFrontend.cpp
--- a/Engine/Plugins/Runtime/Metasound/Source/MetasoundFrontend/Private/MetasoundFrontend.cpp
+++ b/Engine/Plugins/Runtime/Metasound/Source/MetasoundFrontend/Private/MetasoundFrontend.cpp
@@ -3,13 +3,16 @@
 
 #include "MetasoundFrontend.h"
 
+#include "Delegates/IDelegateInstance.h"
 #include "HAL/FileManager.h"
 #include "HAL/IConsoleManager.h"
+#include "Interfaces/IPluginManager.h"
 #include "MetasoundAudioBuffer.h"
 #include "MetasoundDataTypeRegistrationMacro.h"
 #include "MetasoundFrontendDocument.h"
 #include "MetasoundFrontendDocumentVersioning.h"
 #include "MetasoundFrontendGraph.h"
+#include "MetasoundFrontendNodeClassRegistryPrivate.h"
 #include "MetasoundFrontendNodeTemplateRegistry.h"
 #include "MetasoundFrontendRegistryContainer.h"
 #include "MetasoundFrontendRegistryTransaction.h"
@@ -199,13 +202,19 @@ class FMetasoundFrontendModule : public IModuleInterface
 		FVersioningManager::Initialize();
 #endif // WITH_EDITORONLY_DATA
 
-	    FModuleManager::LoadModuleChecked<IModuleInterface>("AudioChannelAgnosticCore");
+		FModuleManager::LoadModuleChecked<IModuleInterface>("AudioChannelAgnosticCore");
+		OnPluginUnmountedHandle = IPluginManager::Get().OnPluginUnmounted().AddLambda([](IPlugin& Plugin)
+		{
+			FNodeClassRegistry::Get().NotifyPluginUnmounted(Plugin);
+		});
 	}
 
 	virtual void ShutdownModule() override
 	{
 		using namespace Metasound::Frontend;
 
+		IPluginManager::Get().OnPluginUnmounted().Remove(OnPluginUnmountedHandle);
+
 #if WITH_EDITORONLY_DATA
 		FVersioningManager::Shutdown();
 #endif // WITH_EDITORONLY_DATA
@@ -215,6 +224,9 @@ class FMetasoundFrontendModule : public IModuleInterface
 		UnregisterNodeTemplate(FInputNodeTemplate::ClassName, FInputNodeTemplate::VersionNumber);
 		UnregisterNodeTemplate(FRerouteNodeTemplate::ClassName, FRerouteNodeTemplate::VersionNumber);
 	}
+
+private:
+	FDelegateHandle OnPluginUnmountedHandle;
 };
 
 METASOUND_IMPLEMENT_MODULE_REGISTRATION_LIST
diff --git a/Engine/Plugins/Runtime/Metasound/Source/MetasoundFrontend/Private/MetasoundFrontendNodeClassRegistry.cpp b/Engine/Plugins/Runtime/Metasound/Source/MetasoundFrontend/Private/MetasoundFrontendNodeClassRegistry.cpp
--- a/Engine/Plugins/Runtime/Metasound/Source/MetasoundFrontend/Private/MetasoundFrontendNodeClassRegistry.cpp
+++ b/Engine/Plugins/Runtime/Metasound/Source/MetasoundFrontend/Private/MetasoundFrontendNodeClassRegistry.cpp
@@ -630,16 +630,20 @@ PRAGMA_ENABLE_DEPRECATION_WARNINGS
 		{
 			// Keep track of exact registration task to ensure correct task instance is removed upon task completion. 
 			int32 RegistrationTaskID = ++RegistrationTaskCntr;
-			
+			TSharedRef<Tasks::FCancellationToken> CancellationToken = MakeShared<Tasks::FCancellationToken>();
 			TArray<FGuid> PageOrder(InPageOrder.GetData(), InPageOrder.Num());
 			Tasks::FTask BuildAndRegisterTask = AsyncRegistrationPipe.Launch(
 				UE_SOURCE_LOCATION,
-				[RegistryKey, RegistryDocInterface, ProxyDataCache = MoveTemp(ProxyDataCache), PageOrder = MoveTemp(PageOrder), RegistrationTaskID]() mutable
+				[RegistryKey, CancellationToken, RegistryDocInterface, ProxyDataCache = MoveTemp(ProxyDataCache), PageOrder = MoveTemp(PageOrder), RegistrationTaskID]() mutable
 				{
 					FNodeClassRegistry& Registry = FNodeClassRegistry::Get();
-					// Unregister the graph before re-registering
-					Registry.UnregisterGraphInternal(RegistryKey);
-					Registry.BuildAndRegisterGraphFromDocument(RegistryDocInterface->GetConstDocument(), ProxyDataCache, RegistryKey, PageOrder);
+					if (!CancellationToken->IsCanceled())
+					{
+						// Unregister the graph before re-registering
+						Registry.UnregisterGraphInternal(RegistryKey);
+						Registry.BuildAndRegisterGraphFromDocument(RegistryDocInterface->GetConstDocument(), ProxyDataCache, RegistryKey, PageOrder);
+					}
+
 					Registry.RemoveRegistrationTask(RegistryKey, RegistrationTaskID, FNodeClassRegistryTransaction::ETransactionType::NodeRegistration);
 					Registry.RemoveDocumentReference(RegistryDocInterface);
 				}
@@ -651,7 +655,9 @@ PRAGMA_ENABLE_DEPRECATION_WARNINGS
 				FNodeClassRegistryTransaction::ETransactionType::NodeRegistration,
 				BuildAndRegisterTask,
 				AssetPath,
-				RegistrationTaskID
+				RegistrationTaskID,
+				RegistryDocInterface,
+				CancellationToken.ToSharedPtr()
 			});
 		}
 		else
@@ -692,7 +698,6 @@ PRAGMA_ENABLE_DEPRECATION_WARNINGS
 			}
 		}
 
-
 		if (NumRemoved == 0)
 		{
 			const bool bIsCooking = IsRunningCookCommandlet();
@@ -1499,6 +1504,46 @@ PRAGMA_ENABLE_DEPRECATION_WARNINGS
 		return nullptr;
 	}
 
+	void FNodeClassRegistry::NotifyPluginUnmounted(IPlugin& Plugin)
+	{
+		METASOUND_TRACE_CPUPROFILER_EVENT_SCOPE(Metasound::FNodeClassRegistry::NotifyPluginUnmounted);
+
+		using namespace UE::Tasks;
+
+		if (!Plugin.CanContainContent())
+		{
+			return;
+		}
+
+		const FString MountPath = Plugin.GetMountedAssetPath();
+		TArray<FTask> TasksToWaitFor;
+		{
+			UE::TScopeLock Lock(ActiveRegistrationTasksCriticalSection);
+			FString PackageName;
+			for (TPair<FNodeClassRegistryKey, TArray<FActiveRegistrationTaskInfo>>& TaskInfoPair : ActiveRegistrationTasks)
+			{
+				for (FActiveRegistrationTaskInfo& TaskInfo : TaskInfoPair.Value)
+				{
+					PackageName = TaskInfo.AssetPath.GetPackageName().ToString();
+					if (PackageName.StartsWith(MountPath))
+					{
+						if (TaskInfo.CancellationToken.IsValid())
+						{
+							TaskInfo.CancellationToken->Cancel();
+							UE_LOG(LogMetaSound, Display, TEXT("Async registration task for plugin '%s' asset '%s' canceled on unmount."),
+								*Plugin.GetName(),
+								*TaskInfo.AssetPath.ToString());
+						}
+
+						TasksToWaitFor.Add(TaskInfo.Task);
+					}
+				}
+			}
+		}
+
+		Algo::ForEach(TasksToWaitFor, [] (const FTask& Task) { Task.Wait(); });
+	}
+
 	void FNodeClassRegistry::WaitForAsyncGraphRegistration(const FGraphRegistryKey& InKey) const
 	{
 		WaitForAsyncRegistrationInternal(InKey.NodeKey, &InKey.AssetPath);
diff --git a/Engine/Plugins/Runtime/Metasound/Source/MetasoundFrontend/Private/MetasoundFrontendNodeClassRegistryPrivate.h b/Engine/Plugins/Runtime/Metasound/Source/MetasoundFrontend/Private/MetasoundFrontendNodeClassRegistryPrivate.h
--- a/Engine/Plugins/Runtime/Metasound/Source/MetasoundFrontend/Private/MetasoundFrontendNodeClassRegistryPrivate.h
+++ b/Engine/Plugins/Runtime/Metasound/Source/MetasoundFrontend/Private/MetasoundFrontendNodeClassRegistryPrivate.h
@@ -4,12 +4,13 @@
 
 #include "Containers/Array.h"
 #include "Containers/Map.h"
-#include "Misc/TransactionallySafeCriticalSection.h"
-#include "Tasks/Pipe.h"
-
+#include "Interfaces/IPluginManager.h"
 #include "MetasoundFrontendDataTypeRegistry.h"
 #include "MetasoundFrontendNodeClassRegistry.h"
 #include "MetasoundFrontendRegistryTransaction.h"
+#include "Misc/TransactionallySafeCriticalSection.h"
+#include "Tasks/Pipe.h"
+
 
 struct FMetasoundFrontendDocument; 
 
@@ -132,6 +133,12 @@ namespace Metasound::Frontend
 		// Returns an empty array if none are available.
 		virtual TArray<FConverterNodeInfo> GetPossibleConverterNodes(const FName& FromDataType, const FName& ToDataType) override;
 
+		// Cancels async graphs contained in the given directory and waits for all
+		// associated unregistration tasks associated with the given plugin to complete
+		// such that subsequent call to GC all remaining plugin-specific content is not
+		// being blocked by in-flight registration tasks.
+		void NotifyPluginUnmounted(IPlugin& Plugin);
+
 		// Register a graph from an IMetaSoundDocumentInterface
 		FGraphRegistryKey RegisterGraph(const TScriptInterface<IMetaSoundDocumentInterface>& InDocumentInterface, TArrayView<const FGuid> InPageOrder);
 
@@ -152,6 +159,9 @@ namespace Metasound::Frontend
 			UE::Tasks::FTask Task;
 			FTopLevelAssetPath AssetPath;
 			int32 RegistrationTaskID = 0;
+			const TScriptInterface<IMetaSoundDocumentInterface> RegistryDocInterface;
+
+			TSharedPtr<UE::Tasks::FCancellationToken> CancellationToken;
 		};
 
 		void BuildAndRegisterGraphFromDocument(const FMetasoundFrontendDocument& InDocument, const FProxyDataCache& InProxyDataCache, const FGraphRegistryKey& InGraphRegistryKey, TArrayView<const FGuid> InPageOrder);
diff --git a/Engine/Plugins/Runtime/Metasound/Source/MetasoundFrontend/Public/MetasoundFrontendDocumentVersioning.h b/Engine/Plugins/Runtime/Metasound/Source/MetasoundFrontend/Public/MetasoundFrontendDocumentVersioning.h
--- a/Engine/Plugins/Runtime/Metasound/Source/MetasoundFrontend/Public/MetasoundFrontendDocumentVersioning.h
+++ b/Engine/Plugins/Runtime/Metasound/Source/MetasoundFrontend/Public/MetasoundFrontendDocumentVersioning.h
@@ -37,7 +37,6 @@ namespace Metasound::Frontend
 
 		UE_API static FVersioningManager& Get();
 		UE_API static void Initialize();
-		UE_API static bool IsActive();
 		UE_API static void Shutdown();
 
 		// Current, max document version
