Stopgap for null audio device causing editor crash when attempting to audition MetaSound with no device
#rnx

diff --git a/Engine/Plugins/Runtime/Metasound/Source/MetasoundEditor/Private/MetasoundEditor.cpp b/Engine/Plugins/Runtime/Metasound/Source/MetasoundEditor/Private/MetasoundEditor.cpp
--- a/Engine/Plugins/Runtime/Metasound/Source/MetasoundEditor/Private/MetasoundEditor.cpp
+++ b/Engine/Plugins/Runtime/Metasound/Source/MetasoundEditor/Private/MetasoundEditor.cpp
@@ -5759,7 +5759,11 @@ PRAGMA_ENABLE_DEPRECATION_WARNINGS
 			}
 
 			FAudioDevice* AudioDevice = PreviewComp->GetAudioDevice();
-			check(AudioDevice);
+			if (!AudioDevice)
+			{
+				return MakeUnique<FGraphConnectionManager>();
+			}
+
 			const FSampleRate DeviceSampleRate = static_cast<FSampleRate>(AudioDevice->GetSampleRate());
 			const uint32 PlayOrder = PreviewComp->GetLastPlayOrder();
 			const uint64 TransmitterID = Audio::GetTransmitterID(PreviewComp->GetAudioComponentID(), 0, PlayOrder);
