Stopgap for loc build errors in CAT

#rnx
#Synced-CL 51567328

diff --git a/Engine/Plugins/Runtime/Metasound/Source/MetasoundFrontend/Private/MetasoundChannelAgnosticDataSchemas.cpp b/Engine/Plugins/Runtime/Metasound/Source/MetasoundFrontend/Private/MetasoundChannelAgnosticDataSchemas.cpp
--- a/Engine/Plugins/Runtime/Metasound/Source/MetasoundFrontend/Private/MetasoundChannelAgnosticDataSchemas.cpp
+++ b/Engine/Plugins/Runtime/Metasound/Source/MetasoundFrontend/Private/MetasoundChannelAgnosticDataSchemas.cpp
@@ -1,4 +1,4 @@
-﻿// Copyright Epic Games, Inc. All Rights Reserved.
+// Copyright Epic Games, Inc. All Rights Reserved.
 
 #include "MetasoundChannelAgnosticDataSchemas.h"
 
@@ -120,7 +120,7 @@ namespace Metasound
 			}
 			return Speakers;
 		}
-		
+
 		FMetasoundDiscreteChannelTypeFamily(const FCatDiscreteDataSchema& InSchema, Audio::IChannelTypeRegistry& InRegistry)
 			: Super(InSchema.Name, InRegistry.FindChannel(InSchema.BaseType), *InSchema.FriendlyName, MakeOrderFromSchema(InSchema), InSchema.bIsParentsDefault, InSchema.bIsAbstract)
 		{
@@ -150,7 +150,7 @@ namespace Metasound
 				GetAmbisonicsOrder(),
 				ToCStr(LexToString(InSchema.bIsParentsDefault)),
 				ToCStr(LexToString(InSchema.bIsAbstract)));
-		}	
+		}
 	};
 
 	class FMetasoundCompositeChannelTypeFamily final : public Audio::FCompositeChannelTypeFamily
@@ -230,7 +230,7 @@ namespace Metasound
 	private:
 		FVertexName VertexName;
 	};
-	
+
 	// Polymorphic Specific INPUT NODE.
 	template<typename DataType, EVertexAccessType VertexAccess=EVertexAccessType::Reference>
 	class TPolyInputNode : public FInputNode
@@ -251,7 +251,7 @@ namespace Metasound
 			return FInputNode::CreateNodeClassMetadata(InVertexName, InConcreteTypeName, VertexAccess);
 		}
 
-			
+
 		explicit TPolyInputNode(FInputNodeConstructorParams&& InParams)
 		:	FInputNode(MoveTemp(InParams), InParams.PolymorphicTypeName, VertexAccess, MakeShared<TPolyInputNodeOperatorFactory<DataType, VertexAccess>>(InParams.VertexName))
 		{
@@ -262,7 +262,7 @@ namespace Metasound
 		{
 		}
 	};
-		
+
 	template<typename SchemaType, typename CppChannelType, typename CppDataType>
 	void AddCatDataFormat(const SchemaType& InSchema)
 	{
@@ -426,7 +426,7 @@ namespace Metasound
 			{
 				return PolymorphicInterface;;
 			}
-		
+
 			virtual const FMetasoundFrontendClass& GetFrontendInputClass() const override
 			{
 				return InputClass;
@@ -592,8 +592,8 @@ namespace Metasound
 				this->ConstructorOutputClassMetadata = MakeShared<Metasound::FNodeClassMetadata>(InConstructorOutputClassMetadata);
 				this->ConstructorOutputClass = Metasound::Frontend::GenerateClass(*this->ConstructorOutputClassMetadata, EMetasoundFrontendClassType::Output);
 			}
-		
-		
+
+
 			virtual TUniquePtr<INode> CreateInputNode(Metasound::FNodeData InNodeData) const override
 			{
 				if (InputClassMetadata)
@@ -828,7 +828,7 @@ namespace Metasound
 				return false;  // fail.
 			}
 		}
-		
+
 		TArray<FName> Ready;
 		Ready.Reserve(InSchemas.Num());
 		for (auto i : Degree)
@@ -839,7 +839,7 @@ namespace Metasound
 				Ready.Add(i.Key);
 			}
 		}
-		
+
 		while (Ready.Num() > 0)
 		{
 			const FName Current = Ready.Pop();
@@ -855,7 +855,7 @@ namespace Metasound
 				}
 			}
 		}
-				
+
 		if (OutLoadOrder.Num() != InSchemas.Num())
 		{
 			for (auto i : Degree)
@@ -904,7 +904,8 @@ namespace Metasound
 		TArray<FName> CyclicNodes;
 		if (FName ProblemNode, ProblemParent; !BuildSchemaLoadOrder(Schemas, Order, ProblemNode, ProblemParent, CyclicNodes))
 		{
-			UE_CLOG(!ProblemNode.IsNone(), LogMetaSound, Error, TEXT("LoadChannelAgosticDataSchemas: Error failed to find parent '%s' in schema '%s'"), 
+#if !NO_LOGGING
+			UE_CLOG(!ProblemNode.IsNone(), LogMetaSound, Display, TEXT("LoadChannelAgosticDataSchemas: Error failed to find parent '%s' in schema '%s'"), 
 				*ProblemParent.ToString(), *ProblemNode.ToString());	
 			
 			// Log each cyclic.
@@ -929,15 +930,16 @@ namespace Metasound
 						*FString::JoinBy(Cyclic, TEXT(" -> "), [](const FName Name){ return Name.ToString(); }));
 				}
 			}
-				
+#endif // !NO_LOGGING
+
 			return {};
 		}
 		for (const FName i : Order)
 		{
 			AddCatDataFormat<SchemaType,CppChannelType, CppDataType>(*Schemas.Find(i));
 		}
-		
-		return Order;			
+
+		return Order;
 	}
 	
 	TArray<FName> LoadChannelAgnosticDataSchemas(const FString& InContentPath)
@@ -960,7 +962,7 @@ namespace Metasound
 			Audio::GetChannelRegistry().UnregisterType(i);
 		}
 	}
-	
+
 	// Put these here for now.
 	FName AddNamespaceToName(const FName InDatatype, const FName InNamespace)
 	{
@@ -975,9 +977,7 @@ namespace Metasound
 	{
 		static const FName Cat(TEXT("Cat:"));
 		return AddNamespaceToName(InDatatype, Cat);
-	}		
-
-
-} 
+	}
+}
 
 #undef LOCTEXT_NAMESPACE
