CL C: GraphCore polymorphic infrastructure

Split from CL 51557033. Core type system changes. MetasoundPolymorphic.h
introduces the polymorphic hierarchy. DataReference.h/DataFactory.h/
Variable.h/VertexData.h modified. All changes UE_EXPERIMENTAL guarded
and existing code paths preserved via if constexpr. Unit tests cover this.
<enter description here>

#tests LowLevelTests Audio

#Synced-CL 51594399

diff --git a/Engine/Plugins/Runtime/Metasound/Source/MetasoundGraphCore/Public/MetasoundDataFactory.h b/Engine/Plugins/Runtime/Metasound/Source/MetasoundGraphCore/Public/MetasoundDataFactory.h
--- a/Engine/Plugins/Runtime/Metasound/Source/MetasoundGraphCore/Public/MetasoundDataFactory.h
+++ b/Engine/Plugins/Runtime/Metasound/Source/MetasoundGraphCore/Public/MetasoundDataFactory.h
@@ -30,10 +30,10 @@ namespace Metasound
 			static constexpr bool bIsArrayConstructibleWithSettingsAndArgs = false;
 
 			// DataType Construction
-			static constexpr bool bIsDefaultConstructible = std::is_constructible<DataType>::value;
-			static constexpr bool bIsConstructibleWithSettings = std::is_constructible<DataType, const FOperatorSettings&>::value;
-			static constexpr bool bIsConstructibleWithArgs = std::is_constructible<DataType, ArgTypes...>::value;
-			static constexpr bool bIsConstructibleWithSettingsAndArgs = std::is_constructible<DataType, const FOperatorSettings&, ArgTypes...>::value;
+			static constexpr bool bIsDefaultConstructible = std::is_constructible<DataType>::value || CHasFactoryFunctionDefaultConstruction<DataType>;
+			static constexpr bool bIsConstructibleWithSettings = std::is_constructible<DataType, const FOperatorSettings&>::value || CHasFactoryFunctionWithSettings<DataType>;
+			static constexpr bool bIsConstructibleWithArgs = std::is_constructible<DataType, ArgTypes...>::value || CHasFactoryFunctionWithArgs<DataType, ArgTypes...>;
+			static constexpr bool bIsConstructibleWithSettingsAndArgs = std::is_constructible<DataType, const FOperatorSettings&, ArgTypes...>::value || CHasFactoryFunctionWithSettingsAndArgs<DataType, ArgTypes...>;
 		};
 
 		/** Partial specialization for TArray<> types. A TArray can be constructed elementwise with
@@ -153,7 +153,7 @@ namespace Metasound
 		{
 			using TConstructorTraits = TDataTypeConstructorTraits<DataType, ArgTypes...>;
 
-			static constexpr bool bForwardSettingsAndArgs = TConstructorTraits::bIsConstructibleWithSettingsAndArgs;
+			static constexpr bool bForwardSettingsAndArgs = TConstructorTraits::bIsConstructibleWithSettingsAndArgs || CHasFactoryFunctionWithSettingsAndArgs<DataType, ArgTypes...>;
 			static constexpr bool bForwardArgs = false;
 			static constexpr bool bForwardSettings = false;
 			static constexpr bool bForwardNone = false;
@@ -172,8 +172,8 @@ namespace Metasound
 		{
 			using TConstructorTraits = TDataTypeConstructorTraits<DataType, ArgTypes...>;
 
-			static constexpr bool bForwardSettingsAndArgs = TConstructorTraits::bIsConstructibleWithSettingsAndArgs;
-			static constexpr bool bForwardArgs = !bForwardSettingsAndArgs && TConstructorTraits::bIsConstructibleWithArgs;
+			static constexpr bool bForwardSettingsAndArgs = TConstructorTraits::bIsConstructibleWithSettingsAndArgs || CHasFactoryFunctionWithSettingsAndArgs<DataType, ArgTypes...>;
+			static constexpr bool bForwardArgs = !bForwardSettingsAndArgs && (TConstructorTraits::bIsConstructibleWithArgs || CHasFactoryFunctionWithArgs<DataType, ArgTypes...>); 
 			static constexpr bool bForwardSettings = false;
 			static constexpr bool bForwardNone = false;
 
@@ -191,10 +191,10 @@ namespace Metasound
 		{
 			using TConstructorTraits = TDataTypeConstructorTraits<DataType, ArgTypes...>;
 
-			static constexpr bool bForwardSettingsAndArgs = TConstructorTraits::bIsConstructibleWithSettingsAndArgs;
-			static constexpr bool bForwardArgs = !bForwardSettingsAndArgs && TConstructorTraits::bIsConstructibleWithArgs;
-			static constexpr bool bForwardSettings = !(bForwardArgs || bForwardSettingsAndArgs) && TConstructorTraits::bIsConstructibleWithSettings;
-			static constexpr bool bForwardNone = !(bForwardArgs || bForwardSettingsAndArgs || bForwardSettings) && TConstructorTraits::bIsDefaultConstructible;
+			static constexpr bool bForwardSettingsAndArgs = TConstructorTraits::bIsConstructibleWithSettingsAndArgs || CHasFactoryFunctionWithSettingsAndArgs<DataType, ArgTypes...>;
+			static constexpr bool bForwardArgs = !bForwardSettingsAndArgs && (TConstructorTraits::bIsConstructibleWithArgs || CHasFactoryFunctionWithArgs<DataType, ArgTypes...>);
+			static constexpr bool bForwardSettings = !(bForwardArgs || bForwardSettingsAndArgs) && (TConstructorTraits::bIsConstructibleWithSettings || CHasFactoryFunctionWithSettings<DataType>);
+			static constexpr bool bForwardNone = !(bForwardArgs || bForwardSettingsAndArgs || bForwardSettings) && (TConstructorTraits::bIsDefaultConstructible || CHasFactoryFunctionDefaultConstruction<DataType>);
 
 			static constexpr bool bCannotForwardToConstructor = !(bForwardSettingsAndArgs || bForwardArgs || bForwardSettings || bForwardNone);
 		};
diff --git a/Engine/Plugins/Runtime/Metasound/Source/MetasoundGraphCore/Public/MetasoundDataReference.h b/Engine/Plugins/Runtime/Metasound/Source/MetasoundGraphCore/Public/MetasoundDataReference.h
--- a/Engine/Plugins/Runtime/Metasound/Source/MetasoundGraphCore/Public/MetasoundDataReference.h
+++ b/Engine/Plugins/Runtime/Metasound/Source/MetasoundGraphCore/Public/MetasoundDataReference.h
@@ -4,6 +4,7 @@
 
 #include <type_traits>
 
+#include "Concepts/DerivedFrom.h"
 #include "Containers/UnrealString.h"
 #include "HAL/Platform.h"
 #include "Internationalization/Text.h"
@@ -24,6 +25,7 @@ namespace Audio
 
 namespace Metasound
 {
+	class FOperatorSettings;
 	// Forward declare for LexToString
 	enum class EDataReferenceAccessType : uint8;
 }
@@ -96,8 +98,8 @@ namespace Metasound
 			return false;
 		}
 	}
-	
-	
+
+
 	// Polymorphic overload
 	template<typename DataType>
 	requires TPolymorphicTraits<DataType>::bIsPolymorphic
@@ -137,56 +139,83 @@ namespace Metasound
 	{
 		NewObject
 	};
-	
+
 	namespace MetasoundDataReferencePrivate
 	{
 		template <typename Base, typename Derived, typename... ArgTypes>
 		UE_EXPERIMENTAL(5.7, "Polymorphic API is under development")
-		TSharedRef<Base, ESPMode::NotThreadSafe> MakeRefOf(ArgTypes&&... args)
+		TSharedRef<Base, ESPMode::NotThreadSafe> MakeRefOf(ArgTypes&&... Args)
 		{
 			static_assert(std::is_base_of_v<Base, Derived>, "Target must derive from Base (or be Base)");
 			static_assert(std::is_constructible_v<Derived, ArgTypes...>, "Args do not match Target's constructor");
 
-			TSharedRef<Derived, ESPMode::NotThreadSafe> DerivedRef = MakeShared<Derived, ESPMode::NotThreadSafe>(Forward<ArgTypes>(args)...);
- 
+			TSharedRef<Derived, ESPMode::NotThreadSafe> DerivedRef = MakeShared<Derived, ESPMode::NotThreadSafe>(Forward<ArgTypes>(Args)...);
+
 			return StaticCastSharedRef<Base>(DerivedRef);
 		}
+
+		template <typename Base, typename Derived, typename... ArgTypes>
+		UE_EXPERIMENTAL(5.7, "Polymorphic API is under development")
+		TSharedRef<Base, ESPMode::NotThreadSafe> MakeRefOf(ArgTypes&&... Args) requires CHasFactoryFunction<Derived, ArgTypes...>
+		{
+			static_assert(std::is_base_of_v<Base, Derived>, "Target must derive from Base (or be Base)");
+			TSharedRef<Derived, ESPMode::NotThreadSafe> DerivedRef = Derived::CreateNew(Forward<ArgTypes>(Args)...);
+ 			return StaticCastSharedRef<Base>(DerivedRef);
+		}
+
 	}
 
 	template <typename DataType>
 	struct TDynamicDataReferencePolicy;
 
 	template<typename DataType>
-	requires (TPolymorphicTraits<DataType>::bIsPolymorphic)
+	requires (!TPolymorphicTraits<DataType>::bIsPolymorphic)
 	struct TDynamicDataReferencePolicy<DataType>
 	{
-		explicit TDynamicDataReferencePolicy(const FName& InDataTypeName = GetMetasoundDataTypeName<DataType>(), const void* InTypeId = GetMetasoundDataTypeId<DataType>())
-			: DataTypeName(InDataTypeName), DataTypeId(InTypeId)
-		{}
-
-		template<typename QueryType> bool IsA() const { return IsChildOfByName(DataTypeName, GetMetasoundDataTypeName<QueryType>()); }
-		const FName& GetDataTypeName() const { return DataTypeName; }
-		const void* GetTypeId() const { return DataTypeId; }
-		FName DataTypeName;
-		const void* DataTypeId;
+		template<typename QueryType> bool IsA(const DataType*) const { return GetMetasoundDataTypeName<DataType>() == GetMetasoundDataTypeName<QueryType>(); }
+		const FName& GetDataTypeName(const DataType*) const { return GetMetasoundDataTypeName<DataType>(); }
+		const void* GetDataTypeId(const DataType*) const { return GetMetasoundDataTypeId<DataType>(); }
 	};
-	
+
 	template<typename DataType>
-	requires (!TPolymorphicTraits<DataType>::bIsPolymorphic)
+	requires (TPolymorphicTraits<DataType>::bIsPolymorphic)
 	struct TDynamicDataReferencePolicy<DataType>
 	{
-		TDynamicDataReferencePolicy() = default;
-		TDynamicDataReferencePolicy(const FName&, const void*) {}
+		/* TODO: Re-enable once all poly types define GetDataTypeName+GetDataTypeId */
+		/* static_assert(CHasPolymorphicTypeFunctions<DataType>, "Dynamic Reference Policy Requires GetDataTypeName+GetDataTypeId functions on DataType"); */
+
+		template<typename QueryType> bool IsA(const DataType* InOuter) const
+		{
+			return IsChildOfByName(GetDataTypeName(InOuter), GetMetasoundDataTypeName<QueryType>());
+		}
 
-		template<typename QueryType> bool IsA() const { return GetDataTypeName() == GetMetasoundDataTypeName<QueryType>(); }
-		const FName& GetDataTypeName() const { return GetMetasoundDataTypeName<DataType>(); }
-		const void* GetTypeId() const { return GetMetasoundDataTypeId<DataType>(); }
+		const FName& GetDataTypeName(const DataType* InOuter) const requires(CHasPolymorphicTypeGetter<DataType>)
+		{
+			return InOuter->GetDataTypeName();
+		}
+
+		// Fallback when type doesn't yet implement GetDataTypeName().
+		const FName& GetDataTypeName(const DataType*) const requires(!CHasPolymorphicTypeGetter<DataType>)
+		{
+			return GetMetasoundDataTypeName<DataType>();
+		}
+
+		const void* GetDataTypeId(const DataType* InOuter) const requires(CHasPolymorphicTypeIdGetter<DataType>)
+		{
+			return InOuter->GetDataTypeId();
+		}
+
+		// Fallback when type doesn't yet implement GetDataTypeId().
+		const void* GetDataTypeId(const DataType*) const requires(!CHasPolymorphicTypeIdGetter<DataType>)
+		{
+			return GetMetasoundDataTypeId<DataType>();
+		}
 	};
-	
-	/** Template class for a paramter reference. 
+
+	/** Template class for a paramter reference.
 	 *
 	 * This fulfills the IParamterRef interface, utilizing TDataReferenceTypeInfo to
-	 * define the the TypeName and TypeId of the parameter. 
+	 * define the the TypeName and TypeId of the parameter.
 	 */
 	template <typename DataType>
 	class TDataReference : public IDataReference
@@ -202,6 +231,13 @@ namespace Metasound
 		{
 		}
 
+		template <typename... ArgTypes>
+		TDataReference(EDataRefShouldConstruct InToken, ArgTypes&&... Args) requires(CIsPolymorphicAndAbstract<DataType> && CHasFactoryFunction<DataType, ArgTypes...>)
+			: ObjectReference(DataType::CreateNew(Forward<ArgTypes>(Args)...))
+		{
+			/* TODO: Re-enable once all poly types define CreateNew */ // static_assert(CHasFactoryFunction<DataType, ArgTypes...>);
+		}
+
 		typedef TSharedRef<DataType, ESPMode::NotThreadSafe> FRefType;
 
 		explicit TDataReference(EDataRefShouldConstruct /*InToken*/, FRefType InRef)
@@ -213,9 +249,8 @@ namespace Metasound
 								const FName InTypeName,
 								const void* InTypeId)
 			: ObjectReference(MoveTemp(InRef))
-			, ObjectReferencePolicy(InTypeName, InTypeId)
 		{}
-	
+
 	public:
 
 		typedef TDataReferenceTypeInfo<DataType> FInfoType;
@@ -229,37 +264,38 @@ namespace Metasound
 			return TDataReference<DataType>(EDataRefShouldConstruct::NewObject, Forward<ArgTypes>(Args)...);
 		}
 
+		/**
+		 * Polymorphic form of CreateNew.
+		 */
+		template <typename... ArgTypes>
+		static TDataReference<DataType> CreateNew(ArgTypes&&... Args) requires (CIsPolymorphicAndAbstract<DataType> && CHasFactoryFunction<DataType, ArgTypes...>)
+		{
+			/* TODO: Re-enable once all poly types define CreateNew */ // static_assert(CHasFactoryFunction<DataType, ArgTypes...>);
+			return TDataReference<DataType>(EDataRefShouldConstruct::NewObject, Forward<ArgTypes>(Args)...);
+		}
+
 		// Constructs Derived but stores as Base (DataType).
 		template <typename Derived, typename... ArgTypes>
-		UE_EXPERIMENTAL(5.7, "Polymorphic API is under development")
 		static TDataReference CreateNewDerived(ArgTypes&&... Args)
 		{
 			static_assert(std::is_base_of_v<DataType, Derived>, "Derived must inherit from DataType");
-			static_assert(std::is_constructible_v<Derived, ArgTypes...>, "Args don't match a Derived constructor");
-			
 			return TDataReference(
 				EDataRefShouldConstruct::NewObject,
-				MetasoundDataReferencePrivate::MakeRefOf<DataType, Derived>(Forward<ArgTypes>(Args)...),
-				GetMetasoundDataTypeName<Derived>(),
-			    GetMetasoundDataTypeId<Derived>()
+				MetasoundDataReferencePrivate::MakeRefOf<DataType, Derived>(Forward<ArgTypes>(Args)...)
 			);
 		}
 
 		template <typename Derived, typename... ArgTypes>
-		UE_EXPERIMENTAL(5.8, "Polymorphic API is under development")
 		static TDataReference CreateNewDerivedAs(const FName& InDerivedTypeName,
 												 FMetasoundDataTypeId InDerivedTypeId,
 												 ArgTypes&&... Args)
 		{
 			static_assert(std::is_base_of_v<DataType, Derived>, "Derived must inherit from DataType");
-			static_assert(std::is_constructible_v<Derived, ArgTypes...>, "Args don't match a Derived constructor");
-
 		#if DO_CHECK
 			PRAGMA_DISABLE_EXPERIMENTAL_WARNINGS
 			check(Metasound::IsChildOfByName(InDerivedTypeName, GetMetasoundDataTypeName<DataType>()));
 			PRAGMA_ENABLE_EXPERIMENTAL_WARNINGS
 		#endif
-
 			return TDataReference(
 				EDataRefShouldConstruct::NewObject,
 				MetasoundDataReferencePrivate::MakeRefOf<DataType, Derived>(Forward<ArgTypes>(Args)...),
@@ -268,7 +304,6 @@ namespace Metasound
 			);
 		}
 
-
 		/** Enable copy constructor */
 		TDataReference(const TDataReference<DataType>& Other) = default;
 
@@ -284,13 +319,13 @@ namespace Metasound
 		/** Return the name of the underlying type. */
 		virtual const FName& GetDataTypeName() const override
 		{
-			return ObjectReferencePolicy.GetDataTypeName();	
+			return ObjectReferencePolicy.GetDataTypeName(&ObjectReference.Get());
 		}
 
 		/** Return the ID of the underlying type. */
 		virtual const void* const GetDataTypeId() const override
 		{
-			return ObjectReferencePolicy.GetTypeId();
+			return ObjectReferencePolicy.GetDataTypeId(&ObjectReference.Get());
 		}
 
 		/** Return a raw pointer to the data. */
@@ -303,29 +338,24 @@ namespace Metasound
 		UE_EXPERIMENTAL(5.7, "Polymorphic API is under development")
 		bool IsA() const
 		{
-			return this->ObjectReferencePolicy.template IsA<QueryType>();
+			return this->ObjectReferencePolicy.template IsA<QueryType>(&ObjectReference.Get());
 		}
-		
+
 		template <typename CastType>
 		UE_EXPERIMENTAL(5.7, "Polymorphic API is under development")
-		TDataReference<CastType> TryCast() const
+		TOptional<TSharedRef<CastType, ESPMode::NotThreadSafe>> TryCast() const
 		{
 			if (IsA<CastType>())
 			{
-				return TDataReference(
-					EDataRefShouldConstruct::NewObject,
-					MetasoundDataReferencePrivate::MakeRefOf<DataType, CastType>(ObjectReference),
-					GetMetasoundDataTypeName<CastType>(),
-					GetMetasoundDataTypeId<CastType>()
-				);
+				return StaticCastSharedRef<CastType>(ObjectReference);
 			}
 			return {};
 		}
 
 	protected:
-		
+
 		// Protected object reference is utilized by subclasses which define what
-		// access is provided to the ObjectReference. 
+		// access is provided to the ObjectReference.
 		FRefType ObjectReference;
 		TDynamicDataReferencePolicy<DataType> ObjectReferencePolicy;
 	};
@@ -338,7 +368,7 @@ namespace Metasound
 	template <typename DataType>
 	class TDataWriteReference;
 
-	/** TDataValueReference represents a constant value and provides read only access. 
+	/** TDataValueReference represents a constant value and provides read only access.
 	 * A TDataValueReference can never change value. */
 	template<typename DataType>
 	class TDataValueReference : public TDataReference<DataType>
@@ -361,7 +391,7 @@ namespace Metasound
 
 		template<typename T>
 		friend TDataValueReference<T> ValueCast(const TDataWriteReference<T>& InRef);
-		
+
 		using FRefType = TDataReference<DataType>::FRefType;
 		TDataValueReference(EDataRefShouldConstruct, FRefType InRef, FName InTypeName, const void* InTypeId)
 			: TDataReference<DataType>(EDataRefShouldConstruct::NewObject, MoveTemp(InRef), InTypeName, InTypeId)
@@ -378,13 +408,17 @@ namespace Metasound
 			return TDataValueReference<DataType>(EDataRefShouldConstruct::NewObject, Forward<ArgTypes>(Args)...);
 		}
 
+		template <typename... ArgTypes>
+		static TDataValueReference<DataType> CreateNew(ArgTypes&&... Args) requires (CIsPolymorphicAndAbstract<DataType> && CHasFactoryFunction<DataType, ArgTypes...>)
+		{
+			/* TODO: Re-enable once all poly types define CreateNew */ // static_assert(CHasFactoryFunction<DataType, ArgTypes...>);
+			return TDataValueReference<DataType>(EDataRefShouldConstruct::NewObject, Forward<ArgTypes>(Args)...);
+		}
+
 		template <typename Derived, typename... ArgTypes>
-		UE_EXPERIMENTAL(5.7, "Polymorphic API is under development")
 		static TDataValueReference CreateNewDerived(ArgTypes&&... Args)
 		{
 			static_assert(std::is_base_of_v<DataType, Derived>, "Derived must inherit from DataType");
-			static_assert(std::is_constructible_v<Derived, ArgTypes...>, "Args don't match a Derived constructor");
-			
 			return TDataValueReference(
 				EDataRefShouldConstruct::NewObject,
 				MetasoundDataReferencePrivate::MakeRefOf<DataType, Derived>(Forward<ArgTypes>(Args)...)
@@ -392,20 +426,16 @@ namespace Metasound
 		}
 
 		template <typename Derived, typename... ArgTypes>
-		UE_EXPERIMENTAL(5.8, "Polymorphic API is under development")
 		static TDataValueReference CreateNewDerivedAs(const FName& InDerivedTypeName,
 													 FMetasoundDataTypeId InDerivedTypeId,
 													 ArgTypes&&... Args)
 		{
 			static_assert(std::is_base_of_v<DataType, Derived>, "Derived must inherit from DataType");
-			static_assert(std::is_constructible_v<Derived, ArgTypes...>, "Args don't match a Derived constructor");
-
 		#if DO_CHECK
 			PRAGMA_DISABLE_EXPERIMENTAL_WARNINGS
 			check(Metasound::IsChildOfByName(InDerivedTypeName, GetMetasoundDataTypeName<DataType>()));
 			PRAGMA_ENABLE_EXPERIMENTAL_WARNINGS
 		#endif
-			
 			return TDataValueReference(
 				EDataRefShouldConstruct::NewObject,
 				MetasoundDataReferencePrivate::MakeRefOf<DataType, Derived>(Forward<ArgTypes>(Args)...),
@@ -413,7 +443,9 @@ namespace Metasound
 				InDerivedTypeId
 			);
 		}
-		
+
+
+
 		/** Enable copy constructor */
 		TDataValueReference(const TDataValueReference<DataType>& Other) = default;
 
@@ -458,10 +490,10 @@ namespace Metasound
 		}
 	};
 
-	/** Cast a TDataReadReference to a TDataValueReference. 
+	/** Cast a TDataReadReference to a TDataValueReference.
 	 *
 	 * In general TDataReadReferences should not be converted into TDataValueReferences unless the caller
-	 * can be certain that no other TDataWriteReference exists for the underlying parameter. Having a 
+	 * can be certain that no other TDataWriteReference exists for the underlying parameter. Having a
 	 * TDataWriteReferences to an existing TDataValueReference can cause confusing behavior as values
 	 * references are not expected to change value.
 	 */
@@ -471,10 +503,10 @@ namespace Metasound
 		return TDataValueReference<T>(InRef);
 	}
 
-	/** Cast a TDataWriteReference to a TDataValueReference. 
+	/** Cast a TDataWriteReference to a TDataValueReference.
 	 *
 	 * In general TDataWriteReferences should never beconverted into TDataValueReferences unless the caller
-	 * can be certain that no other TDataWriteReference exists for the underlying parameter. Having a 
+	 * can be certain that no other TDataWriteReference exists for the underlying parameter. Having a
 	 * TDataWriteReferences to an existing TDataValueReference can cause confusing behavior as values
 	 * references are not expected to change value.
 	 */
@@ -496,7 +528,7 @@ namespace Metasound
 		}
 
 		/** Create a writable ref from a blank parameter ref. Should be done with care and understanding
-		 * of side-effects of converting a ref to a writable ref. 
+		 * of side-effects of converting a ref to a writable ref.
 		 */
 		TDataWriteReference<DataType>(const TDataReference<DataType>& InDataReference)
 		: TDataReference<DataType>(InDataReference)
@@ -522,13 +554,18 @@ namespace Metasound
 			return TDataWriteReference<DataType>(EDataRefShouldConstruct::NewObject, Forward<ArgTypes>(Args)...);
 		}
 
+		/** This should be used to construct a new DataType object and return this TDataWriteReference as a wrapper around it. */
+		template <typename... ArgTypes>
+		static TDataWriteReference<DataType> CreateNew(ArgTypes&&... Args) requires (CIsPolymorphicAndAbstract<DataType> && CHasFactoryFunction<DataType, ArgTypes...>)
+		{
+			/* TODO: Re-enable once all poly types define CreateNew */ // static_assert(CHasFactoryFunction<DataType, ArgTypes...>);
+			return TDataWriteReference<DataType>(EDataRefShouldConstruct::NewObject, Forward<ArgTypes>(Args)...);
+		}
+
 		template <typename Derived, typename... ArgTypes>
-		UE_EXPERIMENTAL(5.8, "Polymorphic API is under development")
 		static TDataWriteReference CreateNewDerived(ArgTypes&&... Args)
 		{
 			static_assert(std::is_base_of_v<DataType, Derived>, "Derived must inherit from DataType");
-			static_assert(std::is_constructible_v<Derived, ArgTypes...>, "Args don't match a Derived constructor");
-			
 			return TDataWriteReference(
 				EDataRefShouldConstruct::NewObject,
 				MetasoundDataReferencePrivate::MakeRefOf<DataType, Derived>(Forward<ArgTypes>(Args)...)
@@ -536,20 +573,16 @@ namespace Metasound
 		}
 
 		template <typename Derived, typename... ArgTypes>
-		UE_EXPERIMENTAL(5.8, "Polymorphic API is under development")
 		static TDataWriteReference CreateNewDerivedAs(const FName& InDerivedTypeName,
 													 FMetasoundDataTypeId InDerivedTypeId,
 													 ArgTypes&&... Args)
 		{
 			static_assert(std::is_base_of_v<DataType, Derived>, "Derived must inherit from DataType");
-			static_assert(std::is_constructible_v<Derived, ArgTypes...>, "Args don't match a Derived constructor");
-
 		#if DO_CHECK
 			PRAGMA_DISABLE_EXPERIMENTAL_WARNINGS
 			check(Metasound::IsChildOfByName(InDerivedTypeName, GetMetasoundDataTypeName<DataType>()));
 			PRAGMA_ENABLE_EXPERIMENTAL_WARNINGS
 		#endif
-			
 			return TDataWriteReference(
 				EDataRefShouldConstruct::NewObject,
 				MetasoundDataReferencePrivate::MakeRefOf<DataType, Derived>(Forward<ArgTypes>(Args)...),
@@ -557,7 +590,9 @@ namespace Metasound
 				InDerivedTypeId
 			);
 		}
-		
+
+
+
 		/** Enable copy constructor */
 		TDataWriteReference(const TDataWriteReference<DataType>& Other) = default;
 
@@ -607,17 +642,17 @@ namespace Metasound
 
 	};
 
-	/** Cast a TDataReadReference to a TDataWriteReference. 
+	/** Cast a TDataReadReference to a TDataWriteReference.
 	 *
 	 * In general TDataReadReferences should not be converted into TDataWriteReferences unless the caller
-	 * can be certain that no other TDataWriteReference exists for the underlying parameter. Having multiple 
+	 * can be certain that no other TDataWriteReference exists for the underlying parameter. Having multiple
 	 * TDataWriteReferences to the same parameter can cause confusion behavior as values are overwritten in
 	 * an underterministic fashion.
 	 */
 	template <typename DataType>
 	TDataWriteReference<DataType> WriteCast(const TDataReadReference<DataType>& InReadableRef)
 	{
-		const TDataReference<DataType>& Ref = static_cast<const TDataReference<DataType>&>(InReadableRef); 
+		const TDataReference<DataType>& Ref = static_cast<const TDataReference<DataType>&>(InReadableRef);
 		return TDataWriteReference<DataType>(Ref);
 	}
 
@@ -650,13 +685,17 @@ namespace Metasound
 			return TDataReadReference<DataType>(EDataRefShouldConstruct::NewObject, Forward<ArgTypes>(Args)...);
 		}
 
+		template <typename... ArgTypes>
+		static TDataReadReference<DataType> CreateNew(ArgTypes&&... Args) requires (CIsPolymorphicAndAbstract<DataType> && CHasFactoryFunction<DataType, ArgTypes...>)
+		{
+			/* TODO: Re-enable once all poly types define CreateNew */ // static_assert(CHasFactoryFunction<DataType, ArgTypes...>);
+			return TDataReadReference<DataType>(EDataRefShouldConstruct::NewObject, Forward<ArgTypes>(Args)...);
+		}
+
 		template <typename Derived, typename... ArgTypes>
-		UE_EXPERIMENTAL(5.8, "Polymorphic API is under development")
 		static TDataReadReference CreateNewDerived(ArgTypes&&... Args)
 		{
 			static_assert(std::is_base_of_v<DataType, Derived>, "Derived must inherit from DataType");
-			static_assert(std::is_constructible_v<Derived, ArgTypes...>, "Args don't match a Derived constructor");
-			
 			return TDataReadReference(
 				EDataRefShouldConstruct::NewObject,
 				MetasoundDataReferencePrivate::MakeRefOf<DataType, Derived>(Forward<ArgTypes>(Args)...)
@@ -664,18 +703,14 @@ namespace Metasound
 		}
 
 		template <typename Derived, typename... ArgTypes>
-		UE_EXPERIMENTAL(5.8, "Polymorphic API is under development")
 		static TDataReadReference CreateNewDerivedAs(const FName& InDerivedTypeName, FMetasoundDataTypeId InDerivedTypeId, ArgTypes&&... Args)
 		{
 			static_assert(std::is_base_of_v<DataType, Derived>, "Derived must inherit from DataType");
-			static_assert(std::is_constructible_v<Derived, ArgTypes...>, "Args don't match a Derived constructor");
-
 		#if DO_CHECK
 			PRAGMA_DISABLE_EXPERIMENTAL_WARNINGS
 			check(Metasound::IsChildOfByName(InDerivedTypeName, GetMetasoundDataTypeName<DataType>()));
 			PRAGMA_ENABLE_EXPERIMENTAL_WARNINGS
 		#endif
-			
 			return TDataReadReference(
 				EDataRefShouldConstruct::NewObject,
 				MetasoundDataReferencePrivate::MakeRefOf<DataType, Derived>(Forward<ArgTypes>(Args)...),
@@ -683,7 +718,7 @@ namespace Metasound
 				InDerivedTypeId
 			);
 		}
-		
+
 		TDataReadReference(const TDataReadReference& Other) = default;
 
 		/** Construct a readable parameter ref from a writable parameter ref. */
@@ -739,9 +774,9 @@ namespace Metasound
 	};
 
 	/** EDataReferenceAccessType describes the underlying data reference access
-	 * type for a data reference contained in a FAnyDataReference. 
+	 * type for a data reference contained in a FAnyDataReference.
 	 *
-	 * This value can be used to determine which methods are supported for accessing 
+	 * This value can be used to determine which methods are supported for accessing
 	 * a data reference using GetDataReadReference<>() or GetDataWriteReference<>()
 	 */
 	enum class EDataReferenceAccessType : uint8
@@ -754,17 +789,17 @@ namespace Metasound
 
 
 
-	/** Container for any data reference. 
+	/** Container for any data reference.
 	 *
 	 * This container maintains the underlying containers access type (Read or Write)
-	 * and data type. This allows for convenient storage by implementing a virtual 
+	 * and data type. This allows for convenient storage by implementing a virtual
 	 * copy constructor and assignment operator.
 	 */
 	class FAnyDataReference : public IDataReference
 	{
 		/** TGetAs allows for template specialization of data retrieval and conversion
 		 * rules for the various flavors of data referencing. The default implementation
-		 * assumes we are attempting to retrieve the value stored in the data reference. 
+		 * assumes we are attempting to retrieve the value stored in the data reference.
 		 */
 		template<typename ResultType>
 		struct TGetAs
@@ -773,7 +808,7 @@ namespace Metasound
 			using DataType = typename std::decay_t<std::remove_pointer_t<ResultType>>;
 
 			// Only deref a pointer if the returned type is specifically requested to be
-			// a non-pointer. 
+			// a non-pointer.
 			static ResultType DerefPointerIfNonPointerResultType(const DataType* InData)
 			{
 				if constexpr (std::is_pointer_v<ResultType>)
@@ -788,7 +823,7 @@ namespace Metasound
 			};
 
 		public:
-		
+
 			// Return a the underlying data by value or by const pointer depending
 			// upon the template parameter ResultType
 			static ResultType Get(EDataReferenceAccessType InAccessType, const IDataReference& InRef)
@@ -883,7 +918,7 @@ namespace Metasound
 		};
 
 
-		// Private constructor. 
+		// Private constructor.
 		FAnyDataReference(EDataReferenceAccessType InAccessType, const IDataReference& InDataRef)
 		: AccessType(InAccessType)
 		, DataRefPtr(InDataRef.Clone())
@@ -944,7 +979,7 @@ namespace Metasound
 			check(DataRefPtr.IsValid());
 			return DataRefPtr->GetDataTypeName();
 		}
-		
+
 		/** Returns the data type ID of the underlying data reference. */
 		virtual const void* const GetDataTypeId() const override
 		{
@@ -968,7 +1003,7 @@ namespace Metasound
 
 		/** Return the data in the specified format. The template argument
 		 * can be any of the support TData*Reference<> types, the underlying DataType
-		 * or a const DataType* pointer. 
+		 * or a const DataType* pointer.
 		 */
 		template<typename ReturnType>
 		ReturnType GetAs() const
@@ -984,10 +1019,10 @@ namespace Metasound
 			}
 		}
 
-		/** Returns the current value of a reference. 
+		/** Returns the current value of a reference.
 		 *
 		 * This method's behavior is undefined and an assert will be called if
-		 * the DataType differs from the underlying data reference's DataType. 
+		 * the DataType differs from the underlying data reference's DataType.
 		 */
 		template<typename DataType>
 		const DataType* GetValue() const
@@ -1000,11 +1035,11 @@ namespace Metasound
 		}
 
 		/** Return a non-const pointer to the data.
-		 * 
-		 * If the vertex is bound with this will return a valid pointer. 
+		 *
+		 * If the vertex is bound with this will return a valid pointer.
 		 * Otherwise it will return a nullptr.  This method will assert
 		 * if the reference is bound, but not writable or if the data type
-		 * does match.  
+		 * does match.
 		 */
 		template<typename DataType>
 		DataType* GetWritableValue() const
@@ -1019,14 +1054,14 @@ namespace Metasound
 			return nullptr;
 		}
 
-		/** Get access to a TDataValueReference. 
+		/** Get access to a TDataValueReference.
 		 *
 		 * This method will return a valid TDataValueReference of the templated data
 		 * type. The returned object is only valid if:
 		 *     1. The template parameter DataType matches that of the underlying data reference.
 		 *     2. The underlying data reference is has Value access.
 		 *
-		 * If this method's behavior is undefined and will assert if it is called 
+		 * If this method's behavior is undefined and will assert if it is called
 		 * with a mismatched data type or unsupported access type.
 		 */
 		template<typename DataType>
@@ -1036,14 +1071,14 @@ namespace Metasound
 			return GetAs<TDataValueReference<DataType>>();
 		}
 
-		/** Get access to a TDataReadReference. 
+		/** Get access to a TDataReadReference.
 		 *
 		 * This method will return a valid TDataReadReference of the templated data
 		 * type. The returned object is only valid if:
 		 *     1. The template parameter DataType matches that of the underlying data reference.
 		 *     2. The underlying data reference is has Read, Write or Value access.
 		 *
-		 * If this method's behavior is undefined and will assert if it is called 
+		 * If this method's behavior is undefined and will assert if it is called
 		 * with a mismatched data type or unsupported access type.
 		 */
 		template<typename DataType>
@@ -1052,14 +1087,14 @@ namespace Metasound
 			return GetAs<TDataReadReference<DataType>>();
 		}
 
-		/** Get access to a TDataWriteReference. 
+		/** Get access to a TDataWriteReference.
 		 *
 		 * This method will return a valid TDataWriteReference of the templated data
 		 * type. The returned object is only valid if:
 		 *     1. The template paramter DataType matches that of the underlying data reference.
 		 *     2. The underlying data reference is has Write access.
 		 *
-		 * If this method's behavior is undefined and will assert if it is called 
+		 * If this method's behavior is undefined and will assert if it is called
 		 * with a mismatched data type or unsupported access type.
 		 */
 		template<typename DataType>
diff --git a/Engine/Plugins/Runtime/Metasound/Source/MetasoundGraphCore/Public/MetasoundPolymorphic.h b/Engine/Plugins/Runtime/Metasound/Source/MetasoundGraphCore/Public/MetasoundPolymorphic.h
--- a/Engine/Plugins/Runtime/Metasound/Source/MetasoundGraphCore/Public/MetasoundPolymorphic.h
+++ b/Engine/Plugins/Runtime/Metasound/Source/MetasoundGraphCore/Public/MetasoundPolymorphic.h
@@ -3,9 +3,12 @@
 #pragma once
 
 #include "MetasoundDataTypeGetTraits.h"
+#include "Concepts/DerivedFrom.h"
 
 namespace Metasound
 {
+	class FOperatorSettings;
+
 	template<typename T>
 	struct TPolymorphicTraits
 	{
@@ -21,6 +24,56 @@ namespace Metasound
 		bool bIsPolymorphic = false;
 	};
 
+	template<typename DataType>
+	concept CIsPolymorphic = TPolymorphicTraits<DataType>::bIsPolymorphic;
+
+	template<typename DataType>
+	concept CIsNotPolymorphic = !TPolymorphicTraits<DataType>::bIsPolymorphic;
+
+	template<typename DataType>
+	concept CIsPolymorphicAndAbstract = TPolymorphicTraits<DataType>::bIsPolymorphic && TPolymorphicTraits<DataType>::bIsAbstract;
+
+	template<typename DataType, typename... ArgTypes>
+	concept CHasFactoryFunction = requires(ArgTypes&&... Args)
+	{
+		{ DataType::CreateNew(Forward<ArgTypes>(Args)...) } -> UE::CDerivedFrom<TSharedRef<DataType, ESPMode::NotThreadSafe>>;
+	};
+
+	// Checks if the Constructor is public and accessible (std::is_constructable does not do that).
+	template <class DataType, class... ArgTypes>
+	concept CHasPublicConstructor =
+	requires(ArgTypes&&... Args)
+	{
+		requires std::is_constructible_v<DataType, ArgTypes...>;
+		::new DataType(Forward<ArgTypes>(Args)...);
+	};
+
+	template<typename DataType>
+	concept CHasFactoryFunctionWithSettings = CHasFactoryFunction<DataType, const FOperatorSettings&>;
+	template<typename DataType, typename... ArgTypes>
+	concept CHasFactoryFunctionWithSettingsAndArgs = CHasFactoryFunction<DataType, const FOperatorSettings&, ArgTypes...>;
+	template<typename DataType, typename... ArgTypes>
+	concept CHasFactoryFunctionWithArgs = CHasFactoryFunction<DataType, ArgTypes...>;
+	template<typename DataType>
+	concept CHasFactoryFunctionDefaultConstruction = CHasFactoryFunction<DataType>;
+
+	template <typename T>
+	concept CHasPolymorphicTypeGetter =
+	requires(const T& t)
+	{
+		{ t.GetDataTypeName() } -> UE::CSameAs<const FName&>;
+	};
+
+	template <typename T>
+	concept CHasPolymorphicTypeIdGetter =
+	requires(const T& t)
+	{
+		{ t.GetDataTypeId() } -> UE::CSameAs<const void*>;
+	};
+
+	template <typename T>
+	concept CHasPolymorphicTypeFunctions = CHasPolymorphicTypeIdGetter<T> && CHasPolymorphicTypeGetter<T>;
+
 	class FPolyRegistry
 	{
 	public:
@@ -63,9 +116,9 @@ namespace Metasound
 			{
 				Info.BaseTypeName = NAME_None;
 			}
-			PRAGMA_DISABLE_INTERNAL_WARNINGS			
+			PRAGMA_DISABLE_INTERNAL_WARNINGS
 			FPolyRegistry::Get().Register(Info);
-			PRAGMA_ENABLE_INTERNAL_WARNINGS			
+			PRAGMA_ENABLE_INTERNAL_WARNINGS
 		}
 	};
 
@@ -74,8 +127,8 @@ namespace Metasound
 
 	UE_EXPERIMENTAL(5.7, "Metasound Polymorphism is still under development")
 	METASOUNDGRAPHCORE_API bool IsCastable(const FName InType, const FName InPotentialBase);
-	
-	
+
+
 }
 
 #define DECLARE_METASOUND_POLY_TYPE(TYPE, BASE, IS_ABSTRACT)\
@@ -83,8 +136,11 @@ namespace Metasound
 	struct Metasound::TPolymorphicTraits<TYPE>\
 	{\
 		static constexpr bool bIsPolymorphic = true;\
-		static constexpr bool bIsAbstract = IS_ABSTRACT;\
+		static constexpr bool bIsAbstract = IS_ABSTRACT || TIsAbstract<TYPE>::Value;\
 		using BaseType = BASE;\
+		/* TODO: Re-enable once all poly types define CreateNew + GetDataTypeName/GetDataTypeId */\
+		/* static_assert((CHasFactoryFunctionWithSettings<TYPE> && bIsAbstract && bIsPolymorphic) || !bIsAbstract, "Abstract Types Must Define a Factory Function 'CreateNew'"); */\
+		/* static_assert(CHasPolymorphicTypeFunctions<TYPE>); */\
 	};
 
 #define REGISTER_METASOUND_POLY_TYPE(T) static Metasound::TPolyRegistrar<T> UE_JOIN(GPolyReg,__COUNTER__);
\ No newline at end of file
diff --git a/Engine/Plugins/Runtime/Metasound/Source/MetasoundGraphCore/Public/MetasoundVariable.h b/Engine/Plugins/Runtime/Metasound/Source/MetasoundGraphCore/Public/MetasoundVariable.h
--- a/Engine/Plugins/Runtime/Metasound/Source/MetasoundGraphCore/Public/MetasoundVariable.h
+++ b/Engine/Plugins/Runtime/Metasound/Source/MetasoundGraphCore/Public/MetasoundVariable.h
@@ -7,6 +7,7 @@
 #include "MetasoundLiteral.h"
 #include "MetasoundOperatorSettings.h"
 #include "MetasoundVertexData.h"
+#include "MetasoundPolymorphic.h"
 
 namespace Metasound
 {
@@ -44,13 +45,23 @@ namespace Metasound
 
 		bool RequiresDelayedDataCopy() const
 		{
-			return DelayedDataReference.IsSet() && DataReference.IsSet();
+			if constexpr (TPolymorphicTraits<DataType>::bIsPolymorphic)
+			{
+				return false; // Polymorphic types do not support delayed data copy.
+			}
+			else
+			{
+				return DelayedDataReference.IsSet() && DataReference.IsSet();
+			}
 		}
 
 		void CopyReferencedData()
 		{
-			check(RequiresDelayedDataCopy());
-			*(*DelayedDataReference) = *(*DataReference);
+			if constexpr (!TPolymorphicTraits<DataType>::bIsPolymorphic)
+			{
+				check(RequiresDelayedDataCopy());
+				*(*DelayedDataReference) = *(*DataReference);
+			}
 		}
 
 		void SetDataReference(FReadReference InDataReference)
@@ -103,10 +114,13 @@ namespace Metasound
 
 		void Reset(const FOperatorSettings& InOperatorSettings)
 		{
-			if (DelayedDataReference.IsSet())
+			if constexpr (!TPolymorphicTraits<DataType>::bIsPolymorphic)
 			{
-				FWriteReference DelayedWritable = *DelayedDataReference;
-				*DelayedWritable = TDataTypeLiteralFactory<DataType>::CreateExplicitArgs(InOperatorSettings, Literal);
+				if (DelayedDataReference.IsSet())
+				{
+					FWriteReference DelayedWritable = *DelayedDataReference;
+					*DelayedWritable = TDataTypeLiteralFactory<DataType>::CreateExplicitArgs(InOperatorSettings, Literal);
+				}
 			}
 		}
 
diff --git a/Engine/Plugins/Runtime/Metasound/Source/MetasoundGraphCore/Public/MetasoundVertexData.h b/Engine/Plugins/Runtime/Metasound/Source/MetasoundGraphCore/Public/MetasoundVertexData.h
--- a/Engine/Plugins/Runtime/Metasound/Source/MetasoundGraphCore/Public/MetasoundVertexData.h
+++ b/Engine/Plugins/Runtime/Metasound/Source/MetasoundGraphCore/Public/MetasoundVertexData.h
@@ -240,7 +240,7 @@ namespace Metasound
 			void Set(const DataReferenceType& InDataReference)
 			{
 				PRAGMA_DISABLE_INTERNAL_WARNINGS
-				check(Metasound::IsCastable(Vertex.DataTypeName, InDataReference.GetDataTypeName()));
+				check(Metasound::IsCastable(InDataReference.GetDataTypeName(), Vertex.DataTypeName));
 				PRAGMA_ENABLE_INTERNAL_WARNINGS
 				
 				Data.Emplace(InDataReference);
