CL D: Frontend registration + channel-agnostic types

Split from CL 51557033. MetasoundDataTypeRegistrationMacro.h changes
guarded by if constexpr. New channel-agnostic type files.
Includes MetasoundRouter.h fix (BH-1: TReceiver::Pop return value).
<enter description here>

#tests LowLevelTests (AudioUnitTests)
#Synced-CL 51594399

diff --git a/Engine/Plugins/Runtime/Metasound/Source/MetasoundFrontend/Private/Interfaces/MetasoundFrontendSourceInterface.cpp b/Engine/Plugins/Runtime/Metasound/Source/MetasoundFrontend/Private/Interfaces/MetasoundFrontendSourceInterface.cpp
--- a/Engine/Plugins/Runtime/Metasound/Source/MetasoundFrontend/Private/Interfaces/MetasoundFrontendSourceInterface.cpp
+++ b/Engine/Plugins/Runtime/Metasound/Source/MetasoundFrontend/Private/Interfaces/MetasoundFrontendSourceInterface.cpp
@@ -248,6 +248,8 @@ namespace Metasound::Frontend
 			const FName TransmitterID = AUDIO_PARAMETER_INTERFACE_MEMBER_DEFINE("TransmitterID");
 			const FName AudioMixerNumOutputFrames = AUDIO_PARAMETER_INTERFACE_MEMBER_DEFINE("AudioMixerNumOutputFrames");
 			const FName AudioMixerSampleRate = AUDIO_PARAMETER_INTERFACE_MEMBER_DEFINE("AudioMixerSampleRate");
+			const FName AudioMixerNumChannels = AUDIO_PARAMETER_INTERFACE_MEMBER_DEFINE("AudioMixerNumChannels");
+			const FName SourceFormatName = AUDIO_PARAMETER_INTERFACE_MEMBER_DEFINE("SourceFormatName");
 		}
 
 		Audio::FParameterInterfacePtr CreateInterface(const UClass& InClass)
@@ -313,6 +315,12 @@ namespace Metasound::Frontend
 							LOCTEXT("AudioMixerSampleRateDescription", "The sample rate of the audio mixer."),
 							FName(), 
 							Environment::AudioMixerSampleRate
+						},
+						{
+							LOCTEXT("SourceFormatNameDisplayName", "Source Format"),
+							LOCTEXT("SourceFormatNameDescription", "The \'Source\' channel agnostic format of the graph"),
+							GetMetasoundDataTypeName<FString>(),	// <-- TODO: Metasound can't support FName parameters, so we must pass as a string.
+							Environment::SourceFormatName
 						}
 					};
 				}
diff --git a/Engine/Plugins/Runtime/Metasound/Source/MetasoundFrontend/Private/MetasoundChannelAgnosticAutoTypes.cpp b/Engine/Plugins/Runtime/Metasound/Source/MetasoundFrontend/Private/MetasoundChannelAgnosticAutoTypes.cpp
new file mode 100644
--- /dev/null
+++ b/Engine/Plugins/Runtime/Metasound/Source/MetasoundFrontend/Private/MetasoundChannelAgnosticAutoTypes.cpp
@@ -0,0 +1,4 @@
+// Copyright Epic Games, Inc. All Rights Reserved.
+
+#include "MetasoundChannelAgnosticAutoTypes.h"
+
diff --git a/Engine/Plugins/Runtime/Metasound/Source/MetasoundFrontend/Private/MetasoundChannelAgnosticDataSchemas.cpp b/Engine/Plugins/Runtime/Metasound/Source/MetasoundFrontend/Private/MetasoundChannelAgnosticDataSchemas.cpp
--- a/Engine/Plugins/Runtime/Metasound/Source/MetasoundFrontend/Private/MetasoundChannelAgnosticDataSchemas.cpp
+++ b/Engine/Plugins/Runtime/Metasound/Source/MetasoundFrontend/Private/MetasoundChannelAgnosticDataSchemas.cpp
@@ -375,20 +375,26 @@ namespace Metasound
 
 			void InitRawAssignmentFunction()
 			{
-				this->RawAssignmentFunction = [](const void* Src, void* Dest)
-					{
-						*(static_cast<FBaseType*>(Dest)) = *(static_cast<const FBaseType*>(Src));
-					};
+				if constexpr (TIsCopyable<CppDataType>::Value)
+				{
+					this->RawAssignmentFunction = [](const void* Src, void* Dest)
+						{
+							*(static_cast<FBaseType*>(Dest)) = *(static_cast<const FBaseType*>(Src));
+						};
+				}
 			}
 
 			void InitLiteralAssignmentFunction()
 			{
-				if  (Info.bIsParsable)
+				if constexpr (TIsCopyable<CppDataType>::Value)
 				{
-					this->LiteralAssignmentFunction = [](const FOperatorSettings& InOperatorSettings, const FLiteral& InLiteral, const FAnyDataReference& OutDataRef)
-						{
-							*OutDataRef.GetWritableValue<FBaseType>() = TDataTypeLiteralFactory<FBaseType>::CreateExplicitArgs(InOperatorSettings, InLiteral);
-						};
+					if (Info.bIsParsable)
+					{
+						this->LiteralAssignmentFunction = [](const FOperatorSettings& InOperatorSettings, const FLiteral& InLiteral, const FAnyDataReference& OutDataRef)
+							{
+								*OutDataRef.GetWritableValue<FBaseType>() = TDataTypeLiteralFactory<FBaseType>::CreateExplicitArgs(InOperatorSettings, InLiteral);
+							};
+					}
 				}
 			}
 
@@ -411,6 +417,8 @@ namespace Metasound
 		
 			virtual TUniquePtr<IDataTypeRegistryEntry> Clone() const override
 			{
+				// The outer reference is cloned, this type is non-copyable.
+				checkNoEntry();
 				return {};
 			}
 		
@@ -692,13 +700,11 @@ namespace Metasound
 																	 const FOperatorSettings& InOperatorSettings) const override
 			{
 				using enum EDataReferenceAccessType;;
-				// Use the CPP types ID. For cat that will mean Derived,Composite,Soundfield.
-				auto ID = GetMetasoundDataTypeId<CppDataType>();
 				switch (InAccessType)
 				{
-				case Read: return  FAnyDataReference(TDataReadReference<FChannelAgnosticType>::CreateNewDerivedAs<CppDataType>(Info.DataTypeName, ID, InOperatorSettings, Info.DataTypeName));
-				case Write: return FAnyDataReference(TDataWriteReference<FChannelAgnosticType>::CreateNewDerivedAs<CppDataType>(Info.DataTypeName, ID, InOperatorSettings, Info.DataTypeName));
-				case Value: return FAnyDataReference(TDataValueReference<FChannelAgnosticType>::CreateNewDerivedAs<CppDataType>(Info.DataTypeName, ID, InOperatorSettings, Info.DataTypeName));
+				case Read: return  FAnyDataReference(TDataReadReference<CppDataType>::CreateNew(InOperatorSettings, Info.DataTypeName));
+				case Write: return FAnyDataReference(TDataWriteReference<CppDataType>::CreateNew(InOperatorSettings, Info.DataTypeName));
+				case Value: return FAnyDataReference(TDataValueReference<CppDataType>::CreateNew(InOperatorSettings, Info.DataTypeName));
 				default:
 					checkNoEntry();
 				}
@@ -905,7 +911,7 @@ namespace Metasound
 		if (FName ProblemNode, ProblemParent; !BuildSchemaLoadOrder(Schemas, Order, ProblemNode, ProblemParent, CyclicNodes))
 		{
 #if !NO_LOGGING
-			UE_CLOG(!ProblemNode.IsNone(), LogMetaSound, Display, TEXT("LoadChannelAgosticDataSchemas: Error failed to find parent '%s' in schema '%s'"), 
+			UE_CLOG(!ProblemNode.IsNone(), LogMetaSound, Error, TEXT("LoadChannelAgosticDataSchemas: Error failed to find parent '%s' in schema '%s'"), 
 				*ProblemParent.ToString(), *ProblemNode.ToString());	
 			
 			// Log each cyclic.
@@ -949,7 +955,7 @@ namespace Metasound
 		LoadedSchemas.Append(LoadSchemas<FCatDiscreteDataSchema,		FMetasoundDiscreteChannelTypeFamily,	FDiscreteChannelAgnosticType>(		SchemaPath / TEXT("Discrete")));
 		LoadedSchemas.Append(LoadSchemas<FCatSoundfieldDataSchema,	FMetasoundSoundfieldChannelTypeFamily,	FSoundfieldChannelAgnosticType>(	SchemaPath / TEXT("Soundfield")));
 		LoadedSchemas.Append(LoadSchemas<FCatCompositeDataSchema,	FMetasoundCompositeChannelTypeFamily,	FCompositeChannelAgnosticType>(		SchemaPath / TEXT("Composite")));
-		UE_LOG(LogMetaSound, Display, TEXT("LoadChannelAgosticDataSchemas: Loaded %d schemas from [%s]"),LoadedSchemas.Num(), *SchemaPath);	
+		UE_LOG(LogMetaSound, Error, TEXT("LoadChannelAgosticDataSchemas: Loaded %d schemas from [%s]"),LoadedSchemas.Num(), *SchemaPath);	
 		
 		return LoadedSchemas;
 	}
diff --git a/Engine/Plugins/Runtime/Metasound/Source/MetasoundFrontend/Private/MetasoundChannelAgnosticType.cpp b/Engine/Plugins/Runtime/Metasound/Source/MetasoundFrontend/Private/MetasoundChannelAgnosticType.cpp
--- a/Engine/Plugins/Runtime/Metasound/Source/MetasoundFrontend/Private/MetasoundChannelAgnosticType.cpp
+++ b/Engine/Plugins/Runtime/Metasound/Source/MetasoundFrontend/Private/MetasoundChannelAgnosticType.cpp
@@ -4,6 +4,7 @@
 #include "TypeFamily/TypeFamily.h"
 #include "TypeFamily/ChannelTypeFamily.h"
 #include "MetasoundPolymorphic.h"
+#include "Interfaces/IPluginManager.h"
 
 namespace Metasound
 {
@@ -16,34 +17,133 @@ namespace Metasound
 		static const FName NAME_Discrete(TEXT("Cat:Discrete"));
 	}
 	
-	/**
-	 * Returns the default (i.e. default constructed) format of a CAT. Static.
-	 * @return Name to Pass to Registry. 
-	 */
-	FName FChannelAgnosticType::GetDefaultCatFormat()
+	FName FChannelAgnosticType::GetCatDefaultCatFormat()
+	{
+		return FDiscreteChannelAgnosticType::GetDefaultCatFormat();
+	}
+
+	// Abstract factories.
+	TSharedRef<FChannelAgnosticType, ESPMode::NotThreadSafe> FChannelAgnosticType::CreateNew(const FOperatorSettings& InSettings, const FName InDerivedType)
+	{
+		// Clunky, but for now keep it simple.
+		using namespace Audio;
+		auto& Type = GetChannelRegistry().FindConcreteChannelChecked(InDerivedType);
+		if (Type.Cast<FCompositeChannelTypeFamily>())
+		{
+			return FCompositeChannelAgnosticType::CreateNew(InSettings, InDerivedType);
+		}
+		if (Type.Cast<FSoundfieldChannelTypeFamily>())
+		{
+			return FSoundfieldChannelAgnosticType::CreateNew(InSettings, InDerivedType);
+		}
+		check(Type.Cast<FDiscreteChannelTypeFamily>());
+		return FDiscreteChannelAgnosticType::CreateNew(InSettings, InDerivedType);
+	}
+	TSharedRef<FDiscreteChannelAgnosticType, ESPMode::NotThreadSafe> FDiscreteChannelAgnosticType::CreateNew(const FOperatorSettings& InSettings, const FName InDerivedType)
+	{
+		return MakeShared<FDiscreteChannelAgnosticType, ESPMode::NotThreadSafe>(InSettings, InDerivedType);
+	}
+	TSharedRef<FCompositeChannelAgnosticType, ESPMode::NotThreadSafe> FCompositeChannelAgnosticType::CreateNew(const FOperatorSettings& InSettings, const FName InDerivedType)
+	{
+		return MakeShared<FCompositeChannelAgnosticType, ESPMode::NotThreadSafe>(InSettings, InDerivedType); 	
+	}
+	TSharedRef<FSoundfieldChannelAgnosticType, ESPMode::NotThreadSafe> FSoundfieldChannelAgnosticType::CreateNew(const FOperatorSettings& InSettings, const FName InDerivedType)
+	{
+		return MakeShared<FSoundfieldChannelAgnosticType, ESPMode::NotThreadSafe>(InSettings, InDerivedType); 	
+	}		
+
+	FName FSoundfieldChannelAgnosticType::GetDefaultCatFormat()
+	{
+		static const FName DefaultFormat(TEXT("Cat:AmbisonicsOrder2"));
+		return DefaultFormat;
+	}
+
+	FName FCompositeChannelAgnosticType::GetDefaultCatFormat()
+	{
+		static const FName DefaultFormat(TEXT("Cat:CompositePlaceholder"));	
+		return DefaultFormat;
+	}
+	
+	FName FDiscreteChannelAgnosticType::GetDefaultCatFormat()
 	{
 		static const FName DefaultFormat(TEXT("Cat:Mono1Dot0"));
 		if (Audio::GetChannelRegistry().FindConcreteChannel(DefaultFormat) == nullptr)
 		{
-			UE_LOG(LogMetaSound, Warning, TEXT("The default channel type '%s' is not registered. Calling EnsureDefaultTypesRegistered() "), 
-				*DefaultFormat.ToString());
+			UE_LOG(LogMetaSound, Warning, TEXT("The default channel type '%s' is not registered. Calling EnsureDefaultTypesRegistered() "), *DefaultFormat.ToString());
 			ChannelAgnosticTypePrivate::EnsureDefaultTypesRegistered();
 		}
 		return DefaultFormat;
 	}
 
+	// Factory helper.	
+	TOptional<TDataWriteReference<FChannelAgnosticType>> FChannelAgnosticType::CreateWriteReference(const FName InDerivedType, const FOperatorSettings& InOpSettings)
+	{
+		using namespace Frontend;
+		const FLiteral Literal(InDerivedType.ToString());
+		if (FDataTypeRegistryInfo Info; IDataTypeRegistry::Get().GetDataTypeInfo(InDerivedType, Info) && Info.bIsAbstract && Info.bIsPolymorphic)
+		{
+			// Error: Trying to create an abstract type!
+			return {};
+		}
+		const TOptional<FAnyDataReference> Ref = IDataTypeRegistry::Get().CreateDataReference(InDerivedType, EDataReferenceAccessType::Write, Literal, InOpSettings);
+		if (Ref.IsSet())
+		{
+			return Ref->GetAs<TDataWriteReference<FChannelAgnosticType>>();
+		}
+		return {};
+	}
+
+	
+
+
 	FChannelAgnosticType::FChannelAgnosticType(const FOperatorSettings& InSettings, const FName& InChannelTypeName)
 		: FChannelAgnosticType(InSettings.GetNumFramesPerBlock(), InChannelTypeName)
 	{}
-	
+
+	TOptional<TDataWriteReference<FDiscreteChannelAgnosticType>> FDiscreteChannelAgnosticType::CreateWriteReference(const FName InDerivedType, const FOperatorSettings& InOpSettings)
+	{
+		using namespace Frontend;
+		const FLiteral Literal(InDerivedType.ToString());
+		if (FDataTypeRegistryInfo Info; IDataTypeRegistry::Get().GetDataTypeInfo(InDerivedType, Info) && Info.bIsAbstract && Info.bIsPolymorphic)
+		{
+			// Error: Trying to create an abstract type!
+			return {};
+		}
+		const TOptional<FAnyDataReference> Ref = IDataTypeRegistry::Get().CreateDataReference(InDerivedType, EDataReferenceAccessType::Write, Literal, InOpSettings);
+		if (Ref.IsSet())
+		{
+			return Ref->GetAs<TDataWriteReference<FDiscreteChannelAgnosticType>>();
+		}
+		return {};
+	}
+
 	FChannelAgnosticType::FChannelAgnosticType(const int32 InNumFramesPerBlock, const FName& InChannelTypeName)
 		: Super(Audio::GetChannelRegistry().FindConcreteChannelChecked(InChannelTypeName), InNumFramesPerBlock)
 	{}
 	
+	static_assert(CHasPolymorphicTypeGetter<FChannelAgnosticType> && TPolymorphicTraits<FChannelAgnosticType>::bIsPolymorphic);
+	static_assert(CHasFactoryFunction<FChannelAgnosticType, const FOperatorSettings&, const FName>);
+	static_assert(CHasFactoryFunction<FDiscreteChannelAgnosticType, const FOperatorSettings&, const FName>);
+	static_assert(CHasFactoryFunction<FCompositeChannelAgnosticType, const FOperatorSettings&, const FName>);
+	static_assert(CHasFactoryFunction<FSoundfieldChannelAgnosticType, const FOperatorSettings&, const FName>);
+
 	// Register.
+	static_assert(CHasFactoryFunctionWithSettingsAndArgs<FChannelAgnosticType, const FName&>);
+
+	static_assert(TEnableConstructorVertex<FChannelAgnosticType>::Value);
+	static_assert(TLiteralTraits<FChannelAgnosticType>::bIsParsableFromAnyLiteralType);
+	
+	static_assert(::Metasound::TInputNode<FChannelAgnosticType, EVertexAccessType::Reference>::bCanRegister);
+	static_assert(::Metasound::TInputNode<FChannelAgnosticType, EVertexAccessType::Reference>::bCanRegister);
+	static_assert(::Metasound::Frontend::TMetasoundDataTypeRegistration<FChannelAgnosticType>::bCanRegister);
+	
 	REGISTER_METASOUND_DATATYPE(FChannelAgnosticType, "Cat");
 	REGISTER_METASOUND_POLY_TYPE(FChannelAgnosticType);
 
+	static_assert(TEnableConstructorVertex<FSoundfieldChannelAgnosticType>::Value);
+	static_assert(CHasFactoryFunctionWithSettings<FSoundfieldChannelAgnosticType>);
+	static_assert(TLiteralTraits<FSoundfieldChannelAgnosticType>::bIsParsableFromAnyLiteralType);
+
 	REGISTER_METASOUND_DATATYPE(FCompositeChannelAgnosticType, "Cat:Composite");
 	REGISTER_METASOUND_POLY_TYPE(FCompositeChannelAgnosticType);
 
@@ -104,6 +204,12 @@ namespace Metasound
 			}
 		}
 	}
+	
+	bool IsExperimentalPluginEnabled()
+	{
+		const TSharedPtr<IPlugin> Plugin = IPluginManager::Get().FindEnabledPlugin(TEXT("MetasoundExperimental"));
+		return Plugin.IsValid();
+	}
 }
 
 
diff --git a/Engine/Plugins/Runtime/Metasound/Source/MetasoundFrontend/Public/Interfaces/MetasoundFrontendSourceInterface.h b/Engine/Plugins/Runtime/Metasound/Source/MetasoundFrontend/Public/Interfaces/MetasoundFrontendSourceInterface.h
--- a/Engine/Plugins/Runtime/Metasound/Source/MetasoundFrontend/Public/Interfaces/MetasoundFrontendSourceInterface.h
+++ b/Engine/Plugins/Runtime/Metasound/Source/MetasoundFrontend/Public/Interfaces/MetasoundFrontendSourceInterface.h
@@ -82,6 +82,8 @@ namespace Metasound
 				METASOUNDFRONTEND_API const extern FName TransmitterID;
 				METASOUNDFRONTEND_API const extern FName AudioMixerNumOutputFrames;
 				METASOUNDFRONTEND_API const extern FName AudioMixerSampleRate;
+				METASOUNDFRONTEND_API const extern FName AudioMixerNumChannels;
+				METASOUNDFRONTEND_API const extern FName SourceFormatName;
 			}
 
 			METASOUNDFRONTEND_API const FMetasoundFrontendVersion& GetVersion();
diff --git a/Engine/Plugins/Runtime/Metasound/Source/MetasoundFrontend/Public/MetasoundChannelAgnosticAutoTypes.h b/Engine/Plugins/Runtime/Metasound/Source/MetasoundFrontend/Public/MetasoundChannelAgnosticAutoTypes.h
new file mode 100644
--- /dev/null
+++ b/Engine/Plugins/Runtime/Metasound/Source/MetasoundFrontend/Public/MetasoundChannelAgnosticAutoTypes.h
@@ -0,0 +1,25 @@
+// Copyright Epic Games, Inc. All Rights Reserved.
+
+#pragma once
+#include "MetasoundChannelAgnosticType.h"
+
+#include "MetasoundChannelAgnosticAutoTypes.generated.h"
+
+#define UE_API METASOUNDFRONTEND_API
+
+UENUM()
+enum class EMetasoundChannelAgnosticSourceFormatChooser : uint8
+{
+	Auto,
+	Custom
+};
+
+UENUM()
+enum class EMetasoundChannelAgnosticNodeFormatChooser : uint8
+{
+	Auto,
+	Source,
+	Custom
+};
+
+#undef UE_API 
\ No newline at end of file
diff --git a/Engine/Plugins/Runtime/Metasound/Source/MetasoundFrontend/Public/MetasoundChannelAgnosticType.h b/Engine/Plugins/Runtime/Metasound/Source/MetasoundFrontend/Public/MetasoundChannelAgnosticType.h
--- a/Engine/Plugins/Runtime/Metasound/Source/MetasoundFrontend/Public/MetasoundChannelAgnosticType.h
+++ b/Engine/Plugins/Runtime/Metasound/Source/MetasoundFrontend/Public/MetasoundChannelAgnosticType.h
@@ -6,96 +6,142 @@
 #include "ChannelAgnostic/ChannelAgnosticType.h"
 #include "MetasoundDataTypeRegistrationMacro.h"
 
+#define UE_API METASOUNDFRONTEND_API
+
 namespace Metasound
 {
 	class FChannelAgnosticType : public Audio::FChannelAgnosticType
 	{
 	public:
+		UE_NONCOPYABLE(FChannelAgnosticType)
 		using Super = Audio::FChannelAgnosticType;
-		using Super::Super;
 
 		static FName GetRegisteredTypeName()
 		{
 			static const FName Name("Cat");			
 			return Name;
-		}						
-
-		/**
-		 * Returns the default (i.e. default constructed) format of a CAT.
-		 * @return Name to Pass to Registry. 
-		 */
-		UE_EXPERIMENTAL(5.8, "Channel Angostic Types are Experimental")
-		METASOUNDFRONTEND_API static FName GetDefaultCatFormat();
-
-		UE_EXPERIMENTAL(5.8, "Channel Angostic Types are Experimental")
-		METASOUNDFRONTEND_API explicit FChannelAgnosticType(const FOperatorSettings& InSettings, const FName& InChannelTypeName = GetDefaultCatFormat());
-		
-		UE_EXPERIMENTAL(5.8, "Channel Angostic Types are Experimental")
-		METASOUNDFRONTEND_API explicit FChannelAgnosticType(const int32 InNumFramesPerBlock, const FName& InChannelTypeName = GetDefaultCatFormat());
+		}				
 		
-		static FName LiteralToFName(const FString& In)
-		{
-			if (In.IsEmpty())
-			{
-				static const FName Default = TEXT("Cat:Mono1Dot0");
-				return Default;
-			}
-			return FName(In);
-		}
-		/* FLiteral compatible Constructor
-		 * TODO: Literals only support FString and not FNames, so there is a conversion here.
-		 */
 		PRAGMA_DISABLE_EXPERIMENTAL_WARNINGS		
-		UE_EXPERIMENTAL(5.8, "Channel Angostic Types are Experimental")
-		explicit FChannelAgnosticType(const FOperatorSettings& InSettings, const FString& InChannelTypeName)
-			: FChannelAgnosticType(InSettings, LiteralToFName((InChannelTypeName)))
-		{}
+				
+		// Metasound Polymorphic Type Id+Name. (required for intrusive references to work).
+		const FName& GetDataTypeName() const { return GetType().GetName(); }
+		const void* GetDataTypeId() const { return &GetType(); }
+		
+		// Factory helper.
+		UE_EXPERIMENTAL(5.8, "Channel Agnostic Types are Experimental")
+		UE_API static TOptional<TDataWriteReference<FChannelAgnosticType>> CreateWriteReference(const FName InDerivedType, const FOperatorSettings& InOpSettings);
+	
+		UE_API static FName GetCatDefaultCatFormat();
+		
+		UE_API static TSharedRef<FChannelAgnosticType, ESPMode::NotThreadSafe> CreateNew(const FOperatorSettings& InSettings, const FName InDerivedType=GetCatDefaultCatFormat());
+public:
+		template <typename ObjectType, ESPMode Mode>
+		friend class SharedPointerInternals::TIntrusiveReferenceController;
+		
+		UE_API FChannelAgnosticType(const int32 InNumFramesPerBlock, const FName& InChannelTypeName);
+		UE_API FChannelAgnosticType(const FOperatorSettings& InSettings, const FName& InChannelTypeName);
+	private:
 		PRAGMA_ENABLE_EXPERIMENTAL_WARNINGS
 	};
 
 	class FDiscreteChannelAgnosticType : public FChannelAgnosticType
 	{
+	public:
+		UE_NONCOPYABLE(FDiscreteChannelAgnosticType)
+		
 		static FName GetRegisteredTypeName()
 		{
 			static const FName Name("Cat:Discrete");			
 			return Name;
 		}				
 		using Super = FChannelAgnosticType;
-		using Super::Super;
+			
+		UE_EXPERIMENTAL(5.8, "Channel Agnostic Types are Experimental")
+		UE_API static TOptional<TDataWriteReference<FDiscreteChannelAgnosticType>> CreateWriteReference(const FName InDerivedType, const FOperatorSettings& InOpSettings);
+		
+		UE_API static FName GetDefaultCatFormat();
+		
+		// Factory function for the default case.
+		UE_API static TSharedRef<FDiscreteChannelAgnosticType, ESPMode::NotThreadSafe> CreateNew(const FOperatorSettings& InSettings, const FName InDerivedType=GetDefaultCatFormat());
+		
+	public:
+		template <typename ObjectType, ESPMode Mode>
+		friend class SharedPointerInternals::TIntrusiveReferenceController;
+		
+		FDiscreteChannelAgnosticType(const FOperatorSettings& InSettings, const FName InDerivedType)
+			: Super(InSettings, InDerivedType) {}
 	};
 	class FSoundfieldChannelAgnosticType : public FChannelAgnosticType
 	{
+	public:
+		UE_NONCOPYABLE(FSoundfieldChannelAgnosticType)
+		
 		static FName GetRegisteredTypeName()
 		{
 			static const FName Name("Cat:Soundfield");			
 			return Name;
 		}		
 		using Super = FChannelAgnosticType;
-		using Super::Super;
+		
+		UE_API static FName GetDefaultCatFormat();
+		UE_API static TSharedRef<FSoundfieldChannelAgnosticType, ESPMode::NotThreadSafe> CreateNew(const FOperatorSettings& InSettings, const FName InDerivedType=GetDefaultCatFormat());
+		
+	public:
+		FSoundfieldChannelAgnosticType(const FOperatorSettings& InSettings, const FName InDerivedType)
+			: Super(InSettings, InDerivedType) 
+		{}
+		
+		template <typename ObjectType, ESPMode Mode>
+		friend class SharedPointerInternals::TIntrusiveReferenceController;
 	};
 	class FCompositeChannelAgnosticType : public FChannelAgnosticType
 	{
+	public:
+		UE_NONCOPYABLE(FCompositeChannelAgnosticType)
+		
 		static FName GetRegisteredTypeName()
 		{
 			static const FName Name("Cat:Composite");			
 			return Name;
 		}
 		using Super = FChannelAgnosticType;
-		using Super::Super;
+		
+		UE_API static FName GetDefaultCatFormat();
+		UE_API static TSharedRef<FCompositeChannelAgnosticType, ESPMode::NotThreadSafe> CreateNew(const FOperatorSettings& InSettings, const FName InDerivedType=GetDefaultCatFormat());
+	public:
+		FCompositeChannelAgnosticType(const FOperatorSettings& InSettings, const FName InDerivedType)
+			: Super(InSettings, InDerivedType) 
+		{}	
+		
+		template <typename ObjectType, ESPMode Mode>
+		friend class SharedPointerInternals::TIntrusiveReferenceController;
 	};
+	
 }
 
 namespace Metasound
 {
 	// Declare Metasound Data reference Types. 
-	DECLARE_METASOUND_DATA_REFERENCE_TYPES(Metasound::FChannelAgnosticType,				METASOUNDFRONTEND_API, FChannelAgnosticTypeTypeInfo,				FChannelAgnosticTypeReadRef,			FChannelAgnosticTypeWriteRef);
-	DECLARE_METASOUND_DATA_REFERENCE_TYPES(Metasound::FCompositeChannelAgnosticType,	METASOUNDFRONTEND_API, FCompositeChannelAgnosticTypeTypeInfo,	FCompositeChannelAgnosticTypeReadRef,	FCompositeChannelAgnosticTypeWriteRef);
-	DECLARE_METASOUND_DATA_REFERENCE_TYPES(Metasound::FDiscreteChannelAgnosticType,		METASOUNDFRONTEND_API, FDiscreteChannelAgnosticTypeTypeInfo,		FDiscreteChannelAgnosticTypeReadRef,	FDiscreteChannelAgnosticTypeWriteRef);
-	DECLARE_METASOUND_DATA_REFERENCE_TYPES(Metasound::FSoundfieldChannelAgnosticType,	METASOUNDFRONTEND_API, FSoundfieldChannelAgnosticTypeTypeInfo,	FSoundfieldChannelAgnosticTypeReadRef,	FSoundChannelAgnosticTypeWriteRef);
+	DECLARE_METASOUND_DATA_REFERENCE_TYPES(Metasound::FChannelAgnosticType,				UE_API, FChannelAgnosticTypeTypeInfo,				FChannelAgnosticTypeReadRef,			FChannelAgnosticTypeWriteRef);
+	DECLARE_METASOUND_DATA_REFERENCE_TYPES(Metasound::FCompositeChannelAgnosticType,	UE_API, FCompositeChannelAgnosticTypeTypeInfo,	FCompositeChannelAgnosticTypeReadRef,	FCompositeChannelAgnosticTypeWriteRef);
+	DECLARE_METASOUND_DATA_REFERENCE_TYPES(Metasound::FDiscreteChannelAgnosticType,		UE_API, FDiscreteChannelAgnosticTypeTypeInfo,		FDiscreteChannelAgnosticTypeReadRef,	FDiscreteChannelAgnosticTypeWriteRef);
+	DECLARE_METASOUND_DATA_REFERENCE_TYPES(Metasound::FSoundfieldChannelAgnosticType,	UE_API, FSoundfieldChannelAgnosticTypeTypeInfo,	FSoundfieldChannelAgnosticTypeReadRef,	FSoundChannelAgnosticTypeWriteRef);
 
 	// Declare Polymorphic DataTypes.
-	DECLARE_METASOUND_POLY_TYPE(FChannelAgnosticType,			void, true);
-	DECLARE_METASOUND_POLY_TYPE(FCompositeChannelAgnosticType,	Metasound::FChannelAgnosticType, true);
-	DECLARE_METASOUND_POLY_TYPE(FDiscreteChannelAgnosticType,	Metasound::FChannelAgnosticType, true);
-	DECLARE_METASOUND_POLY_TYPE(FSoundfieldChannelAgnosticType,	Metasound::FChannelAgnosticType, true);
-}
\ No newline at end of file
+	DECLARE_METASOUND_POLY_TYPE(Metasound::FChannelAgnosticType,			void, true);
+	DECLARE_METASOUND_POLY_TYPE(Metasound::FCompositeChannelAgnosticType,	Metasound::FChannelAgnosticType, true);
+	DECLARE_METASOUND_POLY_TYPE(Metasound::FDiscreteChannelAgnosticType,	Metasound::FChannelAgnosticType, true);
+	DECLARE_METASOUND_POLY_TYPE(Metasound::FSoundfieldChannelAgnosticType,	Metasound::FChannelAgnosticType, true);
+	
+	// These are bloaty, so turn off for now.
+	DISABLE_METASOUND_ARRAYS_FOR_TYPE(FChannelAgnosticType)
+	DISABLE_METASOUND_ARRAYS_FOR_TYPE(FCompositeChannelAgnosticType)
+	DISABLE_METASOUND_ARRAYS_FOR_TYPE(FDiscreteChannelAgnosticType)
+	DISABLE_METASOUND_ARRAYS_FOR_TYPE(FSoundfieldChannelAgnosticType)
+
+	// Public helper.
+	UE_API bool IsExperimentalPluginEnabled();
+}
+
+#undef UE_API
\ No newline at end of file
diff --git a/Engine/Plugins/Runtime/Metasound/Source/MetasoundFrontend/Public/MetasoundDataTypeRegistrationMacro.h b/Engine/Plugins/Runtime/Metasound/Source/MetasoundFrontend/Public/MetasoundDataTypeRegistrationMacro.h
--- a/Engine/Plugins/Runtime/Metasound/Source/MetasoundFrontend/Public/MetasoundDataTypeRegistrationMacro.h
+++ b/Engine/Plugins/Runtime/Metasound/Source/MetasoundFrontend/Public/MetasoundDataTypeRegistrationMacro.h
@@ -9,6 +9,7 @@
 #include "MetasoundConverterNodeRegistrationMacro.h"
 #include "MetasoundDataFactory.h"
 #include "MetasoundDataReference.h"
+#include "MetasoundPolymorphic.h"
 #include "MetasoundDataReferenceMacro.h"
 #include "MetasoundEnum.h"
 #include "MetasoundFrontendDataTypeRegistry.h"
@@ -430,7 +431,7 @@ namespace Metasound::Frontend
 
 				void InitLiteralAssignmentFunction()
 				{
-					if constexpr (std::is_copy_assignable_v<TDataType> && bIsParsable)
+					if constexpr (std::is_copy_assignable_v<TDataType> && bIsParsable && !std::is_abstract<TDataType>::value)
 					{
 						this->LiteralAssignmentFunction = [](const FOperatorSettings& InOperatorSettings, const FLiteral& InLiteral, const FAnyDataReference& OutDataRef)
 						{
@@ -735,7 +736,7 @@ namespace Metasound::Frontend
 
 				virtual TOptional<FAnyDataReference> CreateDataReference(EDataReferenceAccessType InAccessType, const FLiteral& InLiteral, const FOperatorSettings& InOperatorSettings) const override
 				{
-					if constexpr (bIsParsable)
+					if constexpr (bIsParsable && !std::is_abstract<TDataType>::value)
 					{
 						switch (InAccessType)
 						{
@@ -800,12 +801,12 @@ namespace Metasound::Frontend
 		 * @return True on success, false on failure.
 		 */
 		template<typename TDataType, ELiteralType PreferredArgType>
-		bool RegisterDataTypeArrayWithFrontend(const FModuleInfo& InModuleInfo)
+		bool RegisterDataTypeArrayWithFrontend(const FModuleInfo& InModuleInfo) 
 		{
 			using namespace MetasoundDataTypeRegistrationPrivate;
 			using TArrayType = TArray<TDataType>;
 
-			if constexpr (TEnableAutoArrayTypeRegistration<TDataType>::Value)
+			if constexpr (TEnableAutoArrayTypeRegistration<TDataType>::Value && !CIsPolymorphicAndAbstract<TDataType>)
 			{
 				bool bSuccess = RegisterDataTypeWithFrontendInternal<TArrayType, TLiteralArrayEnum<PreferredArgType>::Value>(InModuleInfo);
 				bSuccess = bSuccess && RegisterArrayNodes<TArrayType>(InModuleInfo);
@@ -823,7 +824,7 @@ namespace Metasound::Frontend
 		{
 			using namespace MetasoundDataTypeRegistrationPrivate;
 
-			if constexpr (TEnableAutoArrayTypeRegistration<TDataType>::Value)
+			if constexpr (TEnableAutoArrayTypeRegistration<TDataType>::Value && !CIsPolymorphicAndAbstract<TDataType>)
 			{
 				using TArrayType = TArray<TDataType>;
 
@@ -976,4 +977,13 @@ namespace Metasound
 	DEFINE_METASOUND_DATA_TYPE(DataType, DataTypeName); \
 	ENQUEUE_METASOUND_DATATYPE_REGISTRATION_COMMAND(DataType, DataTypeName, ##__VA_ARGS__)
 
+// Helper macro for disabling auto array registration
+#define DISABLE_METASOUND_ARRAYS_FOR_TYPE(TYPENAME)\
+	template<>\
+	struct TEnableAutoArrayTypeRegistration<TYPENAME>\
+	{\
+		static constexpr bool Value = false;\
+	};
+
+
 #undef UE_API
diff --git a/Engine/Plugins/Runtime/Metasound/Source/MetasoundFrontend/Public/MetasoundReceiveNode.h b/Engine/Plugins/Runtime/Metasound/Source/MetasoundFrontend/Public/MetasoundReceiveNode.h
--- a/Engine/Plugins/Runtime/Metasound/Source/MetasoundFrontend/Public/MetasoundReceiveNode.h
+++ b/Engine/Plugins/Runtime/Metasound/Source/MetasoundFrontend/Public/MetasoundReceiveNode.h
@@ -142,7 +142,10 @@ namespace Metasound
 
 				if (bHasNotReceivedData)
 				{
-					*OutputData = *DefaultData;
+					if constexpr (std::is_copy_assignable<TDataType>::value)
+					{
+						*OutputData = *DefaultData;
+					}
 					bHasNewData = true;
 				}
 
@@ -154,7 +157,10 @@ namespace Metasound
 
 			void Reset(const IOperator::FResetParams& InParams)
 			{
-				*OutputData = *DefaultData;
+				if constexpr (std::is_copy_assignable<TDataType>::value)
+				{
+					*OutputData = *DefaultData;
+				}
 				bHasNotReceivedData = true;
 			}
 
diff --git a/Engine/Plugins/Runtime/Metasound/Source/MetasoundFrontend/Public/MetasoundRouter.h b/Engine/Plugins/Runtime/Metasound/Source/MetasoundFrontend/Public/MetasoundRouter.h
--- a/Engine/Plugins/Runtime/Metasound/Source/MetasoundFrontend/Public/MetasoundRouter.h
+++ b/Engine/Plugins/Runtime/Metasound/Source/MetasoundFrontend/Public/MetasoundRouter.h
@@ -523,7 +523,7 @@ namespace Metasound
 	template<typename TDataType>
 	class TReceiver : public IReceiver
 	{
-		static_assert(std::is_copy_constructible<TDataType>::value, "The generic TDataChannel requires the DataType it is specialized with to be copy constructible.");
+	
 	public:
 
 		bool CanPop() const
@@ -546,9 +546,16 @@ namespace Metasound
 		{
 			DataChannel->PopOpaque(*this);
 
-			check(Payload.IsValid());
-			OutValue = *Payload;
-			return true;
+			if constexpr (std::is_copy_assignable<TDataType>::value)
+			{
+				check(Payload.IsValid());
+				OutValue = *Payload;
+				return true;
+			}
+			else
+			{
+				return false;
+			}
 		}
 
 		static FName GetDataTypeName()
@@ -591,7 +598,7 @@ namespace Metasound
 	template<typename TDataType>
 	class TCopyableDataChannel : public IDataChannel
 	{
-		static_assert(std::is_copy_constructible<TDataType>::value, "The generic TDataChannel requires the DataType it is specialized with to be copy constructible.");
+		
 
 	public:
 
diff --git a/Engine/Plugins/Runtime/Metasound/Source/MetasoundGraphCore/Public/MetasoundDataReference.h b/Engine/Plugins/Runtime/Metasound/Source/MetasoundGraphCore/Public/MetasoundDataReference.h
--- a/Engine/Plugins/Runtime/Metasound/Source/MetasoundGraphCore/Public/MetasoundDataReference.h
+++ b/Engine/Plugins/Runtime/Metasound/Source/MetasoundGraphCore/Public/MetasoundDataReference.h
@@ -235,7 +235,6 @@ namespace Metasound
 		TDataReference(EDataRefShouldConstruct InToken, ArgTypes&&... Args) requires(CIsPolymorphicAndAbstract<DataType> && CHasFactoryFunction<DataType, ArgTypes...>)
 			: ObjectReference(DataType::CreateNew(Forward<ArgTypes>(Args)...))
 		{
-			/* TODO: Re-enable once all poly types define CreateNew */ // static_assert(CHasFactoryFunction<DataType, ArgTypes...>);
 		}
 
 		typedef TSharedRef<DataType, ESPMode::NotThreadSafe> FRefType;
@@ -270,7 +269,6 @@ namespace Metasound
 		template <typename... ArgTypes>
 		static TDataReference<DataType> CreateNew(ArgTypes&&... Args) requires (CIsPolymorphicAndAbstract<DataType> && CHasFactoryFunction<DataType, ArgTypes...>)
 		{
-			/* TODO: Re-enable once all poly types define CreateNew */ // static_assert(CHasFactoryFunction<DataType, ArgTypes...>);
 			return TDataReference<DataType>(EDataRefShouldConstruct::NewObject, Forward<ArgTypes>(Args)...);
 		}
 
@@ -411,7 +409,6 @@ namespace Metasound
 		template <typename... ArgTypes>
 		static TDataValueReference<DataType> CreateNew(ArgTypes&&... Args) requires (CIsPolymorphicAndAbstract<DataType> && CHasFactoryFunction<DataType, ArgTypes...>)
 		{
-			/* TODO: Re-enable once all poly types define CreateNew */ // static_assert(CHasFactoryFunction<DataType, ArgTypes...>);
 			return TDataValueReference<DataType>(EDataRefShouldConstruct::NewObject, Forward<ArgTypes>(Args)...);
 		}
 
@@ -558,7 +555,6 @@ namespace Metasound
 		template <typename... ArgTypes>
 		static TDataWriteReference<DataType> CreateNew(ArgTypes&&... Args) requires (CIsPolymorphicAndAbstract<DataType> && CHasFactoryFunction<DataType, ArgTypes...>)
 		{
-			/* TODO: Re-enable once all poly types define CreateNew */ // static_assert(CHasFactoryFunction<DataType, ArgTypes...>);
 			return TDataWriteReference<DataType>(EDataRefShouldConstruct::NewObject, Forward<ArgTypes>(Args)...);
 		}
 
@@ -688,7 +684,6 @@ namespace Metasound
 		template <typename... ArgTypes>
 		static TDataReadReference<DataType> CreateNew(ArgTypes&&... Args) requires (CIsPolymorphicAndAbstract<DataType> && CHasFactoryFunction<DataType, ArgTypes...>)
 		{
-			/* TODO: Re-enable once all poly types define CreateNew */ // static_assert(CHasFactoryFunction<DataType, ArgTypes...>);
 			return TDataReadReference<DataType>(EDataRefShouldConstruct::NewObject, Forward<ArgTypes>(Args)...);
 		}
 
