Add FindDiscreteFormatFromNumChannels utility to FChannelAgnosticUtils for looking up discrete channel type families by channel count
Add Metasound channel-agnostic auto format support and wire up AudioChannelAgnosticCore dependency
Add engine and audio unit tests for the new functionality

#tests EngineTest and AudioUnitTests
#Synced-CL 51594399

diff --git a/Engine/Plugins/Runtime/Metasound/Source/MetasoundEngine/MetasoundEngine.Build.cs b/Engine/Plugins/Runtime/Metasound/Source/MetasoundEngine/MetasoundEngine.Build.cs
--- a/Engine/Plugins/Runtime/Metasound/Source/MetasoundEngine/MetasoundEngine.Build.cs
+++ b/Engine/Plugins/Runtime/Metasound/Source/MetasoundEngine/MetasoundEngine.Build.cs
@@ -26,7 +26,8 @@ namespace UnrealBuildTool.Rules
 					"AudioMixer",
 					"MetasoundGraphCore",
 					"MetasoundGenerator",
-					"SignalProcessing"
+					"SignalProcessing",
+					"AudioChannelAgnosticCore"
 				}
 			);
 
diff --git a/Engine/Plugins/Runtime/Metasound/Source/MetasoundEngine/Private/MetasoundChannelAgnosticAutoFormat.cpp b/Engine/Plugins/Runtime/Metasound/Source/MetasoundEngine/Private/MetasoundChannelAgnosticAutoFormat.cpp
new file mode 100644
--- /dev/null
+++ b/Engine/Plugins/Runtime/Metasound/Source/MetasoundEngine/Private/MetasoundChannelAgnosticAutoFormat.cpp
@@ -0,0 +1,90 @@
+// Copyright Epic Games, Inc. All Rights Reserved.
+
+#include "MetasoundChannelAgnosticAutoFormat.h"
+
+#include "ChannelAgnostic/ChannelAgnosticTypeUtils.h"
+
+#include "MetasoundLog.h"
+#include "Sound/ISoundWaveContainer.h"
+
+#include "Sound/SoundWave.h"
+
+namespace Metasound
+{
+	int32 FChannelAgnosticAutoFormatHelper::ComputeAutoChannelCountFromProxies(const TArray<FSoundWaveProxyPtr>& InProxies)
+	{
+		// Greatest for now.
+		if (const TSharedPtr<FSoundWaveProxy>* Greatest = Algo::MaxElementBy(InProxies, [](const TSharedPtr<FSoundWaveProxy>& j) -> int32
+			{
+				return j->GetNumChannels();
+			}))
+		{
+			return (*Greatest)->GetNumChannels();
+		}
+		
+		// Fail.
+		return INDEX_NONE;
+	}
+
+	TSharedPtr<const Audio::FChannelTypeFamily> FChannelAgnosticAutoFormatHelper::ComputeAutoFormatFromWaveContainer(
+		const TSharedPtr<const Audio::IProxyData>& InData)
+	{
+		using namespace Audio; 
+		if (const ISoundWaveContainer* WaveContainer = Audio::QueryInterface<ISoundWaveContainer>(InData.Get()))
+		{
+			if (const int32 ChannelCount = ComputeAutoChannelCountFromProxies(WaveContainer->GetContainedWaveProxies()); ChannelCount != INDEX_NONE)
+			{
+				if (const TSharedPtr<const Audio::FDiscreteChannelTypeFamily> Format = NumChannelsToFormat(ChannelCount); Format.IsValid())
+				{
+					return Format;
+				}
+			}
+		}
+
+		// Fail.
+		return {};
+	}
+
+	FName FChannelAgnosticAutoFormatHelper::ComputeNodeOutputFormat(
+		const EMetasoundChannelAgnosticNodeFormatChooser InMethod,
+		const FName InCustomFormat,
+		const TSharedPtr<const Audio::IProxyData>& InContent,
+		const FMetasoundEnvironment& InEnvironment)
+	{
+		switch (InMethod)
+		{
+			case EMetasoundChannelAgnosticNodeFormatChooser::Custom:
+			{
+				return InCustomFormat;
+			}
+			case EMetasoundChannelAgnosticNodeFormatChooser::Source:
+			{
+				static const FName SourceFormatName = TEXT("SourceFormatName");
+				if (InEnvironment.Contains<FName>(SourceFormatName))
+				{
+					return InEnvironment.GetValue<FName>(SourceFormatName);
+				}
+			}
+			// Fall-through.
+			case EMetasoundChannelAgnosticNodeFormatChooser::Auto:
+			{
+				if (const TSharedPtr<const Audio::FChannelTypeFamily> Format = ComputeAutoFormatFromWaveContainer(InContent); Format.IsValid())
+				{
+					return Format->GetName();
+				}
+				break;
+			}
+		}
+
+		// Fail, but return default.
+		return FDiscreteChannelAgnosticType::GetDefaultCatFormat();
+	}
+		
+	TSharedPtr<const Audio::FDiscreteChannelTypeFamily>
+	FChannelAgnosticAutoFormatHelper::NumChannelsToFormat(const int32 NumChannels)
+	{
+		return Audio::FChannelAgnosticUtils::FindDiscreteFormatFromNumChannels(NumChannels);
+	}
+
+	
+}//namespace Metasound
\ No newline at end of file
diff --git a/Engine/Plugins/Runtime/Metasound/Source/MetasoundEngine/Public/MetasoundChannelAgnosticAutoFormat.h b/Engine/Plugins/Runtime/Metasound/Source/MetasoundEngine/Public/MetasoundChannelAgnosticAutoFormat.h
new file mode 100644
--- /dev/null
+++ b/Engine/Plugins/Runtime/Metasound/Source/MetasoundEngine/Public/MetasoundChannelAgnosticAutoFormat.h
@@ -0,0 +1,29 @@
+// Copyright Epic Games, Inc. All Rights Reserved.
+
+#pragma once
+#include "MetasoundChannelAgnosticType.h"
+
+#include "MetasoundChannelAgnosticAutoTypes.h"
+
+//#include "MetasoundChannelAgnosticAutoFormat.generated.h"
+
+#define UE_API METASOUNDENGINE_API
+
+namespace Metasound
+{
+	class FChannelAgnosticAutoFormatHelper
+	{
+	public:
+		UE_API static int32 ComputeAutoChannelCountFromProxies(const TArray<FSoundWaveProxyPtr>& InProxies);
+		UE_API static TSharedPtr<const Audio::FChannelTypeFamily> ComputeAutoFormatFromWaveContainer(const TSharedPtr<const Audio::IProxyData>& InData);
+		UE_API static TSharedPtr<const Audio::FDiscreteChannelTypeFamily> NumChannelsToFormat(const int32 NumChannels);
+		
+		UE_API static FName ComputeNodeOutputFormat(
+			const EMetasoundChannelAgnosticNodeFormatChooser InMethod,
+			const FName InCustomFormat,
+			const TSharedPtr<const Audio::IProxyData>& InContent,
+			const FMetasoundEnvironment& InEnvironment);
+	};
+}//namespace Metasound
+
+#undef UE_API 
\ No newline at end of file
diff --git a/Engine/Plugins/Runtime/Metasound/Source/MetasoundEngineTest/MetasoundEngineTest.Build.cs b/Engine/Plugins/Runtime/Metasound/Source/MetasoundEngineTest/MetasoundEngineTest.Build.cs
--- a/Engine/Plugins/Runtime/Metasound/Source/MetasoundEngineTest/MetasoundEngineTest.Build.cs
+++ b/Engine/Plugins/Runtime/Metasound/Source/MetasoundEngineTest/MetasoundEngineTest.Build.cs
@@ -17,6 +17,7 @@ namespace UnrealBuildTool.Rules
 					"CoreUObject",
 					"Engine",
 					"AudioExtensions",
+					"AudioChannelAgnosticCore",
 					"AudioMixer",
 					"SignalProcessing",
 					"Projects"
diff --git a/Engine/Plugins/Runtime/Metasound/Source/MetasoundEngineTest/Private/EngineTestChannelAgnosticAutoFormat.cpp b/Engine/Plugins/Runtime/Metasound/Source/MetasoundEngineTest/Private/EngineTestChannelAgnosticAutoFormat.cpp
new file mode 100644
--- /dev/null
+++ b/Engine/Plugins/Runtime/Metasound/Source/MetasoundEngineTest/Private/EngineTestChannelAgnosticAutoFormat.cpp
@@ -0,0 +1,213 @@
+// Copyright Epic Games, Inc. All Rights Reserved.
+
+#include "Misc/AutomationTest.h"
+
+#if WITH_DEV_AUTOMATION_TESTS
+
+PRAGMA_DISABLE_EXPERIMENTAL_WARNINGS
+
+#include "MetasoundChannelAgnosticAutoFormat.h"
+#include "MetasoundChannelAgnosticAutoTypes.h"
+#include "MetasoundEnvironment.h"
+
+#include "ChannelAgnostic/ChannelAgnosticTypeUtils.h"
+#include "TypeFamily/ChannelTypeFamily.h"
+
+#include "Sound/ISoundWaveContainer.h"
+
+namespace Metasound::Test::AutoFormat
+{
+	// Mock IProxyData that implements ISoundWaveContainer, returning an empty
+	// proxy list. This avoids needing real FSoundWaveProxy objects (which
+	// require a fully-initialised USoundWave asset pipeline).
+	class FMockEmptyWaveContainerProxy final
+		: public Audio::IProxyData
+		, public Audio::ISoundWaveContainer
+	{
+	public:
+		FMockEmptyWaveContainerProxy()
+			: Audio::IProxyData(TEXT("FMockEmptyWaveContainerProxy"))
+		{
+		}
+
+		void* QueryInterface(const FName InterfaceId) override
+		{
+			if (InterfaceId == Audio::ISoundWaveContainer::GetInterfaceId())
+			{
+				return static_cast<Audio::ISoundWaveContainer*>(this);
+			}
+			return Audio::IProxyData::QueryInterface(InterfaceId);
+		}
+
+		TArray<FSoundWaveProxyPtr> GetContainedWaveProxies() const override
+		{
+			return {};
+		}
+	};
+
+	// Mock IProxyData that does NOT implement ISoundWaveContainer, so
+	// QueryInterface<ISoundWaveContainer> returns nullptr.
+	class FMockNonContainerProxy final : public Audio::IProxyData
+	{
+	public:
+		FMockNonContainerProxy()
+			: Audio::IProxyData(TEXT("FMockNonContainerProxy"))
+		{
+		}
+
+		void* QueryInterface(const FName InterfaceId) override
+		{
+			return Audio::IProxyData::QueryInterface(InterfaceId);
+		}
+	};
+
+	//
+	// ComputeAutoChannelCountFromProxies
+	//
+
+	IMPLEMENT_SIMPLE_AUTOMATION_TEST(
+		FAutoChannelCount_EmptyArray,
+		"Audio.Metasound.AutoFormat.ComputeAutoChannelCount.EmptyArray",
+		EAutomationTestFlags::EditorContext | EAutomationTestFlags::EngineFilter)
+	bool FAutoChannelCount_EmptyArray::RunTest(const FString&)
+	{
+		const TArray<FSoundWaveProxyPtr> Empty;
+		const int32 Result = FChannelAgnosticAutoFormatHelper::ComputeAutoChannelCountFromProxies(Empty);
+		TestEqual(TEXT("Empty array returns INDEX_NONE"), Result, INDEX_NONE);
+		return true;
+	}
+
+	// Note: Tests with real FSoundWaveProxy objects are not included here because
+	// FSoundWaveProxy requires a fully initialised USoundWave (loaded asset with
+	// cooked/streaming data). ComputeAutoChannelCountFromProxies is a thin
+	// Algo::MaxElementBy wrapper, so coverage of its edge case (empty) is sufficient.
+
+	//
+	// ComputeAutoFormatFromWaveContainer
+	//
+
+	IMPLEMENT_SIMPLE_AUTOMATION_TEST(
+		FAutoFormat_NullProxy,
+		"Audio.Metasound.AutoFormat.ComputeAutoFormatFromWaveContainer.NullProxy",
+		EAutomationTestFlags::EditorContext | EAutomationTestFlags::EngineFilter)
+	bool FAutoFormat_NullProxy::RunTest(const FString&)
+	{
+		const TSharedPtr<const Audio::FChannelTypeFamily> Result =
+			FChannelAgnosticAutoFormatHelper::ComputeAutoFormatFromWaveContainer(nullptr);
+		TestFalse(TEXT("Null proxy returns invalid format"), Result.IsValid());
+		return true;
+	}
+
+	IMPLEMENT_SIMPLE_AUTOMATION_TEST(
+		FAutoFormat_NonContainerProxy,
+		"Audio.Metasound.AutoFormat.ComputeAutoFormatFromWaveContainer.NonContainerProxy",
+		EAutomationTestFlags::EditorContext | EAutomationTestFlags::EngineFilter)
+	bool FAutoFormat_NonContainerProxy::RunTest(const FString&)
+	{
+		auto Proxy = MakeShared<FMockNonContainerProxy>();
+		const TSharedPtr<const Audio::FChannelTypeFamily> Result =
+			FChannelAgnosticAutoFormatHelper::ComputeAutoFormatFromWaveContainer(Proxy);
+		TestFalse(TEXT("Non-container proxy returns invalid format"), Result.IsValid());
+		return true;
+	}
+
+	IMPLEMENT_SIMPLE_AUTOMATION_TEST(
+		FAutoFormat_EmptyContainer,
+		"Audio.Metasound.AutoFormat.ComputeAutoFormatFromWaveContainer.EmptyContainer",
+		EAutomationTestFlags::EditorContext | EAutomationTestFlags::EngineFilter)
+	bool FAutoFormat_EmptyContainer::RunTest(const FString&)
+	{
+		auto Container = MakeShared<FMockEmptyWaveContainerProxy>();
+		const TSharedPtr<const Audio::FChannelTypeFamily> Result =
+			FChannelAgnosticAutoFormatHelper::ComputeAutoFormatFromWaveContainer(Container);
+		TestFalse(TEXT("Empty container returns invalid format"), Result.IsValid());
+		return true;
+	}
+
+	//
+	// ComputeNodeOutputFormat
+	//
+
+	IMPLEMENT_SIMPLE_AUTOMATION_TEST(
+		FNodeOutputFormat_Custom,
+		"Audio.Metasound.AutoFormat.ComputeNodeOutputFormat.Custom",
+		EAutomationTestFlags::EditorContext | EAutomationTestFlags::EngineFilter)
+	bool FNodeOutputFormat_Custom::RunTest(const FString&)
+	{
+		const FName CustomFormat(TEXT("Cat:MyCustomFormat"));
+		const FMetasoundEnvironment Env;
+		const FName Result = FChannelAgnosticAutoFormatHelper::ComputeNodeOutputFormat(
+			EMetasoundChannelAgnosticNodeFormatChooser::Custom,
+			CustomFormat, nullptr, Env);
+		TestEqual(TEXT("Custom returns the custom format name"), Result, CustomFormat);
+		return true;
+	}
+
+	IMPLEMENT_SIMPLE_AUTOMATION_TEST(
+		FNodeOutputFormat_SourceWithEnv,
+		"Audio.Metasound.AutoFormat.ComputeNodeOutputFormat.SourceWithEnvironment",
+		EAutomationTestFlags::EditorContext | EAutomationTestFlags::EngineFilter)
+	bool FNodeOutputFormat_SourceWithEnv::RunTest(const FString&)
+	{
+		const FName ExpectedFormat(TEXT("Cat:Discrete:Stereo"));
+		FMetasoundEnvironment Env;
+		Env.SetValue<FName>(TEXT("SourceFormatName"), ExpectedFormat);
+		const FName Result = FChannelAgnosticAutoFormatHelper::ComputeNodeOutputFormat(
+			EMetasoundChannelAgnosticNodeFormatChooser::Source,
+			NAME_None, nullptr, Env);
+		TestEqual(TEXT("Source with env returns source format"), Result, ExpectedFormat);
+		return true;
+	}
+
+	IMPLEMENT_SIMPLE_AUTOMATION_TEST(
+		FNodeOutputFormat_SourceNoEnvFallsThrough,
+		"Audio.Metasound.AutoFormat.ComputeNodeOutputFormat.SourceNoEnvFallsThrough",
+		EAutomationTestFlags::EditorContext | EAutomationTestFlags::EngineFilter)
+	bool FNodeOutputFormat_SourceNoEnvFallsThrough::RunTest(const FString&)
+	{
+		// Source with no SourceFormatName in env falls through to Auto.
+		// Auto with no content returns the default format.
+		const FMetasoundEnvironment Env;
+		const FName Result = FChannelAgnosticAutoFormatHelper::ComputeNodeOutputFormat(
+			EMetasoundChannelAgnosticNodeFormatChooser::Source,
+			NAME_None, nullptr, Env);
+		const FName DefaultFormat = FDiscreteChannelAgnosticType::GetDefaultCatFormat();
+		TestEqual(TEXT("Source with no env falls through to default"), Result, DefaultFormat);
+		return true;
+	}
+
+	IMPLEMENT_SIMPLE_AUTOMATION_TEST(
+		FNodeOutputFormat_AutoNoContent,
+		"Audio.Metasound.AutoFormat.ComputeNodeOutputFormat.AutoNoContentReturnsDefault",
+		EAutomationTestFlags::EditorContext | EAutomationTestFlags::EngineFilter)
+	bool FNodeOutputFormat_AutoNoContent::RunTest(const FString&)
+	{
+		const FMetasoundEnvironment Env;
+		const FName Result = FChannelAgnosticAutoFormatHelper::ComputeNodeOutputFormat(
+			EMetasoundChannelAgnosticNodeFormatChooser::Auto,
+			NAME_None, nullptr, Env);
+		const FName DefaultFormat = FDiscreteChannelAgnosticType::GetDefaultCatFormat();
+		TestEqual(TEXT("Auto with no content returns default format"), Result, DefaultFormat);
+		return true;
+	}
+
+	IMPLEMENT_SIMPLE_AUTOMATION_TEST(
+		FNodeOutputFormat_AutoEmptyContainer,
+		"Audio.Metasound.AutoFormat.ComputeNodeOutputFormat.AutoEmptyContainerReturnsDefault",
+		EAutomationTestFlags::EditorContext | EAutomationTestFlags::EngineFilter)
+	bool FNodeOutputFormat_AutoEmptyContainer::RunTest(const FString&)
+	{
+		const FMetasoundEnvironment Env;
+		auto Container = MakeShared<FMockEmptyWaveContainerProxy>();
+		const FName Result = FChannelAgnosticAutoFormatHelper::ComputeNodeOutputFormat(
+			EMetasoundChannelAgnosticNodeFormatChooser::Auto,
+			NAME_None, Container, Env);
+		const FName DefaultFormat = FDiscreteChannelAgnosticType::GetDefaultCatFormat();
+		TestEqual(TEXT("Auto with empty container returns default format"), Result, DefaultFormat);
+		return true;
+	}
+}
+
+PRAGMA_ENABLE_EXPERIMENTAL_WARNINGS
+
+#endif // WITH_DEV_AUTOMATION_TESTS
diff --git a/Engine/Source/Runtime/AudioChannelAgnosticCore/Private/ChannelAgnostic/ChannelAgnosticTypeUtils.cpp b/Engine/Source/Runtime/AudioChannelAgnosticCore/Private/ChannelAgnostic/ChannelAgnosticTypeUtils.cpp
--- a/Engine/Source/Runtime/AudioChannelAgnosticCore/Private/ChannelAgnostic/ChannelAgnosticTypeUtils.cpp
+++ b/Engine/Source/Runtime/AudioChannelAgnosticCore/Private/ChannelAgnostic/ChannelAgnosticTypeUtils.cpp
@@ -56,4 +56,14 @@ namespace Audio
 		Audio::Interleave(In.Buffer.GetView(), In.NumChannels(), Out);
 	}
 
+	TSharedPtr<const FDiscreteChannelTypeFamily> FChannelAgnosticUtils::FindDiscreteFormatFromNumChannels(const int32 InNumChannels)
+	{
+		const TOptional<uint32> Mask = FindChannelMaskFromAudioMixerNumChannels(InNumChannels);
+		if (!Mask.IsSet())
+		{
+			return {};
+		}
+		return GetChannelRegistry().FindDiscreteChannelFromBitMask(*Mask);
+	}
+
 }
\ No newline at end of file
diff --git a/Engine/Source/Runtime/AudioChannelAgnosticCore/Public/ChannelAgnostic/ChannelAgnosticTypeUtils.h b/Engine/Source/Runtime/AudioChannelAgnosticCore/Public/ChannelAgnostic/ChannelAgnosticTypeUtils.h
--- a/Engine/Source/Runtime/AudioChannelAgnosticCore/Public/ChannelAgnostic/ChannelAgnosticTypeUtils.h
+++ b/Engine/Source/Runtime/AudioChannelAgnosticCore/Public/ChannelAgnostic/ChannelAgnosticTypeUtils.h
@@ -24,7 +24,15 @@ namespace Audio
 		UE_EXPERIMENTAL(5.8,"AudioChannelAgnosticCore is experimental and in flux")
 		AUDIOCHANNELAGNOSTICCORE_API static bool ChannelIdToShortMixerChannels(const TArray<FName>& InIds, TArray<ESpeakerShortNames>& OutShortNames);
 				
-		UE_EXPERIMENTAL(5.8,"AudioChannelAgnosticCore is experimental and in flux")		
+		UE_EXPERIMENTAL(5.8,"AudioChannelAgnosticCore is experimental and in flux")
 		static AUDIOCHANNELAGNOSTICCORE_API void Interleave(const FChannelAgnosticType& In, TArrayView<float> Out);
+
+		/**
+		 * Finds a discrete channel type family matching the given number of channels,
+		 * using the standard AudioMixer channel layouts.
+		 * @return Discrete channel type family, or null if no match found.
+		 */
+		UE_EXPERIMENTAL(5.8,"AudioChannelAgnosticCore is experimental and in flux")
+		AUDIOCHANNELAGNOSTICCORE_API static TSharedPtr<const FDiscreteChannelTypeFamily> FindDiscreteFormatFromNumChannels(const int32 InNumChannels);
 	};
 }
