[Backout] - CL51656697
#rnx
Original CL Desc
-----------------------------------------------------------------
CL C: GraphCore polymorphic infrastructure

Split from CL 51557033. Core type system changes. MetasoundPolymorphic.h
introduces the polymorphic hierarchy. DataReference.h/DataFactory.h/
Variable.h/VertexData.h modified. All changes UE_EXPERIMENTAL guarded
and existing code paths preserved via if constexpr. Unit tests cover this.
<enter description here>

#tests LowLevelTests Audio

diff --git a/Engine/Plugins/Runtime/Metasound/Source/MetasoundGraphCore/Public/MetasoundDataFactory.h b/Engine/Plugins/Runtime/Metasound/Source/MetasoundGraphCore/Public/MetasoundDataFactory.h
--- a/Engine/Plugins/Runtime/Metasound/Source/MetasoundGraphCore/Public/MetasoundDataFactory.h
+++ b/Engine/Plugins/Runtime/Metasound/Source/MetasoundGraphCore/Public/MetasoundDataFactory.h
@@ -30,10 +30,10 @@ namespace Metasound
 			static constexpr bool bIsArrayConstructibleWithSettingsAndArgs = false;
 
 			// DataType Construction
-			static constexpr bool bIsDefaultConstructible = std::is_constructible<DataType>::value || CHasFactoryFunctionDefaultConstruction<DataType>;
-			static constexpr bool bIsConstructibleWithSettings = std::is_constructible<DataType, const FOperatorSettings&>::value || CHasFactoryFunctionWithSettings<DataType>;
-			static constexpr bool bIsConstructibleWithArgs = std::is_constructible<DataType, ArgTypes...>::value || CHasFactoryFunctionWithArgs<DataType, ArgTypes...>;
-			static constexpr bool bIsConstructibleWithSettingsAndArgs = std::is_constructible<DataType, const FOperatorSettings&, ArgTypes...>::value || CHasFactoryFunctionWithSettingsAndArgs<DataType, ArgTypes...>;
+			static constexpr bool bIsDefaultConstructible = std::is_constructible<DataType>::value;
+			static constexpr bool bIsConstructibleWithSettings = std::is_constructible<DataType, const FOperatorSettings&>::value;
+			static constexpr bool bIsConstructibleWithArgs = std::is_constructible<DataType, ArgTypes...>::value;
+			static constexpr bool bIsConstructibleWithSettingsAndArgs = std::is_constructible<DataType, const FOperatorSettings&, ArgTypes...>::value;
 		};
 
 		/** Partial specialization for TArray<> types. A TArray can be constructed elementwise with
@@ -153,7 +153,7 @@ namespace Metasound
 		{
 			using TConstructorTraits = TDataTypeConstructorTraits<DataType, ArgTypes...>;
 
-			static constexpr bool bForwardSettingsAndArgs = TConstructorTraits::bIsConstructibleWithSettingsAndArgs || CHasFactoryFunctionWithSettingsAndArgs<DataType, ArgTypes...>;
+			static constexpr bool bForwardSettingsAndArgs = TConstructorTraits::bIsConstructibleWithSettingsAndArgs;
 			static constexpr bool bForwardArgs = false;
 			static constexpr bool bForwardSettings = false;
 			static constexpr bool bForwardNone = false;
@@ -172,8 +172,8 @@ namespace Metasound
 		{
 			using TConstructorTraits = TDataTypeConstructorTraits<DataType, ArgTypes...>;
 
-			static constexpr bool bForwardSettingsAndArgs = TConstructorTraits::bIsConstructibleWithSettingsAndArgs || CHasFactoryFunctionWithSettingsAndArgs<DataType, ArgTypes...>;
-			static constexpr bool bForwardArgs = !bForwardSettingsAndArgs && (TConstructorTraits::bIsConstructibleWithArgs || CHasFactoryFunctionWithArgs<DataType, ArgTypes...>); 
+			static constexpr bool bForwardSettingsAndArgs = TConstructorTraits::bIsConstructibleWithSettingsAndArgs;
+			static constexpr bool bForwardArgs = !bForwardSettingsAndArgs && TConstructorTraits::bIsConstructibleWithArgs;
 			static constexpr bool bForwardSettings = false;
 			static constexpr bool bForwardNone = false;
 
@@ -191,10 +191,10 @@ namespace Metasound
 		{
 			using TConstructorTraits = TDataTypeConstructorTraits<DataType, ArgTypes...>;
 
-			static constexpr bool bForwardSettingsAndArgs = TConstructorTraits::bIsConstructibleWithSettingsAndArgs || CHasFactoryFunctionWithSettingsAndArgs<DataType, ArgTypes...>;
-			static constexpr bool bForwardArgs = !bForwardSettingsAndArgs && (TConstructorTraits::bIsConstructibleWithArgs || CHasFactoryFunctionWithArgs<DataType, ArgTypes...>);
-			static constexpr bool bForwardSettings = !(bForwardArgs || bForwardSettingsAndArgs) && (TConstructorTraits::bIsConstructibleWithSettings || CHasFactoryFunctionWithSettings<DataType>);
-			static constexpr bool bForwardNone = !(bForwardArgs || bForwardSettingsAndArgs || bForwardSettings) && (TConstructorTraits::bIsDefaultConstructible || CHasFactoryFunctionDefaultConstruction<DataType>);
+			static constexpr bool bForwardSettingsAndArgs = TConstructorTraits::bIsConstructibleWithSettingsAndArgs;
+			static constexpr bool bForwardArgs = !bForwardSettingsAndArgs && TConstructorTraits::bIsConstructibleWithArgs;
+			static constexpr bool bForwardSettings = !(bForwardArgs || bForwardSettingsAndArgs) && TConstructorTraits::bIsConstructibleWithSettings;
+			static constexpr bool bForwardNone = !(bForwardArgs || bForwardSettingsAndArgs || bForwardSettings) && TConstructorTraits::bIsDefaultConstructible;
 
 			static constexpr bool bCannotForwardToConstructor = !(bForwardSettingsAndArgs || bForwardArgs || bForwardSettings || bForwardNone);
 		};
diff --git a/Engine/Plugins/Runtime/Metasound/Source/MetasoundGraphCore/Public/MetasoundDataReference.h b/Engine/Plugins/Runtime/Metasound/Source/MetasoundGraphCore/Public/MetasoundDataReference.h
--- a/Engine/Plugins/Runtime/Metasound/Source/MetasoundGraphCore/Public/MetasoundDataReference.h
+++ b/Engine/Plugins/Runtime/Metasound/Source/MetasoundGraphCore/Public/MetasoundDataReference.h
@@ -4,7 +4,6 @@
 
 #include <type_traits>
 
-#include "Concepts/DerivedFrom.h"
 #include "Containers/UnrealString.h"
 #include "HAL/Platform.h"
 #include "Internationalization/Text.h"
@@ -25,7 +24,6 @@ namespace Audio
 
 namespace Metasound
 {
-	class FOperatorSettings;
 	// Forward declare for LexToString
 	enum class EDataReferenceAccessType : uint8;
 }
@@ -98,8 +96,8 @@ namespace Metasound
 			return false;
 		}
 	}
-
-
+	
+	
 	// Polymorphic overload
 	template<typename DataType>
 	requires TPolymorphicTraits<DataType>::bIsPolymorphic
@@ -139,83 +137,56 @@ namespace Metasound
 	{
 		NewObject
 	};
-
+	
 	namespace MetasoundDataReferencePrivate
 	{
 		template <typename Base, typename Derived, typename... ArgTypes>
 		UE_EXPERIMENTAL(5.7, "Polymorphic API is under development")
-		TSharedRef<Base, ESPMode::NotThreadSafe> MakeRefOf(ArgTypes&&... Args)
+		TSharedRef<Base, ESPMode::NotThreadSafe> MakeRefOf(ArgTypes&&... args)
 		{
 			static_assert(std::is_base_of_v<Base, Derived>, "Target must derive from Base (or be Base)");
 			static_assert(std::is_constructible_v<Derived, ArgTypes...>, "Args do not match Target's constructor");
 
-			TSharedRef<Derived, ESPMode::NotThreadSafe> DerivedRef = MakeShared<Derived, ESPMode::NotThreadSafe>(Forward<ArgTypes>(Args)...);
-
+			TSharedRef<Derived, ESPMode::NotThreadSafe> DerivedRef = MakeShared<Derived, ESPMode::NotThreadSafe>(Forward<ArgTypes>(args)...);
+ 
 			return StaticCastSharedRef<Base>(DerivedRef);
 		}
-
-		template <typename Base, typename Derived, typename... ArgTypes>
-		UE_EXPERIMENTAL(5.7, "Polymorphic API is under development")
-		TSharedRef<Base, ESPMode::NotThreadSafe> MakeRefOf(ArgTypes&&... Args) requires CHasFactoryFunction<Derived, ArgTypes...>
-		{
-			static_assert(std::is_base_of_v<Base, Derived>, "Target must derive from Base (or be Base)");
-			TSharedRef<Derived, ESPMode::NotThreadSafe> DerivedRef = Derived::CreateNew(Forward<ArgTypes>(Args)...);
- 			return StaticCastSharedRef<Base>(DerivedRef);
-		}
-
 	}
 
 	template <typename DataType>
 	struct TDynamicDataReferencePolicy;
 
 	template<typename DataType>
-	requires (!TPolymorphicTraits<DataType>::bIsPolymorphic)
+	requires (TPolymorphicTraits<DataType>::bIsPolymorphic)
 	struct TDynamicDataReferencePolicy<DataType>
 	{
-		template<typename QueryType> bool IsA(const DataType*) const { return GetMetasoundDataTypeName<DataType>() == GetMetasoundDataTypeName<QueryType>(); }
-		const FName& GetDataTypeName(const DataType*) const { return GetMetasoundDataTypeName<DataType>(); }
-		const void* GetDataTypeId(const DataType*) const { return GetMetasoundDataTypeId<DataType>(); }
-	};
+		explicit TDynamicDataReferencePolicy(const FName& InDataTypeName = GetMetasoundDataTypeName<DataType>(), const void* InTypeId = GetMetasoundDataTypeId<DataType>())
+			: DataTypeName(InDataTypeName), DataTypeId(InTypeId)
+		{}
 
+		template<typename QueryType> bool IsA() const { return IsChildOfByName(DataTypeName, GetMetasoundDataTypeName<QueryType>()); }
+		const FName& GetDataTypeName() const { return DataTypeName; }
+		const void* GetTypeId() const { return DataTypeId; }
+		FName DataTypeName;
+		const void* DataTypeId;
+	};
+	
 	template<typename DataType>
-	requires (TPolymorphicTraits<DataType>::bIsPolymorphic)
+	requires (!TPolymorphicTraits<DataType>::bIsPolymorphic)
 	struct TDynamicDataReferencePolicy<DataType>
 	{
-		/* TODO: Re-enable once all poly types define GetDataTypeName+GetDataTypeId */
-		/* static_assert(CHasPolymorphicTypeFunctions<DataType>, "Dynamic Reference Policy Requires GetDataTypeName+GetDataTypeId functions on DataType"); */
-
-		template<typename QueryType> bool IsA(const DataType* InOuter) const
-		{
-			return IsChildOfByName(GetDataTypeName(InOuter), GetMetasoundDataTypeName<QueryType>());
-		}
+		TDynamicDataReferencePolicy() = default;
+		TDynamicDataReferencePolicy(const FName&, const void*) {}
 
-		const FName& GetDataTypeName(const DataType* InOuter) const requires(CHasPolymorphicTypeGetter<DataType>)
-		{
-			return InOuter->GetDataTypeName();
-		}
-
-		// Fallback when type doesn't yet implement GetDataTypeName().
-		const FName& GetDataTypeName(const DataType*) const requires(!CHasPolymorphicTypeGetter<DataType>)
-		{
-			return GetMetasoundDataTypeName<DataType>();
-		}
-
-		const void* GetDataTypeId(const DataType* InOuter) const requires(CHasPolymorphicTypeIdGetter<DataType>)
-		{
-			return InOuter->GetDataTypeId();
-		}
-
-		// Fallback when type doesn't yet implement GetDataTypeId().
-		const void* GetDataTypeId(const DataType*) const requires(!CHasPolymorphicTypeIdGetter<DataType>)
-		{
-			return GetMetasoundDataTypeId<DataType>();
-		}
+		template<typename QueryType> bool IsA() const { return GetDataTypeName() == GetMetasoundDataTypeName<QueryType>(); }
+		const FName& GetDataTypeName() const { return GetMetasoundDataTypeName<DataType>(); }
+		const void* GetTypeId() const { return GetMetasoundDataTypeId<DataType>(); }
 	};
-
-	/** Template class for a paramter reference.
+	
+	/** Template class for a paramter reference. 
 	 *
 	 * This fulfills the IParamterRef interface, utilizing TDataReferenceTypeInfo to
-	 * define the the TypeName and TypeId of the parameter.
+	 * define the the TypeName and TypeId of the parameter. 
 	 */
 	template <typename DataType>
 	class TDataReference : public IDataReference
@@ -231,13 +202,6 @@ namespace Metasound
 		{
 		}
 
-		template <typename... ArgTypes>
-		TDataReference(EDataRefShouldConstruct InToken, ArgTypes&&... Args) requires(CIsPolymorphicAndAbstract<DataType> && CHasFactoryFunction<DataType, ArgTypes...>)
-			: ObjectReference(DataType::CreateNew(Forward<ArgTypes>(Args)...))
-		{
-			/* TODO: Re-enable once all poly types define CreateNew */ // static_assert(CHasFactoryFunction<DataType, ArgTypes...>);
-		}
-
 		typedef TSharedRef<DataType, ESPMode::NotThreadSafe> FRefType;
 
 		explicit TDataReference(EDataRefShouldConstruct /*InToken*/, FRefType InRef)
@@ -249,8 +213,9 @@ namespace Metasound
 								const FName InTypeName,
 								const void* InTypeId)
 			: ObjectReference(MoveTemp(InRef))
+			, ObjectReferencePolicy(InTypeName, InTypeId)
 		{}
-
+	
 	public:
 
 		typedef TDataReferenceTypeInfo<DataType> FInfoType;
@@ -264,38 +229,37 @@ namespace Metasound
 			return TDataReference<DataType>(EDataRefShouldConstruct::NewObject, Forward<ArgTypes>(Args)...);
 		}
 
-		/**
-		 * Polymorphic form of CreateNew.
-		 */
-		template <typename... ArgTypes>
-		static TDataReference<DataType> CreateNew(ArgTypes&&... Args) requires (CIsPolymorphicAndAbstract<DataType> && CHasFactoryFunction<DataType, ArgTypes...>)
-		{
-			/* TODO: Re-enable once all poly types define CreateNew */ // static_assert(CHasFactoryFunction<DataType, ArgTypes...>);
-			return TDataReference<DataType>(EDataRefShouldConstruct::NewObject, Forward<ArgTypes>(Args)...);
-		}
-
 		// Constructs Derived but stores as Base (DataType).
 		template <typename Derived, typename... ArgTypes>
+		UE_EXPERIMENTAL(5.7, "Polymorphic API is under development")
 		static TDataReference CreateNewDerived(ArgTypes&&... Args)
 		{
 			static_assert(std::is_base_of_v<DataType, Derived>, "Derived must inherit from DataType");
+			static_assert(std::is_constructible_v<Derived, ArgTypes...>, "Args don't match a Derived constructor");
+			
 			return TDataReference(
 				EDataRefShouldConstruct::NewObject,
-				MetasoundDataReferencePrivate::MakeRefOf<DataType, Derived>(Forward<ArgTypes>(Args)...)
+				MetasoundDataReferencePrivate::MakeRefOf<DataType, Derived>(Forward<ArgTypes>(Args)...),
+				GetMetasoundDataTypeName<Derived>(),
+			    GetMetasoundDataTypeId<Derived>()
 			);
 		}
 
 		template <typename Derived, typename... ArgTypes>
+		UE_EXPERIMENTAL(5.8, "Polymorphic API is under development")
 		static TDataReference CreateNewDerivedAs(const FName& InDerivedTypeName,
 												 FMetasoundDataTypeId InDerivedTypeId,
 												 ArgTypes&&... Args)
 		{
 			static_assert(std::is_base_of_v<DataType, Derived>, "Derived must inherit from DataType");
+			static_assert(std::is_constructible_v<Derived, ArgTypes...>, "Args don't match a Derived constructor");
+
 		#if DO_CHECK
 			PRAGMA_DISABLE_EXPERIMENTAL_WARNINGS
 			check(Metasound::IsChildOfByName(InDerivedTypeName, GetMetasoundDataTypeName<DataType>()));
 			PRAGMA_ENABLE_EXPERIMENTAL_WARNINGS
 		#endif
+
 			return TDataReference(
 				EDataRefShouldConstruct::NewObject,
 				MetasoundDataReferencePrivate::MakeRefOf<DataType, Derived>(Forward<ArgTypes>(Args)...),
@@ -304,6 +268,7 @@ namespace Metasound
 			);
 		}
 
+
 		/** Enable copy constructor */
 		TDataReference(const TDataReference<DataType>& Other) = default;
 
@@ -319,13 +284,13 @@ namespace Metasound
 		/** Return the name of the underlying type. */
 		virtual const FName& GetDataTypeName() const override
 		{
-			return ObjectReferencePolicy.GetDataTypeName(&ObjectReference.Get());
+			return ObjectReferencePolicy.GetDataTypeName();	
 		}
 
 		/** Return the ID of the underlying type. */
 		virtual const void* const GetDataTypeId() const override
 		{
-			return ObjectReferencePolicy.GetDataTypeId(&ObjectReference.Get());
+			return ObjectReferencePolicy.GetTypeId();
 		}
 
 		/** Return a raw pointer to the data. */
@@ -338,24 +303,29 @@ namespace Metasound
 		UE_EXPERIMENTAL(5.7, "Polymorphic API is under development")
 		bool IsA() const
 		{
-			return this->ObjectReferencePolicy.template IsA<QueryType>(&ObjectReference.Get());
+			return this->ObjectReferencePolicy.template IsA<QueryType>();
 		}
-
+		
 		template <typename CastType>
 		UE_EXPERIMENTAL(5.7, "Polymorphic API is under development")
-		TOptional<TSharedRef<CastType, ESPMode::NotThreadSafe>> TryCast() const
+		TDataReference<CastType> TryCast() const
 		{
 			if (IsA<CastType>())
 			{
-				return StaticCastSharedRef<CastType>(ObjectReference);
+				return TDataReference(
+					EDataRefShouldConstruct::NewObject,
+					MetasoundDataReferencePrivate::MakeRefOf<DataType, CastType>(ObjectReference),
+					GetMetasoundDataTypeName<CastType>(),
+					GetMetasoundDataTypeId<CastType>()
+				);
 			}
 			return {};
 		}
 
 	protected:
-
+		
 		// Protected object reference is utilized by subclasses which define what
-		// access is provided to the ObjectReference.
+		// access is provided to the ObjectReference. 
 		FRefType ObjectReference;
 		TDynamicDataReferencePolicy<DataType> ObjectReferencePolicy;
 	};
@@ -368,7 +338,7 @@ namespace Metasound
 	template <typename DataType>
 	class TDataWriteReference;
 
-	/** TDataValueReference represents a constant value and provides read only access.
+	/** TDataValueReference represents a constant value and provides read only access. 
 	 * A TDataValueReference can never change value. */
 	template<typename DataType>
 	class TDataValueReference : public TDataReference<DataType>
@@ -391,7 +361,7 @@ namespace Metasound
 
 		template<typename T>
 		friend TDataValueReference<T> ValueCast(const TDataWriteReference<T>& InRef);
-
+		
 		using FRefType = TDataReference<DataType>::FRefType;
 		TDataValueReference(EDataRefShouldConstruct, FRefType InRef, FName InTypeName, const void* InTypeId)
 			: TDataReference<DataType>(EDataRefShouldConstruct::NewObject, MoveTemp(InRef), InTypeName, InTypeId)
@@ -408,17 +378,13 @@ namespace Metasound
 			return TDataValueReference<DataType>(EDataRefShouldConstruct::NewObject, Forward<ArgTypes>(Args)...);
 		}
 
-		template <typename... ArgTypes>
-		static TDataValueReference<DataType> CreateNew(ArgTypes&&... Args) requires (CIsPolymorphicAndAbstract<DataType> && CHasFactoryFunction<DataType, ArgTypes...>)
-		{
-			/* TODO: Re-enable once all poly types define CreateNew */ // static_assert(CHasFactoryFunction<DataType, ArgTypes...>);
-			return TDataValueReference<DataType>(EDataRefShouldConstruct::NewObject, Forward<ArgTypes>(Args)...);
-		}
-
 		template <typename Derived, typename... ArgTypes>
+		UE_EXPERIMENTAL(5.7, "Polymorphic API is under development")
 		static TDataValueReference CreateNewDerived(ArgTypes&&... Args)
 		{
 			static_assert(std::is_base_of_v<DataType, Derived>, "Derived must inherit from DataType");
+			static_assert(std::is_constructible_v<Derived, ArgTypes...>, "Args don't match a Derived constructor");
+			
 			return TDataValueReference(
 				EDataRefShouldConstruct::NewObject,
 				MetasoundDataReferencePrivate::MakeRefOf<DataType, Derived>(Forward<ArgTypes>(Args)...)
@@ -426,16 +392,20 @@ namespace Metasound
 		}
 
 		template <typename Derived, typename... ArgTypes>
+		UE_EXPERIMENTAL(5.8, "Polymorphic API is under development")
 		static TDataValueReference CreateNewDerivedAs(const FName& InDerivedTypeName,
 													 FMetasoundDataTypeId InDerivedTypeId,
 													 ArgTypes&&... Args)
 		{
 			static_assert(std::is_base_of_v<DataType, Derived>, "Derived must inherit from DataType");
+			static_assert(std::is_constructible_v<Derived, ArgTypes...>, "Args don't match a Derived constructor");
+
 		#if DO_CHECK
 			PRAGMA_DISABLE_EXPERIMENTAL_WARNINGS
 			check(Metasound::IsChildOfByName(InDerivedTypeName, GetMetasoundDataTypeName<DataType>()));
 			PRAGMA_ENABLE_EXPERIMENTAL_WARNINGS
 		#endif
+			
 			return TDataValueReference(
 				EDataRefShouldConstruct::NewObject,
 				MetasoundDataReferencePrivate::MakeRefOf<DataType, Derived>(Forward<ArgTypes>(Args)...),
@@ -443,9 +413,7 @@ namespace Metasound
 				InDerivedTypeId
 			);
 		}
-
-
-
+		
 		/** Enable copy constructor */
 		TDataValueReference(const TDataValueReference<DataType>& Other) = default;
 
@@ -490,10 +458,10 @@ namespace Metasound
 		}
 	};
 
-	/** Cast a TDataReadReference to a TDataValueReference.
+	/** Cast a TDataReadReference to a TDataValueReference. 
 	 *
 	 * In general TDataReadReferences should not be converted into TDataValueReferences unless the caller
-	 * can be certain that no other TDataWriteReference exists for the underlying parameter. Having a
+	 * can be certain that no other TDataWriteReference exists for the underlying parameter. Having a 
 	 * TDataWriteReferences to an existing TDataValueReference can cause confusing behavior as values
 	 * references are not expected to change value.
 	 */
@@ -503,10 +471,10 @@ namespace Metasound
 		return TDataValueReference<T>(InRef);
 	}
 
-	/** Cast a TDataWriteReference to a TDataValueReference.
+	/** Cast a TDataWriteReference to a TDataValueReference. 
 	 *
 	 * In general TDataWriteReferences should never beconverted into TDataValueReferences unless the caller
-	 * can be certain that no other TDataWriteReference exists for the underlying parameter. Having a
+	 * can be certain that no other TDataWriteReference exists for the underlying parameter. Having a 
 	 * TDataWriteReferences to an existing TDataValueReference can cause confusing behavior as values
 	 * references are not expected to change value.
 	 */
@@ -528,7 +496,7 @@ namespace Metasound
 		}
 
 		/** Create a writable ref from a blank parameter ref. Should be done with care and understanding
-		 * of side-effects of converting a ref to a writable ref.
+		 * of side-effects of converting a ref to a writable ref. 
 		 */
 		TDataWriteReference<DataType>(const TDataReference<DataType>& InDataReference)
 		: TDataReference<DataType>(InDataReference)
@@ -554,18 +522,13 @@ namespace Metasound
 			return TDataWriteReference<DataType>(EDataRefShouldConstruct::NewObject, Forward<ArgTypes>(Args)...);
 		}
 
-		/** This should be used to construct a new DataType object and return this TDataWriteReference as a wrapper around it. */
-		template <typename... ArgTypes>
-		static TDataWriteReference<DataType> CreateNew(ArgTypes&&... Args) requires (CIsPolymorphicAndAbstract<DataType> && CHasFactoryFunction<DataType, ArgTypes...>)
-		{
-			/* TODO: Re-enable once all poly types define CreateNew */ // static_assert(CHasFactoryFunction<DataType, ArgTypes...>);
-			return TDataWriteReference<DataType>(EDataRefShouldConstruct::NewObject, Forward<ArgTypes>(Args)...);
-		}
-
 		template <typename Derived, typename... ArgTypes>
+		UE_EXPERIMENTAL(5.8, "Polymorphic API is under development")
 		static TDataWriteReference CreateNewDerived(ArgTypes&&... Args)
 		{
 			static_assert(std::is_base_of_v<DataType, Derived>, "Derived must inherit from DataType");
+			static_assert(std::is_constructible_v<Derived, ArgTypes...>, "Args don't match a Derived constructor");
+			
 			return TDataWriteReference(
 				EDataRefShouldConstruct::NewObject,
 				MetasoundDataReferencePrivate::MakeRefOf<DataType, Derived>(Forward<ArgTypes>(Args)...)
@@ -573,16 +536,20 @@ namespace Metasound
 		}
 
 		template <typename Derived, typename... ArgTypes>
+		UE_EXPERIMENTAL(5.8, "Polymorphic API is under development")
 		static TDataWriteReference CreateNewDerivedAs(const FName& InDerivedTypeName,
 													 FMetasoundDataTypeId InDerivedTypeId,
 													 ArgTypes&&... Args)
 		{
 			static_assert(std::is_base_of_v<DataType, Derived>, "Derived must inherit from DataType");
+			static_assert(std::is_constructible_v<Derived, ArgTypes...>, "Args don't match a Derived constructor");
+
 		#if DO_CHECK
 			PRAGMA_DISABLE_EXPERIMENTAL_WARNINGS
 			check(Metasound::IsChildOfByName(InDerivedTypeName, GetMetasoundDataTypeName<DataType>()));
 			PRAGMA_ENABLE_EXPERIMENTAL_WARNINGS
 		#endif
+			
 			return TDataWriteReference(
 				EDataRefShouldConstruct::NewObject,
 				MetasoundDataReferencePrivate::MakeRefOf<DataType, Derived>(Forward<ArgTypes>(Args)...),
@@ -590,9 +557,7 @@ namespace Metasound
 				InDerivedTypeId
 			);
 		}
-
-
-
+		
 		/** Enable copy constructor */
 		TDataWriteReference(const TDataWriteReference<DataType>& Other) = default;
 
@@ -642,17 +607,17 @@ namespace Metasound
 
 	};
 
-	/** Cast a TDataReadReference to a TDataWriteReference.
+	/** Cast a TDataReadReference to a TDataWriteReference. 
 	 *
 	 * In general TDataReadReferences should not be converted into TDataWriteReferences unless the caller
-	 * can be certain that no other TDataWriteReference exists for the underlying parameter. Having multiple
+	 * can be certain that no other TDataWriteReference exists for the underlying parameter. Having multiple 
 	 * TDataWriteReferences to the same parameter can cause confusion behavior as values are overwritten in
 	 * an underterministic fashion.
 	 */
 	template <typename DataType>
 	TDataWriteReference<DataType> WriteCast(const TDataReadReference<DataType>& InReadableRef)
 	{
-		const TDataReference<DataType>& Ref = static_cast<const TDataReference<DataType>&>(InReadableRef);
+		const TDataReference<DataType>& Ref = static_cast<const TDataReference<DataType>&>(InReadableRef); 
 		return TDataWriteReference<DataType>(Ref);
 	}
 
@@ -685,17 +650,13 @@ namespace Metasound
 			return TDataReadReference<DataType>(EDataRefShouldConstruct::NewObject, Forward<ArgTypes>(Args)...);
 		}
 
-		template <typename... ArgTypes>
-		static TDataReadReference<DataType> CreateNew(ArgTypes&&... Args) requires (CIsPolymorphicAndAbstract<DataType> && CHasFactoryFunction<DataType, ArgTypes...>)
-		{
-			/* TODO: Re-enable once all poly types define CreateNew */ // static_assert(CHasFactoryFunction<DataType, ArgTypes...>);
-			return TDataReadReference<DataType>(EDataRefShouldConstruct::NewObject, Forward<ArgTypes>(Args)...);
-		}
-
 		template <typename Derived, typename... ArgTypes>
+		UE_EXPERIMENTAL(5.8, "Polymorphic API is under development")
 		static TDataReadReference CreateNewDerived(ArgTypes&&... Args)
 		{
 			static_assert(std::is_base_of_v<DataType, Derived>, "Derived must inherit from DataType");
+			static_assert(std::is_constructible_v<Derived, ArgTypes...>, "Args don't match a Derived constructor");
+			
 			return TDataReadReference(
 				EDataRefShouldConstruct::NewObject,
 				MetasoundDataReferencePrivate::MakeRefOf<DataType, Derived>(Forward<ArgTypes>(Args)...)
@@ -703,14 +664,18 @@ namespace Metasound
 		}
 
 		template <typename Derived, typename... ArgTypes>
+		UE_EXPERIMENTAL(5.8, "Polymorphic API is under development")
 		static TDataReadReference CreateNewDerivedAs(const FName& InDerivedTypeName, FMetasoundDataTypeId InDerivedTypeId, ArgTypes&&... Args)
 		{
 			static_assert(std::is_base_of_v<DataType, Derived>, "Derived must inherit from DataType");
+			static_assert(std::is_constructible_v<Derived, ArgTypes...>, "Args don't match a Derived constructor");
+
 		#if DO_CHECK
 			PRAGMA_DISABLE_EXPERIMENTAL_WARNINGS
 			check(Metasound::IsChildOfByName(InDerivedTypeName, GetMetasoundDataTypeName<DataType>()));
 			PRAGMA_ENABLE_EXPERIMENTAL_WARNINGS
 		#endif
+			
 			return TDataReadReference(
 				EDataRefShouldConstruct::NewObject,
 				MetasoundDataReferencePrivate::MakeRefOf<DataType, Derived>(Forward<ArgTypes>(Args)...),
@@ -718,7 +683,7 @@ namespace Metasound
 				InDerivedTypeId
 			);
 		}
-
+		
 		TDataReadReference(const TDataReadReference& Other) = default;
 
 		/** Construct a readable parameter ref from a writable parameter ref. */
@@ -774,9 +739,9 @@ namespace Metasound
 	};
 
 	/** EDataReferenceAccessType describes the underlying data reference access
-	 * type for a data reference contained in a FAnyDataReference.
+	 * type for a data reference contained in a FAnyDataReference. 
 	 *
-	 * This value can be used to determine which methods are supported for accessing
+	 * This value can be used to determine which methods are supported for accessing 
 	 * a data reference using GetDataReadReference<>() or GetDataWriteReference<>()
 	 */
 	enum class EDataReferenceAccessType : uint8
@@ -789,17 +754,17 @@ namespace Metasound
 
 
 
-	/** Container for any data reference.
+	/** Container for any data reference. 
 	 *
 	 * This container maintains the underlying containers access type (Read or Write)
-	 * and data type. This allows for convenient storage by implementing a virtual
+	 * and data type. This allows for convenient storage by implementing a virtual 
 	 * copy constructor and assignment operator.
 	 */
 	class FAnyDataReference : public IDataReference
 	{
 		/** TGetAs allows for template specialization of data retrieval and conversion
 		 * rules for the various flavors of data referencing. The default implementation
-		 * assumes we are attempting to retrieve the value stored in the data reference.
+		 * assumes we are attempting to retrieve the value stored in the data reference. 
 		 */
 		template<typename ResultType>
 		struct TGetAs
@@ -808,7 +773,7 @@ namespace Metasound
 			using DataType = typename std::decay_t<std::remove_pointer_t<ResultType>>;
 
 			// Only deref a pointer if the returned type is specifically requested to be
-			// a non-pointer.
+			// a non-pointer. 
 			static ResultType DerefPointerIfNonPointerResultType(const DataType* InData)
 			{
 				if constexpr (std::is_pointer_v<ResultType>)
@@ -823,7 +788,7 @@ namespace Metasound
 			};
 
 		public:
-
+		
 			// Return a the underlying data by value or by const pointer depending
 			// upon the template parameter ResultType
 			static ResultType Get(EDataReferenceAccessType InAccessType, const IDataReference& InRef)
@@ -918,7 +883,7 @@ namespace Metasound
 		};
 
 
-		// Private constructor.
+		// Private constructor. 
 		FAnyDataReference(EDataReferenceAccessType InAccessType, const IDataReference& InDataRef)
 		: AccessType(InAccessType)
 		, DataRefPtr(InDataRef.Clone())
@@ -979,7 +944,7 @@ namespace Metasound
 			check(DataRefPtr.IsValid());
 			return DataRefPtr->GetDataTypeName();
 		}
-
+		
 		/** Returns the data type ID of the underlying data reference. */
 		virtual const void* const GetDataTypeId() const override
 		{
@@ -1003,7 +968,7 @@ namespace Metasound
 
 		/** Return the data in the specified format. The template argument
 		 * can be any of the support TData*Reference<> types, the underlying DataType
-		 * or a const DataType* pointer.
+		 * or a const DataType* pointer. 
 		 */
 		template<typename ReturnType>
 		ReturnType GetAs() const
@@ -1019,10 +984,10 @@ namespace Metasound
 			}
 		}
 
-		/** Returns the current value of a reference.
+		/** Returns the current value of a reference. 
 		 *
 		 * This method's behavior is undefined and an assert will be called if
-		 * the DataType differs from the underlying data reference's DataType.
+		 * the DataType differs from the underlying data reference's DataType. 
 		 */
 		template<typename DataType>
 		const DataType* GetValue() const
@@ -1035,11 +1000,11 @@ namespace Metasound
 		}
 
 		/** Return a non-const pointer to the data.
-		 *
-		 * If the vertex is bound with this will return a valid pointer.
+		 * 
+		 * If the vertex is bound with this will return a valid pointer. 
 		 * Otherwise it will return a nullptr.  This method will assert
 		 * if the reference is bound, but not writable or if the data type
-		 * does match.
+		 * does match.  
 		 */
 		template<typename DataType>
 		DataType* GetWritableValue() const
@@ -1054,14 +1019,14 @@ namespace Metasound
 			return nullptr;
 		}
 
-		/** Get access to a TDataValueReference.
+		/** Get access to a TDataValueReference. 
 		 *
 		 * This method will return a valid TDataValueReference of the templated data
 		 * type. The returned object is only valid if:
 		 *     1. The template parameter DataType matches that of the underlying data reference.
 		 *     2. The underlying data reference is has Value access.
 		 *
-		 * If this method's behavior is undefined and will assert if it is called
+		 * If this method's behavior is undefined and will assert if it is called 
 		 * with a mismatched data type or unsupported access type.
 		 */
 		template<typename DataType>
@@ -1071,14 +1036,14 @@ namespace Metasound
 			return GetAs<TDataValueReference<DataType>>();
 		}
 
-		/** Get access to a TDataReadReference.
+		/** Get access to a TDataReadReference. 
 		 *
 		 * This method will return a valid TDataReadReference of the templated data
 		 * type. The returned object is only valid if:
 		 *     1. The template parameter DataType matches that of the underlying data reference.
 		 *     2. The underlying data reference is has Read, Write or Value access.
 		 *
-		 * If this method's behavior is undefined and will assert if it is called
+		 * If this method's behavior is undefined and will assert if it is called 
 		 * with a mismatched data type or unsupported access type.
 		 */
 		template<typename DataType>
@@ -1087,14 +1052,14 @@ namespace Metasound
 			return GetAs<TDataReadReference<DataType>>();
 		}
 
-		/** Get access to a TDataWriteReference.
+		/** Get access to a TDataWriteReference. 
 		 *
 		 * This method will return a valid TDataWriteReference of the templated data
 		 * type. The returned object is only valid if:
 		 *     1. The template paramter DataType matches that of the underlying data reference.
 		 *     2. The underlying data reference is has Write access.
 		 *
-		 * If this method's behavior is undefined and will assert if it is called
+		 * If this method's behavior is undefined and will assert if it is called 
 		 * with a mismatched data type or unsupported access type.
 		 */
 		template<typename DataType>
diff --git a/Engine/Plugins/Runtime/Metasound/Source/MetasoundGraphCore/Public/MetasoundPolymorphic.h b/Engine/Plugins/Runtime/Metasound/Source/MetasoundGraphCore/Public/MetasoundPolymorphic.h
--- a/Engine/Plugins/Runtime/Metasound/Source/MetasoundGraphCore/Public/MetasoundPolymorphic.h
+++ b/Engine/Plugins/Runtime/Metasound/Source/MetasoundGraphCore/Public/MetasoundPolymorphic.h
@@ -3,12 +3,9 @@
 #pragma once
 
 #include "MetasoundDataTypeGetTraits.h"
-#include "Concepts/DerivedFrom.h"
 
 namespace Metasound
 {
-	class FOperatorSettings;
-
 	template<typename T>
 	struct TPolymorphicTraits
 	{
@@ -24,56 +21,6 @@ namespace Metasound
 		bool bIsPolymorphic = false;
 	};
 
-	template<typename DataType>
-	concept CIsPolymorphic = TPolymorphicTraits<DataType>::bIsPolymorphic;
-
-	template<typename DataType>
-	concept CIsNotPolymorphic = !TPolymorphicTraits<DataType>::bIsPolymorphic;
-
-	template<typename DataType>
-	concept CIsPolymorphicAndAbstract = TPolymorphicTraits<DataType>::bIsPolymorphic && TPolymorphicTraits<DataType>::bIsAbstract;
-
-	template<typename DataType, typename... ArgTypes>
-	concept CHasFactoryFunction = requires(ArgTypes&&... Args)
-	{
-		{ DataType::CreateNew(Forward<ArgTypes>(Args)...) } -> UE::CDerivedFrom<TSharedRef<DataType, ESPMode::NotThreadSafe>>;
-	};
-
-	// Checks if the Constructor is public and accessible (std::is_constructable does not do that).
-	template <class DataType, class... ArgTypes>
-	concept CHasPublicConstructor =
-	requires(ArgTypes&&... Args)
-	{
-		requires std::is_constructible_v<DataType, ArgTypes...>;
-		::new DataType(Forward<ArgTypes>(Args)...);
-	};
-
-	template<typename DataType>
-	concept CHasFactoryFunctionWithSettings = CHasFactoryFunction<DataType, const FOperatorSettings&>;
-	template<typename DataType, typename... ArgTypes>
-	concept CHasFactoryFunctionWithSettingsAndArgs = CHasFactoryFunction<DataType, const FOperatorSettings&, ArgTypes...>;
-	template<typename DataType, typename... ArgTypes>
-	concept CHasFactoryFunctionWithArgs = CHasFactoryFunction<DataType, ArgTypes...>;
-	template<typename DataType>
-	concept CHasFactoryFunctionDefaultConstruction = CHasFactoryFunction<DataType>;
-
-	template <typename T>
-	concept CHasPolymorphicTypeGetter =
-	requires(const T& t)
-	{
-		{ t.GetDataTypeName() } -> UE::CSameAs<const FName&>;
-	};
-
-	template <typename T>
-	concept CHasPolymorphicTypeIdGetter =
-	requires(const T& t)
-	{
-		{ t.GetDataTypeId() } -> UE::CSameAs<const void*>;
-	};
-
-	template <typename T>
-	concept CHasPolymorphicTypeFunctions = CHasPolymorphicTypeIdGetter<T> && CHasPolymorphicTypeGetter<T>;
-
 	class FPolyRegistry
 	{
 	public:
@@ -116,9 +63,9 @@ namespace Metasound
 			{
 				Info.BaseTypeName = NAME_None;
 			}
-			PRAGMA_DISABLE_INTERNAL_WARNINGS
+			PRAGMA_DISABLE_INTERNAL_WARNINGS			
 			FPolyRegistry::Get().Register(Info);
-			PRAGMA_ENABLE_INTERNAL_WARNINGS
+			PRAGMA_ENABLE_INTERNAL_WARNINGS			
 		}
 	};
 
@@ -127,8 +74,8 @@ namespace Metasound
 
 	UE_EXPERIMENTAL(5.7, "Metasound Polymorphism is still under development")
 	METASOUNDGRAPHCORE_API bool IsCastable(const FName InType, const FName InPotentialBase);
-
-
+	
+	
 }
 
 #define DECLARE_METASOUND_POLY_TYPE(TYPE, BASE, IS_ABSTRACT)\
@@ -136,11 +83,8 @@ namespace Metasound
 	struct Metasound::TPolymorphicTraits<TYPE>\
 	{\
 		static constexpr bool bIsPolymorphic = true;\
-		static constexpr bool bIsAbstract = IS_ABSTRACT || TIsAbstract<TYPE>::Value;\
+		static constexpr bool bIsAbstract = IS_ABSTRACT;\
 		using BaseType = BASE;\
-		/* TODO: Re-enable once all poly types define CreateNew + GetDataTypeName/GetDataTypeId */\
-		/* static_assert((CHasFactoryFunctionWithSettings<TYPE> && bIsAbstract && bIsPolymorphic) || !bIsAbstract, "Abstract Types Must Define a Factory Function 'CreateNew'"); */\
-		/* static_assert(CHasPolymorphicTypeFunctions<TYPE>); */\
 	};
 
 #define REGISTER_METASOUND_POLY_TYPE(T) static Metasound::TPolyRegistrar<T> UE_JOIN(GPolyReg,__COUNTER__);
\ No newline at end of file
diff --git a/Engine/Plugins/Runtime/Metasound/Source/MetasoundGraphCore/Public/MetasoundVariable.h b/Engine/Plugins/Runtime/Metasound/Source/MetasoundGraphCore/Public/MetasoundVariable.h
--- a/Engine/Plugins/Runtime/Metasound/Source/MetasoundGraphCore/Public/MetasoundVariable.h
+++ b/Engine/Plugins/Runtime/Metasound/Source/MetasoundGraphCore/Public/MetasoundVariable.h
@@ -7,7 +7,6 @@
 #include "MetasoundLiteral.h"
 #include "MetasoundOperatorSettings.h"
 #include "MetasoundVertexData.h"
-#include "MetasoundPolymorphic.h"
 
 namespace Metasound
 {
@@ -45,23 +44,13 @@ namespace Metasound
 
 		bool RequiresDelayedDataCopy() const
 		{
-			if constexpr (TPolymorphicTraits<DataType>::bIsPolymorphic)
-			{
-				return false; // Polymorphic types do not support delayed data copy.
-			}
-			else
-			{
-				return DelayedDataReference.IsSet() && DataReference.IsSet();
-			}
+			return DelayedDataReference.IsSet() && DataReference.IsSet();
 		}
 
 		void CopyReferencedData()
 		{
-			if constexpr (!TPolymorphicTraits<DataType>::bIsPolymorphic)
-			{
-				check(RequiresDelayedDataCopy());
-				*(*DelayedDataReference) = *(*DataReference);
-			}
+			check(RequiresDelayedDataCopy());
+			*(*DelayedDataReference) = *(*DataReference);
 		}
 
 		void SetDataReference(FReadReference InDataReference)
@@ -114,13 +103,10 @@ namespace Metasound
 
 		void Reset(const FOperatorSettings& InOperatorSettings)
 		{
-			if constexpr (!TPolymorphicTraits<DataType>::bIsPolymorphic)
+			if (DelayedDataReference.IsSet())
 			{
-				if (DelayedDataReference.IsSet())
-				{
-					FWriteReference DelayedWritable = *DelayedDataReference;
-					*DelayedWritable = TDataTypeLiteralFactory<DataType>::CreateExplicitArgs(InOperatorSettings, Literal);
-				}
+				FWriteReference DelayedWritable = *DelayedDataReference;
+				*DelayedWritable = TDataTypeLiteralFactory<DataType>::CreateExplicitArgs(InOperatorSettings, Literal);
 			}
 		}
 
diff --git a/Engine/Plugins/Runtime/Metasound/Source/MetasoundGraphCore/Public/MetasoundVertexData.h b/Engine/Plugins/Runtime/Metasound/Source/MetasoundGraphCore/Public/MetasoundVertexData.h
--- a/Engine/Plugins/Runtime/Metasound/Source/MetasoundGraphCore/Public/MetasoundVertexData.h
+++ b/Engine/Plugins/Runtime/Metasound/Source/MetasoundGraphCore/Public/MetasoundVertexData.h
@@ -240,7 +240,7 @@ namespace Metasound
 			void Set(const DataReferenceType& InDataReference)
 			{
 				PRAGMA_DISABLE_INTERNAL_WARNINGS
-				check(Metasound::IsCastable(InDataReference.GetDataTypeName(), Vertex.DataTypeName));
+				check(Metasound::IsCastable(Vertex.DataTypeName, InDataReference.GetDataTypeName()));
 				PRAGMA_ENABLE_INTERNAL_WARNINGS
 				
 				Data.Emplace(InDataReference);
