[CAT] Disable auto array types node for abstract types (Cat/Discrete/Composite/Soundfield)
#tests LowLevelTests Audio
#Synced-CL 51594399

diff --git a/Engine/Plugins/Runtime/Metasound/Source/MetasoundFrontend/Public/MetasoundChannelAgnosticType.h b/Engine/Plugins/Runtime/Metasound/Source/MetasoundFrontend/Public/MetasoundChannelAgnosticType.h
--- a/Engine/Plugins/Runtime/Metasound/Source/MetasoundFrontend/Public/MetasoundChannelAgnosticType.h
+++ b/Engine/Plugins/Runtime/Metasound/Source/MetasoundFrontend/Public/MetasoundChannelAgnosticType.h
@@ -98,4 +98,11 @@ namespace Metasound
 	DECLARE_METASOUND_POLY_TYPE(FCompositeChannelAgnosticType,	Metasound::FChannelAgnosticType, true);
 	DECLARE_METASOUND_POLY_TYPE(FDiscreteChannelAgnosticType,	Metasound::FChannelAgnosticType, true);
 	DECLARE_METASOUND_POLY_TYPE(FSoundfieldChannelAgnosticType,	Metasound::FChannelAgnosticType, true);
+
+	// Disable auto-generated array nodes. CAT types are polymorphic/abstract and
+	// should not expose array variants in the MetaSound editor.
+	template<> struct TEnableAutoArrayTypeRegistration<FChannelAgnosticType> { static constexpr bool Value = false; };
+	template<> struct TEnableAutoArrayTypeRegistration<FCompositeChannelAgnosticType> { static constexpr bool Value = false; };
+	template<> struct TEnableAutoArrayTypeRegistration<FDiscreteChannelAgnosticType> { static constexpr bool Value = false; };
+	template<> struct TEnableAutoArrayTypeRegistration<FSoundfieldChannelAgnosticType> { static constexpr bool Value = false; };
 }
\ No newline at end of file
