CL C: GraphCore polymorphic infrastructure (re-land)

Re-land of CL 51656697 (backed out as CL 51673918).

Core type system changes for MetaSound polymorphic data references.
MetasoundPolymorphic.h introduces the polymorphic hierarchy.
DataReference.h/DataFactory.h/Variable.h/VertexData.h modified.
All changes UE_EXPERIMENTAL guarded and existing code paths
preserved via if constexpr.

Added GetOrCreateDefaultDataReadReference tests:
- Catch2 unit tests for polymorphic abstract types (unbound/bound)
and non-polymorphic baseline
- EngineTest automation tests for FWaveAsset (unbound/bound)

#tests LowLevelTests Audio
#Synced-CL 51594399

diff --git a/Engine/Plugins/Runtime/Metasound/Source/MetasoundEngineTest/Private/EngineTestGetOrCreateDefault.cpp b/Engine/Plugins/Runtime/Metasound/Source/MetasoundEngineTest/Private/EngineTestGetOrCreateDefault.cpp
new file mode 100644
--- /dev/null
+++ b/Engine/Plugins/Runtime/Metasound/Source/MetasoundEngineTest/Private/EngineTestGetOrCreateDefault.cpp
@@ -0,0 +1,63 @@
+// Copyright Epic Games, Inc. All Rights Reserved.
+
+#include "MetasoundOperatorSettings.h"
+#include "MetasoundVertex.h"
+#include "MetasoundVertexData.h"
+#include "MetasoundWave.h"
+#include "Misc/AutomationTest.h"
+
+#if WITH_DEV_AUTOMATION_TESTS
+
+IMPLEMENT_SIMPLE_AUTOMATION_TEST(FMetasoundGetOrCreateDefault_WaveAsset_Unbound, "Audio.Metasound.GetOrCreateDefault.WaveAsset.Unbound", EAutomationTestFlags::EditorContext | EAutomationTestFlags::EngineFilter)
+bool FMetasoundGetOrCreateDefault_WaveAsset_Unbound::RunTest(const FString& Parameters)
+{
+	using namespace Metasound;
+
+	const FOperatorSettings Settings(48000, 480);
+	const FVertexName VertexName = TEXT("TestWaveInput");
+
+	// Build an input vertex interface with a single FWaveAsset vertex (unbound).
+	FInputVertexInterface InputInterface
+	{
+		TInputDataVertex<FWaveAsset>(VertexName, FDataVertexMetadata{})
+	};
+	FInputVertexInterfaceData InterfaceData(InputInterface);
+
+	// Call GetOrCreateDefaultDataReadReference — should create a default FWaveAsset via factory.
+	FWaveAssetReadRef ReadRef = InterfaceData.GetOrCreateDefaultDataReadReference<FWaveAsset>(VertexName, Settings);
+
+	// A default-constructed FWaveAsset has no valid sound wave.
+	TestFalse(TEXT("Default FWaveAsset should not have a valid SoundWave"), ReadRef->IsSoundWaveValid());
+
+	return true;
+}
+
+IMPLEMENT_SIMPLE_AUTOMATION_TEST(FMetasoundGetOrCreateDefault_WaveAsset_Bound, "Audio.Metasound.GetOrCreateDefault.WaveAsset.Bound", EAutomationTestFlags::EditorContext | EAutomationTestFlags::EngineFilter)
+bool FMetasoundGetOrCreateDefault_WaveAsset_Bound::RunTest(const FString& Parameters)
+{
+	using namespace Metasound;
+
+	const FOperatorSettings Settings(48000, 480);
+	const FVertexName VertexName = TEXT("TestWaveInput");
+
+	// Build an input vertex interface with a single FWaveAsset vertex.
+	FInputVertexInterface InputInterface
+	{
+		TInputDataVertex<FWaveAsset>(VertexName, FDataVertexMetadata{})
+	};
+	FInputVertexInterfaceData InterfaceData(InputInterface);
+
+	// Pre-bind a write ref to the vertex.
+	FWaveAssetWriteRef WriteRef = FWaveAssetWriteRef::CreateNew();
+	InterfaceData.BindReadVertex<FWaveAsset>(VertexName, WriteRef);
+
+	// Call GetOrCreateDefaultDataReadReference — should return the bound ref.
+	FWaveAssetReadRef ReadRef = InterfaceData.GetOrCreateDefaultDataReadReference<FWaveAsset>(VertexName, Settings);
+
+	// Verify we got back the same underlying object that was bound.
+	TestEqual(TEXT("Bound ref should be returned"), static_cast<const void*>(&(*ReadRef)), static_cast<const void*>(&(*WriteRef)));
+
+	return true;
+}
+
+#endif // WITH_DEV_AUTOMATION_TESTS
diff --git a/Engine/Plugins/Runtime/Metasound/Source/MetasoundGraphCore/Public/MetasoundDataFactory.h b/Engine/Plugins/Runtime/Metasound/Source/MetasoundGraphCore/Public/MetasoundDataFactory.h
--- a/Engine/Plugins/Runtime/Metasound/Source/MetasoundGraphCore/Public/MetasoundDataFactory.h
+++ b/Engine/Plugins/Runtime/Metasound/Source/MetasoundGraphCore/Public/MetasoundDataFactory.h
@@ -30,10 +30,10 @@ namespace Metasound
 			static constexpr bool bIsArrayConstructibleWithSettingsAndArgs = false;
 
 			// DataType Construction
-			static constexpr bool bIsDefaultConstructible = std::is_constructible<DataType>::value;
-			static constexpr bool bIsConstructibleWithSettings = std::is_constructible<DataType, const FOperatorSettings&>::value;
-			static constexpr bool bIsConstructibleWithArgs = std::is_constructible<DataType, ArgTypes...>::value;
-			static constexpr bool bIsConstructibleWithSettingsAndArgs = std::is_constructible<DataType, const FOperatorSettings&, ArgTypes...>::value;
+			static constexpr bool bIsDefaultConstructible = std::is_constructible<DataType>::value || CHasFactoryFunctionDefaultConstruction<DataType>;
+			static constexpr bool bIsConstructibleWithSettings = std::is_constructible<DataType, const FOperatorSettings&>::value || CHasFactoryFunctionWithSettings<DataType>;
+			static constexpr bool bIsConstructibleWithArgs = std::is_constructible<DataType, ArgTypes...>::value || CHasFactoryFunctionWithArgs<DataType, ArgTypes...>;
+			static constexpr bool bIsConstructibleWithSettingsAndArgs = std::is_constructible<DataType, const FOperatorSettings&, ArgTypes...>::value || CHasFactoryFunctionWithSettingsAndArgs<DataType, ArgTypes...>;
 		};
 
 		/** Partial specialization for TArray<> types. A TArray can be constructed elementwise with
@@ -153,7 +153,7 @@ namespace Metasound
 		{
 			using TConstructorTraits = TDataTypeConstructorTraits<DataType, ArgTypes...>;
 
-			static constexpr bool bForwardSettingsAndArgs = TConstructorTraits::bIsConstructibleWithSettingsAndArgs;
+			static constexpr bool bForwardSettingsAndArgs = TConstructorTraits::bIsConstructibleWithSettingsAndArgs || CHasFactoryFunctionWithSettingsAndArgs<DataType, ArgTypes...>;
 			static constexpr bool bForwardArgs = false;
 			static constexpr bool bForwardSettings = false;
 			static constexpr bool bForwardNone = false;
@@ -172,8 +172,8 @@ namespace Metasound
 		{
 			using TConstructorTraits = TDataTypeConstructorTraits<DataType, ArgTypes...>;
 
-			static constexpr bool bForwardSettingsAndArgs = TConstructorTraits::bIsConstructibleWithSettingsAndArgs;
-			static constexpr bool bForwardArgs = !bForwardSettingsAndArgs && TConstructorTraits::bIsConstructibleWithArgs;
+			static constexpr bool bForwardSettingsAndArgs = TConstructorTraits::bIsConstructibleWithSettingsAndArgs || CHasFactoryFunctionWithSettingsAndArgs<DataType, ArgTypes...>;
+			static constexpr bool bForwardArgs = !bForwardSettingsAndArgs && (TConstructorTraits::bIsConstructibleWithArgs || CHasFactoryFunctionWithArgs<DataType, ArgTypes...>); 
 			static constexpr bool bForwardSettings = false;
 			static constexpr bool bForwardNone = false;
 
@@ -191,10 +191,10 @@ namespace Metasound
 		{
 			using TConstructorTraits = TDataTypeConstructorTraits<DataType, ArgTypes...>;
 
-			static constexpr bool bForwardSettingsAndArgs = TConstructorTraits::bIsConstructibleWithSettingsAndArgs;
-			static constexpr bool bForwardArgs = !bForwardSettingsAndArgs && TConstructorTraits::bIsConstructibleWithArgs;
-			static constexpr bool bForwardSettings = !(bForwardArgs || bForwardSettingsAndArgs) && TConstructorTraits::bIsConstructibleWithSettings;
-			static constexpr bool bForwardNone = !(bForwardArgs || bForwardSettingsAndArgs || bForwardSettings) && TConstructorTraits::bIsDefaultConstructible;
+			static constexpr bool bForwardSettingsAndArgs = TConstructorTraits::bIsConstructibleWithSettingsAndArgs || CHasFactoryFunctionWithSettingsAndArgs<DataType, ArgTypes...>;
+			static constexpr bool bForwardArgs = !bForwardSettingsAndArgs && (TConstructorTraits::bIsConstructibleWithArgs || CHasFactoryFunctionWithArgs<DataType, ArgTypes...>);
+			static constexpr bool bForwardSettings = !(bForwardArgs || bForwardSettingsAndArgs) && (TConstructorTraits::bIsConstructibleWithSettings || CHasFactoryFunctionWithSettings<DataType>);
+			static constexpr bool bForwardNone = !(bForwardArgs || bForwardSettingsAndArgs || bForwardSettings) && (TConstructorTraits::bIsDefaultConstructible || CHasFactoryFunctionDefaultConstruction<DataType>);
 
 			static constexpr bool bCannotForwardToConstructor = !(bForwardSettingsAndArgs || bForwardArgs || bForwardSettings || bForwardNone);
 		};
diff --git a/Engine/Plugins/Runtime/Metasound/Source/MetasoundGraphCore/Public/MetasoundDataReference.h b/Engine/Plugins/Runtime/Metasound/Source/MetasoundGraphCore/Public/MetasoundDataReference.h
--- a/Engine/Plugins/Runtime/Metasound/Source/MetasoundGraphCore/Public/MetasoundDataReference.h
+++ b/Engine/Plugins/Runtime/Metasound/Source/MetasoundGraphCore/Public/MetasoundDataReference.h
@@ -4,6 +4,7 @@
 
 #include <type_traits>
 
+#include "Concepts/DerivedFrom.h"
 #include "Containers/UnrealString.h"
 #include "HAL/Platform.h"
 #include "Internationalization/Text.h"
@@ -24,6 +25,7 @@ namespace Audio
 
 namespace Metasound
 {
+	class FOperatorSettings;
 	// Forward declare for LexToString
 	enum class EDataReferenceAccessType : uint8;
 }
@@ -142,51 +144,77 @@ namespace Metasound
 	{
 		template <typename Base, typename Derived, typename... ArgTypes>
 		UE_EXPERIMENTAL(5.7, "Polymorphic API is under development")
-		TSharedRef<Base, ESPMode::NotThreadSafe> MakeRefOf(ArgTypes&&... args)
+		TSharedRef<Base, ESPMode::NotThreadSafe> MakeRefOf(ArgTypes&&... Args)
 		{
 			static_assert(std::is_base_of_v<Base, Derived>, "Target must derive from Base (or be Base)");
 			static_assert(std::is_constructible_v<Derived, ArgTypes...>, "Args do not match Target's constructor");
 
-			TSharedRef<Derived, ESPMode::NotThreadSafe> DerivedRef = MakeShared<Derived, ESPMode::NotThreadSafe>(Forward<ArgTypes>(args)...);
- 
+			TSharedRef<Derived, ESPMode::NotThreadSafe> DerivedRef = MakeShared<Derived, ESPMode::NotThreadSafe>(Forward<ArgTypes>(Args)...);
+
 			return StaticCastSharedRef<Base>(DerivedRef);
 		}
+
+		template <typename Base, typename Derived, typename... ArgTypes>
+		UE_EXPERIMENTAL(5.7, "Polymorphic API is under development")
+		TSharedRef<Base, ESPMode::NotThreadSafe> MakeRefOf(ArgTypes&&... Args) requires CHasFactoryFunction<Derived, ArgTypes...>
+		{
+			static_assert(std::is_base_of_v<Base, Derived>, "Target must derive from Base (or be Base)");
+			TSharedRef<Derived, ESPMode::NotThreadSafe> DerivedRef = Derived::CreateNew(Forward<ArgTypes>(Args)...);
+ 			return StaticCastSharedRef<Base>(DerivedRef);
+		}
 	}
 
 	template <typename DataType>
 	struct TDynamicDataReferencePolicy;
 
 	template<typename DataType>
-	requires (TPolymorphicTraits<DataType>::bIsPolymorphic)
+	requires (!TPolymorphicTraits<DataType>::bIsPolymorphic)
 	struct TDynamicDataReferencePolicy<DataType>
 	{
-		explicit TDynamicDataReferencePolicy(const FName& InDataTypeName = GetMetasoundDataTypeName<DataType>(), const void* InTypeId = GetMetasoundDataTypeId<DataType>())
-			: DataTypeName(InDataTypeName), DataTypeId(InTypeId)
-		{}
-
-		template<typename QueryType> bool IsA() const { return IsChildOfByName(DataTypeName, GetMetasoundDataTypeName<QueryType>()); }
-		const FName& GetDataTypeName() const { return DataTypeName; }
-		const void* GetTypeId() const { return DataTypeId; }
-		FName DataTypeName;
-		const void* DataTypeId;
+		template<typename QueryType> bool IsA(const DataType*) const { return GetMetasoundDataTypeName<DataType>() == GetMetasoundDataTypeName<QueryType>(); }
+		const FName& GetDataTypeName(const DataType*) const { return GetMetasoundDataTypeName<DataType>(); }
+		const void* GetDataTypeId(const DataType*) const { return GetMetasoundDataTypeId<DataType>(); }
 	};
-	
+
 	template<typename DataType>
-	requires (!TPolymorphicTraits<DataType>::bIsPolymorphic)
+	requires (TPolymorphicTraits<DataType>::bIsPolymorphic)
 	struct TDynamicDataReferencePolicy<DataType>
 	{
-		TDynamicDataReferencePolicy() = default;
-		TDynamicDataReferencePolicy(const FName&, const void*) {}
+		// Note: CHasPolymorphicTypeFunctions is not required here because the requires-based
+		// fallback overloads below handle types that don't implement GetDataTypeName()/GetDataTypeId().
+
+		template<typename QueryType> bool IsA(const DataType* InOuter) const
+		{
+			return IsChildOfByName(GetDataTypeName(InOuter), GetMetasoundDataTypeName<QueryType>());
+		}
 
-		template<typename QueryType> bool IsA() const { return GetDataTypeName() == GetMetasoundDataTypeName<QueryType>(); }
-		const FName& GetDataTypeName() const { return GetMetasoundDataTypeName<DataType>(); }
-		const void* GetTypeId() const { return GetMetasoundDataTypeId<DataType>(); }
+		const FName& GetDataTypeName(const DataType* InOuter) const requires(CHasPolymorphicTypeGetter<DataType>)
+		{
+			return InOuter->GetDataTypeName();
+		}
+
+		// Fallback when type doesn't yet implement GetDataTypeName().
+		const FName& GetDataTypeName(const DataType*) const requires(!CHasPolymorphicTypeGetter<DataType>)
+		{
+			return GetMetasoundDataTypeName<DataType>();
+		}
+
+		const void* GetDataTypeId(const DataType* InOuter) const requires(CHasPolymorphicTypeIdGetter<DataType>)
+		{
+			return InOuter->GetDataTypeId();
+		}
+
+		// Fallback when type doesn't yet implement GetDataTypeId().
+		const void* GetDataTypeId(const DataType*) const requires(!CHasPolymorphicTypeIdGetter<DataType>)
+		{
+			return GetMetasoundDataTypeId<DataType>();
+		}
 	};
-	
-	/** Template class for a paramter reference. 
+
+	/** Template class for a paramter reference.
 	 *
 	 * This fulfills the IParamterRef interface, utilizing TDataReferenceTypeInfo to
-	 * define the the TypeName and TypeId of the parameter. 
+	 * define the the TypeName and TypeId of the parameter.
 	 */
 	template <typename DataType>
 	class TDataReference : public IDataReference
@@ -202,6 +230,12 @@ namespace Metasound
 		{
 		}
 
+		template <typename... ArgTypes>
+		TDataReference(EDataRefShouldConstruct InToken, ArgTypes&&... Args) requires(CIsPolymorphicAndAbstract<DataType> && CHasFactoryFunction<DataType, ArgTypes...>)
+			: ObjectReference(DataType::CreateNew(Forward<ArgTypes>(Args)...))
+		{
+		}
+
 		typedef TSharedRef<DataType, ESPMode::NotThreadSafe> FRefType;
 
 		explicit TDataReference(EDataRefShouldConstruct /*InToken*/, FRefType InRef)
@@ -213,7 +247,6 @@ namespace Metasound
 								const FName InTypeName,
 								const void* InTypeId)
 			: ObjectReference(MoveTemp(InRef))
-			, ObjectReferencePolicy(InTypeName, InTypeId)
 		{}
 	
 	public:
@@ -229,19 +262,25 @@ namespace Metasound
 			return TDataReference<DataType>(EDataRefShouldConstruct::NewObject, Forward<ArgTypes>(Args)...);
 		}
 
+		/**
+		 * Polymorphic form of CreateNew.
+		 */
+		template <typename... ArgTypes>
+		static TDataReference<DataType> CreateNew(ArgTypes&&... Args) requires (CIsPolymorphicAndAbstract<DataType>)
+		{
+			static_assert(CHasFactoryFunction<DataType, ArgTypes...>);
+			return TDataReference<DataType>(EDataRefShouldConstruct::NewObject, Forward<ArgTypes>(Args)...);
+		}
+
 		// Constructs Derived but stores as Base (DataType).
 		template <typename Derived, typename... ArgTypes>
 		UE_EXPERIMENTAL(5.7, "Polymorphic API is under development")
 		static TDataReference CreateNewDerived(ArgTypes&&... Args)
 		{
 			static_assert(std::is_base_of_v<DataType, Derived>, "Derived must inherit from DataType");
-			static_assert(std::is_constructible_v<Derived, ArgTypes...>, "Args don't match a Derived constructor");
-			
 			return TDataReference(
 				EDataRefShouldConstruct::NewObject,
-				MetasoundDataReferencePrivate::MakeRefOf<DataType, Derived>(Forward<ArgTypes>(Args)...),
-				GetMetasoundDataTypeName<Derived>(),
-			    GetMetasoundDataTypeId<Derived>()
+				MetasoundDataReferencePrivate::MakeRefOf<DataType, Derived>(Forward<ArgTypes>(Args)...)
 			);
 		}
 
@@ -252,8 +291,6 @@ namespace Metasound
 												 ArgTypes&&... Args)
 		{
 			static_assert(std::is_base_of_v<DataType, Derived>, "Derived must inherit from DataType");
-			static_assert(std::is_constructible_v<Derived, ArgTypes...>, "Args don't match a Derived constructor");
-
 		#if DO_CHECK
 			PRAGMA_DISABLE_EXPERIMENTAL_WARNINGS
 			check(Metasound::IsChildOfByName(InDerivedTypeName, GetMetasoundDataTypeName<DataType>()));
@@ -284,13 +321,13 @@ namespace Metasound
 		/** Return the name of the underlying type. */
 		virtual const FName& GetDataTypeName() const override
 		{
-			return ObjectReferencePolicy.GetDataTypeName();	
+			return ObjectReferencePolicy.GetDataTypeName(&ObjectReference.Get());
 		}
 
 		/** Return the ID of the underlying type. */
 		virtual const void* const GetDataTypeId() const override
 		{
-			return ObjectReferencePolicy.GetTypeId();
+			return ObjectReferencePolicy.GetDataTypeId(&ObjectReference.Get());
 		}
 
 		/** Return a raw pointer to the data. */
@@ -303,21 +340,16 @@ namespace Metasound
 		UE_EXPERIMENTAL(5.7, "Polymorphic API is under development")
 		bool IsA() const
 		{
-			return this->ObjectReferencePolicy.template IsA<QueryType>();
+			return this->ObjectReferencePolicy.template IsA<QueryType>(&ObjectReference.Get());
 		}
-		
+
 		template <typename CastType>
 		UE_EXPERIMENTAL(5.7, "Polymorphic API is under development")
-		TDataReference<CastType> TryCast() const
+		TOptional<TSharedRef<CastType, ESPMode::NotThreadSafe>> TryCast() const
 		{
 			if (IsA<CastType>())
 			{
-				return TDataReference(
-					EDataRefShouldConstruct::NewObject,
-					MetasoundDataReferencePrivate::MakeRefOf<DataType, CastType>(ObjectReference),
-					GetMetasoundDataTypeName<CastType>(),
-					GetMetasoundDataTypeId<CastType>()
-				);
+				return StaticCastSharedRef<CastType>(ObjectReference);
 			}
 			return {};
 		}
@@ -378,13 +410,17 @@ namespace Metasound
 			return TDataValueReference<DataType>(EDataRefShouldConstruct::NewObject, Forward<ArgTypes>(Args)...);
 		}
 
+		template <typename... ArgTypes>
+		static TDataValueReference<DataType> CreateNew(ArgTypes&&... Args) requires (CIsPolymorphicAndAbstract<DataType> && CHasFactoryFunction<DataType, ArgTypes...>)
+		{
+			return TDataValueReference<DataType>(EDataRefShouldConstruct::NewObject, Forward<ArgTypes>(Args)...);
+		}
+
 		template <typename Derived, typename... ArgTypes>
 		UE_EXPERIMENTAL(5.7, "Polymorphic API is under development")
 		static TDataValueReference CreateNewDerived(ArgTypes&&... Args)
 		{
 			static_assert(std::is_base_of_v<DataType, Derived>, "Derived must inherit from DataType");
-			static_assert(std::is_constructible_v<Derived, ArgTypes...>, "Args don't match a Derived constructor");
-			
 			return TDataValueReference(
 				EDataRefShouldConstruct::NewObject,
 				MetasoundDataReferencePrivate::MakeRefOf<DataType, Derived>(Forward<ArgTypes>(Args)...)
@@ -398,8 +434,6 @@ namespace Metasound
 													 ArgTypes&&... Args)
 		{
 			static_assert(std::is_base_of_v<DataType, Derived>, "Derived must inherit from DataType");
-			static_assert(std::is_constructible_v<Derived, ArgTypes...>, "Args don't match a Derived constructor");
-
 		#if DO_CHECK
 			PRAGMA_DISABLE_EXPERIMENTAL_WARNINGS
 			check(Metasound::IsChildOfByName(InDerivedTypeName, GetMetasoundDataTypeName<DataType>()));
@@ -522,13 +556,18 @@ namespace Metasound
 			return TDataWriteReference<DataType>(EDataRefShouldConstruct::NewObject, Forward<ArgTypes>(Args)...);
 		}
 
+		/** This should be used to construct a new DataType object and return this TDataWriteReference as a wrapper around it. */
+		template <typename... ArgTypes>
+		static TDataWriteReference<DataType> CreateNew(ArgTypes&&... Args) requires (CIsPolymorphicAndAbstract<DataType> && CHasFactoryFunction<DataType, ArgTypes...>)
+		{
+			return TDataWriteReference<DataType>(EDataRefShouldConstruct::NewObject, Forward<ArgTypes>(Args)...);
+		}
+
 		template <typename Derived, typename... ArgTypes>
 		UE_EXPERIMENTAL(5.8, "Polymorphic API is under development")
 		static TDataWriteReference CreateNewDerived(ArgTypes&&... Args)
 		{
 			static_assert(std::is_base_of_v<DataType, Derived>, "Derived must inherit from DataType");
-			static_assert(std::is_constructible_v<Derived, ArgTypes...>, "Args don't match a Derived constructor");
-			
 			return TDataWriteReference(
 				EDataRefShouldConstruct::NewObject,
 				MetasoundDataReferencePrivate::MakeRefOf<DataType, Derived>(Forward<ArgTypes>(Args)...)
@@ -536,14 +575,11 @@ namespace Metasound
 		}
 
 		template <typename Derived, typename... ArgTypes>
-		UE_EXPERIMENTAL(5.8, "Polymorphic API is under development")
 		static TDataWriteReference CreateNewDerivedAs(const FName& InDerivedTypeName,
 													 FMetasoundDataTypeId InDerivedTypeId,
 													 ArgTypes&&... Args)
 		{
 			static_assert(std::is_base_of_v<DataType, Derived>, "Derived must inherit from DataType");
-			static_assert(std::is_constructible_v<Derived, ArgTypes...>, "Args don't match a Derived constructor");
-
 		#if DO_CHECK
 			PRAGMA_DISABLE_EXPERIMENTAL_WARNINGS
 			check(Metasound::IsChildOfByName(InDerivedTypeName, GetMetasoundDataTypeName<DataType>()));
@@ -607,10 +643,10 @@ namespace Metasound
 
 	};
 
-	/** Cast a TDataReadReference to a TDataWriteReference. 
+	/** Cast a TDataReadReference to a TDataWriteReference.
 	 *
 	 * In general TDataReadReferences should not be converted into TDataWriteReferences unless the caller
-	 * can be certain that no other TDataWriteReference exists for the underlying parameter. Having multiple 
+	 * can be certain that no other TDataWriteReference exists for the underlying parameter. Having multiple
 	 * TDataWriteReferences to the same parameter can cause confusion behavior as values are overwritten in
 	 * an underterministic fashion.
 	 */
@@ -650,13 +686,17 @@ namespace Metasound
 			return TDataReadReference<DataType>(EDataRefShouldConstruct::NewObject, Forward<ArgTypes>(Args)...);
 		}
 
+		template <typename... ArgTypes>
+		static TDataReadReference<DataType> CreateNew(ArgTypes&&... Args) requires (CIsPolymorphicAndAbstract<DataType> && CHasFactoryFunction<DataType, ArgTypes...>)
+		{
+			return TDataReadReference<DataType>(EDataRefShouldConstruct::NewObject, Forward<ArgTypes>(Args)...);
+		}
+
 		template <typename Derived, typename... ArgTypes>
 		UE_EXPERIMENTAL(5.8, "Polymorphic API is under development")
 		static TDataReadReference CreateNewDerived(ArgTypes&&... Args)
 		{
 			static_assert(std::is_base_of_v<DataType, Derived>, "Derived must inherit from DataType");
-			static_assert(std::is_constructible_v<Derived, ArgTypes...>, "Args don't match a Derived constructor");
-			
 			return TDataReadReference(
 				EDataRefShouldConstruct::NewObject,
 				MetasoundDataReferencePrivate::MakeRefOf<DataType, Derived>(Forward<ArgTypes>(Args)...)
@@ -668,8 +708,6 @@ namespace Metasound
 		static TDataReadReference CreateNewDerivedAs(const FName& InDerivedTypeName, FMetasoundDataTypeId InDerivedTypeId, ArgTypes&&... Args)
 		{
 			static_assert(std::is_base_of_v<DataType, Derived>, "Derived must inherit from DataType");
-			static_assert(std::is_constructible_v<Derived, ArgTypes...>, "Args don't match a Derived constructor");
-
 		#if DO_CHECK
 			PRAGMA_DISABLE_EXPERIMENTAL_WARNINGS
 			check(Metasound::IsChildOfByName(InDerivedTypeName, GetMetasoundDataTypeName<DataType>()));
diff --git a/Engine/Plugins/Runtime/Metasound/Source/MetasoundGraphCore/Public/MetasoundPolymorphic.h b/Engine/Plugins/Runtime/Metasound/Source/MetasoundGraphCore/Public/MetasoundPolymorphic.h
--- a/Engine/Plugins/Runtime/Metasound/Source/MetasoundGraphCore/Public/MetasoundPolymorphic.h
+++ b/Engine/Plugins/Runtime/Metasound/Source/MetasoundGraphCore/Public/MetasoundPolymorphic.h
@@ -3,9 +3,12 @@
 #pragma once
 
 #include "MetasoundDataTypeGetTraits.h"
+#include "Concepts/DerivedFrom.h"
 
 namespace Metasound
 {
+	class FOperatorSettings;
+
 	template<typename T>
 	struct TPolymorphicTraits
 	{
@@ -21,6 +24,56 @@ namespace Metasound
 		bool bIsPolymorphic = false;
 	};
 
+	template<typename DataType>
+	concept CIsPolymorphic = TPolymorphicTraits<DataType>::bIsPolymorphic;
+
+	template<typename DataType>
+	concept CIsNotPolymorphic = !TPolymorphicTraits<DataType>::bIsPolymorphic;
+
+	template<typename DataType>
+	concept CIsPolymorphicAndAbstract = TPolymorphicTraits<DataType>::bIsPolymorphic && TPolymorphicTraits<DataType>::bIsAbstract;
+
+	template<typename DataType, typename... ArgTypes>
+	concept CHasFactoryFunction = requires(ArgTypes&&... Args)
+	{
+		{ DataType::CreateNew(Forward<ArgTypes>(Args)...) } -> UE::CDerivedFrom<TSharedRef<DataType, ESPMode::NotThreadSafe>>;
+	};
+
+	// Checks if the Constructor is public and accessible (std::is_constructable does not do that).
+	template <class DataType, class... ArgTypes>
+	concept CHasPublicConstructor =
+	requires(ArgTypes&&... Args)
+	{
+		requires std::is_constructible_v<DataType, ArgTypes...>;
+		::new DataType(Forward<ArgTypes>(Args)...);
+	};
+
+	template<typename DataType>
+	concept CHasFactoryFunctionWithSettings = CHasFactoryFunction<DataType, const FOperatorSettings&>;
+	template<typename DataType, typename... ArgTypes>
+	concept CHasFactoryFunctionWithSettingsAndArgs = CHasFactoryFunction<DataType, const FOperatorSettings&, ArgTypes...>;
+	template<typename DataType, typename... ArgTypes>
+	concept CHasFactoryFunctionWithArgs = CHasFactoryFunction<DataType, ArgTypes...>;
+	template<typename DataType>
+	concept CHasFactoryFunctionDefaultConstruction = CHasFactoryFunction<DataType>;
+
+	template <typename T>
+	concept CHasPolymorphicTypeGetter =
+	requires(const T& t)
+	{
+		{ t.GetDataTypeName() } -> UE::CSameAs<const FName&>;
+	};
+
+	template <typename T>
+	concept CHasPolymorphicTypeIdGetter =
+	requires(const T& t)
+	{
+		{ t.GetDataTypeId() } -> UE::CSameAs<const void*>;
+	};
+
+	template <typename T>
+	concept CHasPolymorphicTypeFunctions = CHasPolymorphicTypeIdGetter<T> && CHasPolymorphicTypeGetter<T>;
+
 	class FPolyRegistry
 	{
 	public:
@@ -63,9 +116,9 @@ namespace Metasound
 			{
 				Info.BaseTypeName = NAME_None;
 			}
-			PRAGMA_DISABLE_INTERNAL_WARNINGS			
+			PRAGMA_DISABLE_INTERNAL_WARNINGS
 			FPolyRegistry::Get().Register(Info);
-			PRAGMA_ENABLE_INTERNAL_WARNINGS			
+			PRAGMA_ENABLE_INTERNAL_WARNINGS
 		}
 	};
 
@@ -74,8 +127,8 @@ namespace Metasound
 
 	UE_EXPERIMENTAL(5.7, "Metasound Polymorphism is still under development")
 	METASOUNDGRAPHCORE_API bool IsCastable(const FName InType, const FName InPotentialBase);
-	
-	
+
+
 }
 
 #define DECLARE_METASOUND_POLY_TYPE(TYPE, BASE, IS_ABSTRACT)\
@@ -83,8 +136,11 @@ namespace Metasound
 	struct Metasound::TPolymorphicTraits<TYPE>\
 	{\
 		static constexpr bool bIsPolymorphic = true;\
-		static constexpr bool bIsAbstract = IS_ABSTRACT;\
+		static constexpr bool bIsAbstract = IS_ABSTRACT || TIsAbstract<TYPE>::Value;\
 		using BaseType = BASE;\
+		/* TODO: Re-enable once all poly types define CreateNew + GetDataTypeName/GetDataTypeId */\
+		/* static_assert((CHasFactoryFunctionWithSettings<TYPE> && bIsAbstract && bIsPolymorphic) || !bIsAbstract, "Abstract Types Must Define a Factory Function 'CreateNew'"); */\
+		/* static_assert(CHasPolymorphicTypeFunctions<TYPE>); */\
 	};
 
 #define REGISTER_METASOUND_POLY_TYPE(T) static Metasound::TPolyRegistrar<T> UE_JOIN(GPolyReg,__COUNTER__);
\ No newline at end of file
diff --git a/Engine/Plugins/Runtime/Metasound/Source/MetasoundGraphCore/Public/MetasoundVariable.h b/Engine/Plugins/Runtime/Metasound/Source/MetasoundGraphCore/Public/MetasoundVariable.h
--- a/Engine/Plugins/Runtime/Metasound/Source/MetasoundGraphCore/Public/MetasoundVariable.h
+++ b/Engine/Plugins/Runtime/Metasound/Source/MetasoundGraphCore/Public/MetasoundVariable.h
@@ -7,6 +7,7 @@
 #include "MetasoundLiteral.h"
 #include "MetasoundOperatorSettings.h"
 #include "MetasoundVertexData.h"
+#include "MetasoundPolymorphic.h"
 
 namespace Metasound
 {
@@ -44,13 +45,23 @@ namespace Metasound
 
 		bool RequiresDelayedDataCopy() const
 		{
-			return DelayedDataReference.IsSet() && DataReference.IsSet();
+			if constexpr (TPolymorphicTraits<DataType>::bIsPolymorphic)
+			{
+				return false; // Polymorphic types do not support delayed data copy.
+			}
+			else
+			{
+				return DelayedDataReference.IsSet() && DataReference.IsSet();
+			}
 		}
 
 		void CopyReferencedData()
 		{
-			check(RequiresDelayedDataCopy());
-			*(*DelayedDataReference) = *(*DataReference);
+			if constexpr (!TPolymorphicTraits<DataType>::bIsPolymorphic)
+			{
+				check(RequiresDelayedDataCopy());
+				*(*DelayedDataReference) = *(*DataReference);
+			}
 		}
 
 		void SetDataReference(FReadReference InDataReference)
@@ -103,10 +114,13 @@ namespace Metasound
 
 		void Reset(const FOperatorSettings& InOperatorSettings)
 		{
-			if (DelayedDataReference.IsSet())
+			if constexpr (!TPolymorphicTraits<DataType>::bIsPolymorphic)
 			{
-				FWriteReference DelayedWritable = *DelayedDataReference;
-				*DelayedWritable = TDataTypeLiteralFactory<DataType>::CreateExplicitArgs(InOperatorSettings, Literal);
+				if (DelayedDataReference.IsSet())
+				{
+					FWriteReference DelayedWritable = *DelayedDataReference;
+					*DelayedWritable = TDataTypeLiteralFactory<DataType>::CreateExplicitArgs(InOperatorSettings, Literal);
+				}
 			}
 		}
 
diff --git a/Engine/Plugins/Runtime/Metasound/Source/MetasoundGraphCore/Public/MetasoundVertexData.h b/Engine/Plugins/Runtime/Metasound/Source/MetasoundGraphCore/Public/MetasoundVertexData.h
--- a/Engine/Plugins/Runtime/Metasound/Source/MetasoundGraphCore/Public/MetasoundVertexData.h
+++ b/Engine/Plugins/Runtime/Metasound/Source/MetasoundGraphCore/Public/MetasoundVertexData.h
@@ -240,7 +240,7 @@ namespace Metasound
 			void Set(const DataReferenceType& InDataReference)
 			{
 				PRAGMA_DISABLE_INTERNAL_WARNINGS
-				check(Metasound::IsCastable(Vertex.DataTypeName, InDataReference.GetDataTypeName()));
+				check(Metasound::IsCastable(InDataReference.GetDataTypeName(), Vertex.DataTypeName));
 				PRAGMA_ENABLE_INTERNAL_WARNINGS
 				
 				Data.Emplace(InDataReference);
