Native MetaSound nodes can define a custom icon and color

[RN] minor framework/audio
#Synced-CL 51544635

diff --git a/Engine/Plugins/Runtime/Metasound/Source/MetasoundEditor/Private/MetasoundEditorGraphNode.cpp b/Engine/Plugins/Runtime/Metasound/Source/MetasoundEditor/Private/MetasoundEditorGraphNode.cpp
--- a/Engine/Plugins/Runtime/Metasound/Source/MetasoundEditor/Private/MetasoundEditorGraphNode.cpp
+++ b/Engine/Plugins/Runtime/Metasound/Source/MetasoundEditor/Private/MetasoundEditorGraphNode.cpp
@@ -1589,14 +1589,12 @@ FLinearColor UMetasoundEditorGraphExternalNode::GetNodeTitleColor() const
 
 FSlateIcon UMetasoundEditorGraphExternalNode::GetNodeTitleIcon() const
 {
-	if (bIsClassNative)
-	{
-		return FSlateIcon("MetaSoundStyle", "MetasoundEditor.Graph.Node.Class.Native");
-	}
-	else
+	if (!bIsClassNative)
 	{
 		return FSlateIcon("MetasoundStyle", "MetasoundEditor.Graph.Node.Class.Graph");
 	}
+
+	return FSlateIcon();
 }
 
 bool UMetasoundEditorGraphExternalNode::ShouldDrawNodeAsControlPointOnly(int32& OutInputPinIndex, int32& OutOutputPinIndex) const
diff --git a/Engine/Plugins/Runtime/Metasound/Source/MetasoundEditor/Private/SMetasoundGraphNode.cpp b/Engine/Plugins/Runtime/Metasound/Source/MetasoundEditor/Private/SMetasoundGraphNode.cpp
--- a/Engine/Plugins/Runtime/Metasound/Source/MetasoundEditor/Private/SMetasoundGraphNode.cpp
+++ b/Engine/Plugins/Runtime/Metasound/Source/MetasoundEditor/Private/SMetasoundGraphNode.cpp
@@ -59,6 +59,7 @@
 #include "Widgets/Layout/SSpacer.h"
 #include "Widgets/SBoxPanel.h"
 #include "Widgets/SOverlay.h"
+#include "Widgets/SToolTip.h"
 #include "Widgets/SWidget.h"
 #include "Widgets/Text/SInlineEditableTextBlock.h"
 
@@ -211,7 +212,7 @@ namespace Metasound
 			TSharedPtr<SButton> SimulationButton;
 			TSharedRef<SWidget> SimulationWidget = SNew(SHorizontalBox)
 			+ SHorizontalBox::Slot()
-			.Padding(2.0f, 0.0f, 0.0f, 0.0f)
+			.Padding(7.0f, 0.0f, 0.0f, 0.0f)
 			.HAlign(HAlign_Left)
 			.VAlign(VAlign_Center)
 			[
@@ -450,18 +451,21 @@ namespace Metasound
 
 		TSharedRef<SWidget> SMetaSoundGraphNode::CreateTitleWidget(TSharedPtr<SNodeTitle> NodeTitle)
 		{
-			Frontend::FConstNodeHandle NodeHandle = GetMetaSoundNode().GetConstNodeHandle();
-			if (!NodeHandle->GetClassStyle().Display.bShowName)
+			TSharedRef<SWidget> DefaultTitle = SGraphNode::CreateTitleWidget(NodeTitle);
+
+			const UMetasoundEditorGraphNode& Node = GetMetaSoundNode();
+			const FMetasoundFrontendClass* FrontendClass = Node.GetFrontendClass();
+			if (!FrontendClass || !FrontendClass->Style.Display.bShowName)
 			{
 				return SNullWidget::NullWidget;
 			}
 
 			TSharedPtr<SHorizontalBox> TitleBoxWidget = SNew(SHorizontalBox);
 
-			FSlateIcon NodeIcon = GetMetaSoundNode().GetNodeTitleIcon();
+			const FSlateIcon NodeIcon = Node.GetNodeTitleIcon();
 			if (const FSlateBrush* IconBrush = NodeIcon.GetIcon())
 			{
-				if (IconBrush != FStyleDefaults::GetNoBrush())
+				if (IconBrush != FStyleDefaults::GetNoBrush() && Node.GetBreadcrumb().bIsClassNative)
 				{
 					TSharedPtr<SImage> Image;
 					TitleBoxWidget->AddSlot()
@@ -568,87 +572,253 @@ namespace Metasound
 
 		FLinearColor SMetaSoundGraphNode::GetNodeTitleColorOverride() const
 		{
-			FLinearColor ReturnTitleColor = GraphNode->IsDeprecated() ? FLinearColor::Red : GetNodeObj()->GetNodeTitleColor();
+			const UMetasoundEditorGraphNode& Node = GetMetaSoundNode();
 
-			if (!GraphNode->IsNodeEnabled() || GraphNode->IsDisplayAsDisabledForced() || GraphNode->IsNodeUnrelated())
+			FMetasoundFrontendClassStyleDisplay StyleDisplay;
+			if (const FMetasoundFrontendClass* FrontendClass = Node.GetFrontendClass())
 			{
-				ReturnTitleColor *= FLinearColor(0.5f, 0.5f, 0.5f, 0.4f);
+				StyleDisplay = FrontendClass->Style.Display;
 			}
-			else
+
+			FLinearColor TitleColor = GraphNode->IsDeprecated() ? FLinearColor::Red : GetNodeObj()->GetNodeTitleColor();
+			const bool bDimTitle = !GraphNode->IsNodeEnabled() || GraphNode->IsDisplayAsDisabledForced() || GraphNode->IsNodeUnrelated();
+
+			if (bDimTitle)
+			{
+				TitleColor *= FLinearColor(0.5f, 0.5f, 0.5f, 0.4f);
+				return TitleColor;
+			}
+
+			if (ClassType != EMetasoundFrontendClassType::External)
+			{
+				TitleColor.A = FadeCurve.GetLerp();
+				return TitleColor;
+			}
+
+			if (!Node.GetBreadcrumb().bIsClassNative)
+			{
+				TitleColor = FMath::Lerp(TitleColor, FLinearColor::White, 0.5f);
+				return TitleColor;
+			}
+
+			if (FSlateStyleRegistry::FindSlateStyle(StyleDisplay.StyleSet))
 			{
-				ReturnTitleColor.A = FadeCurve.GetLerp();
+				// Use white text if we're applying a style color to the node title's background.
+				return FLinearColor::White;
 			}
 
-			return ReturnTitleColor;
+			return TitleColor;
 		}
 
 		void SMetaSoundGraphNode::SetDefaultTitleAreaWidget(TSharedRef<SOverlay> DefaultTitleAreaWidget)
 		{
 			SGraphNode::SetDefaultTitleAreaWidget(DefaultTitleAreaWidget);
 
-			Metasound::Frontend::FConstNodeHandle NodeHandle = GetMetaSoundNode().GetConstNodeHandle();
-			if (NodeHandle->GetClassStyle().Display.bShowName)
+			const UMetasoundEditorGraphNode& Node = GetMetaSoundNode();
+
+			FMetasoundFrontendClassStyleDisplay StyleDisplay;
+			if (const FMetasoundFrontendClass* FrontendClass = Node.GetFrontendClass())
 			{
-				DefaultTitleAreaWidget->ClearChildren();
-				TSharedPtr<SNodeTitle> NodeTitle = SNew(SNodeTitle, GraphNode);
+				StyleDisplay = FrontendClass->Style.Display;
+			}
 
-				DefaultTitleAreaWidget->AddSlot()
-				.HAlign(HAlign_Fill)
+			if (!StyleDisplay.bShowName)
+			{
+				DefaultTitleAreaWidget->SetVisibility(EVisibility::Collapsed);
+				return;
+			}
+
+			DefaultTitleAreaWidget->ClearChildren();
+			const TSharedPtr<SNodeTitle> NodeTitle = SNew(SNodeTitle, GraphNode);
+
+			const TSharedRef<SWidget> TitleContent =
+				SNew(SHorizontalBox)
+				+ SHorizontalBox::Slot()
+				.AutoWidth()
+				.HAlign(HAlign_Left)
 				.VAlign(VAlign_Center)
+				.Padding(FMargin(0.f, 0.f, 6.f, 0.f))
 				[
-					SNew(SHorizontalBox)
-					+ SHorizontalBox::Slot()
-					.HAlign(HAlign_Fill)
+					CreateTitleIcon()
+				]
+				+ SHorizontalBox::Slot()
+				.FillWidth(1.f)
+				.HAlign(HAlign_Center)
+				.VAlign(VAlign_Center)
+				[
+					SNew(SVerticalBox)
+					+ SVerticalBox::Slot()
+					.AutoHeight()
+					.HAlign(HAlign_Center)
+					.VAlign(VAlign_Center)
 					[
-						SNew(SBorder)
-						.BorderImage(FAppStyle::GetBrush("NoBorder"))
-						[
-							SNew(SHorizontalBox)
-							+ SHorizontalBox::Slot()
-							[
-								SNew(SVerticalBox)
-								+ SVerticalBox::Slot()
-								.AutoHeight()
-								.HAlign(HAlign_Center)
-								[
-									CreateTitleWidget(NodeTitle)
-								]
-								+ SVerticalBox::Slot()
-								.AutoHeight()
-								[
-									NodeTitle.ToSharedRef()
-								]
-							]
-						]
+						CreateTitleWidget(NodeTitle)
 					]
-					+ SHorizontalBox::Slot()
-					.HAlign(HAlign_Right)
+					+ SVerticalBox::Slot()
+					.AutoHeight()
+					.HAlign(HAlign_Center)
 					.VAlign(VAlign_Center)
-					.Padding(0, 0, 5, 0)
-					.AutoWidth()
 					[
-						CreateTitleRightWidget()
+						NodeTitle.ToSharedRef()
 					]
+				]
+				+ SHorizontalBox::Slot()
+				.AutoWidth()
+				.HAlign(HAlign_Right)
+				.VAlign(VAlign_Center)
+				[
+					CreateTitleRightWidget()
+				];
+
+			const FLinearColor TitleBorderColor = Node.GetBreadcrumb().bIsClassNative ? GetTitleBorderColor() : GetNodeObj()->GetNodeTitleColor();
+
+			const TSharedRef<SWidget> MainTitleArea =
+				SNew(SBorder)
+				.BorderImage(FAppStyle::GetBrush("Graph.Node.ColorSpill"))
+				.Padding(FMargin(10.f, 5.f, 10.f, 3.f))
+				.BorderBackgroundColor(TitleBorderColor)
+				.HAlign(HAlign_Fill)
+				.VAlign(VAlign_Center)
+				[
+					TitleContent
 				];
 
-				DefaultTitleAreaWidget->AddSlot()
-				.VAlign(VAlign_Top)
+			DefaultTitleAreaWidget->AddSlot()
+			.HAlign(HAlign_Fill)
+			.VAlign(VAlign_Center)
+			[
+				MainTitleArea
+			];
+
+			DefaultTitleAreaWidget->AddSlot()
+			.VAlign(VAlign_Top)
+			[
+				SNew(SBorder)
+				.Visibility(EVisibility::HitTestInvisible)
+				.BorderImage( FAppStyle::GetBrush( "Graph.Node.TitleHighlight" ) )
+				.BorderBackgroundColor( this, &SGraphNode::GetNodeTitleIconColor )
+				[
+					SNew(SSpacer)
+					.Size(FVector2D(20,20))
+				]
+			];
+		}
+
+		TSharedRef<SWidget> SMetaSoundGraphNode::CreateTitleIcon() const
+		{
+			if (ClassType == EMetasoundFrontendClassType::Input || ClassType == EMetasoundFrontendClassType::Output)
+			{
+				return SNullWidget::NullWidget;
+			}
+
+			const UMetasoundEditorGraphNode& Node = GetMetaSoundNode();
+
+			const bool bIsAsset = !Node.GetBreadcrumb().bIsClassNative;
+
+			FText AuthorText = FText::GetEmpty();
+			bool bShowAuthor = false;
+
+			if (bIsAsset)
+			{
+				if (const FMetasoundFrontendClass* FrontendClass = Node.GetFrontendClass())
+				{
+					const FMetasoundFrontendClassMetadata& Metadata = FrontendClass->Metadata;
+					const FString MetadataAuthor = Metadata.GetAuthor();
+
+					if (!MetadataAuthor.IsEmpty())
+					{
+						AuthorText = FText::Format(LOCTEXT("MetaSoundNodeAuthorFormat", "Author: {0}"), FText::FromString(MetadataAuthor));
+						bShowAuthor = true;
+					}
+				}
+			}
+
+			const FSlateBrush* IconBrush = bIsAsset ? Node.GetNodeTitleIcon().GetIcon() : GetIconBrush();
+
+			if (!IconBrush)
+			{
+				return SNullWidget::NullWidget;
+			}
+
+			TSharedPtr<SToolTip> AuthorToolTip;
+			if (bShowAuthor)
+			{
+				AuthorToolTip =
+					SNew(SToolTip)
+					[
+						SNew(SVerticalBox)
+						+ SVerticalBox::Slot()
+						.AutoHeight()
+						.Padding(0, 0, 0, 2)
+						[
+							SNew(STextBlock)
+							.Text(AuthorText)
+						]
+					];
+			}
+
+			return SNew(SBox)
+				.WidthOverride(20.f)
+				.HeightOverride(20.f)
+				.ToolTip(AuthorToolTip)
 				[
 					SNew(SBorder)
-					.Visibility(EVisibility::HitTestInvisible)
-					.BorderImage( FAppStyle::GetBrush( "Graph.Node.TitleHighlight" ) )
-					.BorderBackgroundColor( this, &SGraphNode::GetNodeTitleIconColor )
+					.HAlign(HAlign_Fill)
+					.BorderImage_Lambda([]()
+					{
+						static const FSlateRoundedBoxBrush RoundedBoxBrush(FLinearColor::White, 3.f);
+						return &RoundedBoxBrush;
+					})
+					.BorderBackgroundColor(FLinearColor(0, 0, 0, 0.33f))
 					[
-						SNew(SSpacer)
-						.Size(FVector2D(20,20))
+						SNew(SBox)
+						.WidthOverride(16.f)
+						.HeightOverride(16.f)
+						.HAlign(HAlign_Center)
+						.VAlign(VAlign_Center)
+						[
+							SNew(SImage)
+							.Image(IconBrush)
+							.ColorAndOpacity(bIsAsset ? GetNodeObj()->GetNodeTitleColor() : FLinearColor::White)
+						]
 					]
 				];
+		}
 
+		const FSlateBrush* SMetaSoundGraphNode::GetIconBrush() const
+		{
+			const UMetasoundEditorGraphNode& Node = GetMetaSoundNode();
+
+			FMetasoundFrontendClassStyleDisplay StyleDisplay;
+			if (const FMetasoundFrontendClass* FrontendClass = Node.GetFrontendClass())
+			{
+				StyleDisplay = FrontendClass->Style.Display;
 			}
-			else
+
+			if (const ISlateStyle* Style = FSlateStyleRegistry::FindSlateStyle(StyleDisplay.StyleSet))
 			{
-				DefaultTitleAreaWidget->SetVisibility(EVisibility::Collapsed);
+				if (const FSlateBrush* Brush = Style->GetBrush(StyleDisplay.Icon))
+				{
+					return Brush;
+				}
 			}
+
+			return nullptr;
+		}
+
+		FLinearColor SMetaSoundGraphNode::GetTitleBorderColor() const
+		{
+			const UMetasoundEditorGraphNode& Node = GetMetaSoundNode();
+			if (const FMetasoundFrontendClass* FrontendClass = Node.GetFrontendClass())
+			{
+				const FMetasoundFrontendClassStyleDisplay& StyleDisplay = FrontendClass->Style.Display;
+				if (const ISlateStyle* Style = FSlateStyleRegistry::FindSlateStyle(StyleDisplay.StyleSet))
+				{
+					return Style->GetSlateColor(StyleDisplay.Color).GetSpecifiedColor();
+				}
+			}
+
+			return FLinearColor::Transparent;
 		}
 
 		void SMetaSoundGraphNode::MoveTo(const FVector2f& NewPosition, FNodeSet& NodeFilter, bool bMarkDirty)
diff --git a/Engine/Plugins/Runtime/Metasound/Source/MetasoundEditor/Private/SMetasoundGraphNode.h b/Engine/Plugins/Runtime/Metasound/Source/MetasoundEditor/Private/SMetasoundGraphNode.h
--- a/Engine/Plugins/Runtime/Metasound/Source/MetasoundEditor/Private/SMetasoundGraphNode.h
+++ b/Engine/Plugins/Runtime/Metasound/Source/MetasoundEditor/Private/SMetasoundGraphNode.h
@@ -75,6 +75,10 @@ namespace Metasound::Editor
 		void BeginOrUpdateValueTransaction(TWeakObjectPtr<UMetasoundEditorGraphMember> GraphMemberPtr, TFunctionRef<void(const FGuid& /*BuildPageID*/, UMetasoundEditorGraphMember& /*GraphMember*/)> SetValue);
 		void FinalizeValueTransaction(TWeakObjectPtr<UMetasoundEditorGraphMember> GraphMemberPtr, TFunctionRef<void(const FGuid& /*BuildPageID*/, UMetasoundEditorGraphMember& /*GraphMember*/, bool /*bPostTransaction*/)> SetValue);
 
+		TSharedRef<SWidget> CreateTitleIcon() const;
+		const FSlateBrush* GetIconBrush() const;
+		FLinearColor GetTitleBorderColor() const;
+
 		// Returns attribute that returns whether or not the input widget should be enabled.
 		TAttribute<bool> GetInputWidgetEnabled() const;
 
diff --git a/Engine/Plugins/Runtime/Metasound/Source/MetasoundFrontend/Private/MetasoundFrontendDocument.cpp b/Engine/Plugins/Runtime/Metasound/Source/MetasoundFrontend/Private/MetasoundFrontendDocument.cpp
--- a/Engine/Plugins/Runtime/Metasound/Source/MetasoundFrontend/Private/MetasoundFrontendDocument.cpp
+++ b/Engine/Plugins/Runtime/Metasound/Source/MetasoundFrontend/Private/MetasoundFrontendDocument.cpp
@@ -1259,6 +1259,9 @@ FMetasoundFrontendClassStyle FMetasoundFrontendClassStyle::GenerateClassStyle(co
 	Style.Display.bShowInputNames = InNodeDisplayStyle.bShowInputNames;
 	Style.Display.bShowOutputNames = InNodeDisplayStyle.bShowOutputNames;
 	Style.Display.ImageName = InNodeDisplayStyle.ImageName;
+	Style.Display.StyleSet = InNodeDisplayStyle.StyleSet;
+	Style.Display.Icon = InNodeDisplayStyle.Icon;
+	Style.Display.Color = InNodeDisplayStyle.Color;
 
 	return Style;
 }
diff --git a/Engine/Plugins/Runtime/Metasound/Source/MetasoundFrontend/Public/MetasoundFrontendDocument.h b/Engine/Plugins/Runtime/Metasound/Source/MetasoundFrontend/Public/MetasoundFrontendDocument.h
--- a/Engine/Plugins/Runtime/Metasound/Source/MetasoundFrontend/Public/MetasoundFrontendDocument.h
+++ b/Engine/Plugins/Runtime/Metasound/Source/MetasoundFrontend/Public/MetasoundFrontendDocument.h
@@ -971,6 +971,9 @@ struct FMetasoundFrontendClassStyleDisplay
 	,	bShowInputNames(InDisplayStyle.bShowInputNames)
 	,	bShowOutputNames(InDisplayStyle.bShowOutputNames)
 	,	bShowLiterals(InDisplayStyle.bShowLiterals)
+	,	StyleSet(InDisplayStyle.StyleSet)
+	,	Icon(InDisplayStyle.Icon)
+	,	Color(InDisplayStyle.Color)
 	{
 	}
 
@@ -988,6 +991,15 @@ struct FMetasoundFrontendClassStyleDisplay
 
 	UPROPERTY()
 	bool bShowLiterals = true;
+
+	UPROPERTY()
+	FName StyleSet;
+
+	UPROPERTY()
+	FName Icon;
+
+	UPROPERTY()
+	FName Color;
 #endif // WITH_EDITORONLY_DATA
 };
 
diff --git a/Engine/Plugins/Runtime/Metasound/Source/MetasoundGraphCore/Public/MetasoundNodeInterface.h b/Engine/Plugins/Runtime/Metasound/Source/MetasoundGraphCore/Public/MetasoundNodeInterface.h
--- a/Engine/Plugins/Runtime/Metasound/Source/MetasoundGraphCore/Public/MetasoundNodeInterface.h
+++ b/Engine/Plugins/Runtime/Metasound/Source/MetasoundGraphCore/Public/MetasoundNodeInterface.h
@@ -69,6 +69,15 @@ namespace Metasound
 
 		/** Whether or not to show input literals in visual layout. */
 		bool bShowLiterals = true;
+
+		/** Name of the Slate style associated with node (color, icon). */
+		FName StyleSet;
+
+		/** Name of the node's icon defined in Style. */
+		FName Icon;
+
+		/** Name of the color defined in Style. */
+		FName Color;
 	};
 
 	/** Name of a node class.
