[MetasoundEditor] Minor improvements to fallback styling of metasound node titles when omitting either an icon or color

#rnx
#Synced-CL 51544635

diff --git a/Engine/Plugins/Runtime/Metasound/Source/MetasoundEditor/Private/SMetasoundGraphNode.cpp b/Engine/Plugins/Runtime/Metasound/Source/MetasoundEditor/Private/SMetasoundGraphNode.cpp
--- a/Engine/Plugins/Runtime/Metasound/Source/MetasoundEditor/Private/SMetasoundGraphNode.cpp
+++ b/Engine/Plugins/Runtime/Metasound/Source/MetasoundEditor/Private/SMetasoundGraphNode.cpp
@@ -604,7 +604,10 @@ namespace Metasound
 			if (FSlateStyleRegistry::FindSlateStyle(StyleDisplay.StyleSet))
 			{
 				// Use white text if we're applying a style color to the node title's background.
-				return FLinearColor::White;
+				if (!StyleDisplay.Color.IsNone())
+				{
+					return FLinearColor::White;
+				}
 			}
 
 			return TitleColor;
@@ -795,15 +798,13 @@ namespace Metasound
 				StyleDisplay = FrontendClass->Style.Display;
 			}
 
-			if (const ISlateStyle* Style = FSlateStyleRegistry::FindSlateStyle(StyleDisplay.StyleSet))
+			const ISlateStyle* Style = FSlateStyleRegistry::FindSlateStyle(StyleDisplay.StyleSet);
+			if (!Style || StyleDisplay.Icon.IsNone())
 			{
-				if (const FSlateBrush* Brush = Style->GetBrush(StyleDisplay.Icon))
-				{
-					return Brush;
-				}
+				return nullptr;
 			}
 
-			return nullptr;
+			return Style->GetBrush(StyleDisplay.Icon);
 		}
 
 		FLinearColor SMetaSoundGraphNode::GetTitleBorderColor() const
@@ -814,7 +815,7 @@ namespace Metasound
 				const FMetasoundFrontendClassStyleDisplay& StyleDisplay = FrontendClass->Style.Display;
 				if (const ISlateStyle* Style = FSlateStyleRegistry::FindSlateStyle(StyleDisplay.StyleSet))
 				{
-					return Style->GetSlateColor(StyleDisplay.Color).GetSpecifiedColor();
+					return Style->GetSlateColor(StyleDisplay.Color, nullptr, FLinearColor::Transparent).GetSpecifiedColor();
 				}
 			}
 
