DocumentConfiguration Rename to MetaSound(Document)Template

#rnx

diff --git a/Engine/Plugins/Experimental/TechAudioTools/Source/TechAudioToolsMetaSound/Private/ViewModels/MetaSoundViewModel.cpp b/Engine/Plugins/Experimental/TechAudioTools/Source/TechAudioToolsMetaSound/Private/ViewModels/MetaSoundViewModel.cpp
--- a/Engine/Plugins/Experimental/TechAudioTools/Source/TechAudioToolsMetaSound/Private/ViewModels/MetaSoundViewModel.cpp
+++ b/Engine/Plugins/Experimental/TechAudioTools/Source/TechAudioToolsMetaSound/Private/ViewModels/MetaSoundViewModel.cpp
@@ -2,11 +2,10 @@
 
 #include "ViewModels/MetaSoundViewModel.h"
 
-#include "DocumentConfigurations/MetasoundFrontendPresetDocumentConfiguration.h"
+#include "DocumentTemplates/MetasoundFrontendPresetTemplate.h"
 #include "Logging/StructuredLog.h"
 #include "MetasoundDocumentBuilderRegistry.h"
 #include "MetasoundFrontendPages.h"
-#include "MetasoundOutputSubsystem.h"
 #include "TechAudioToolsLog.h"
 #include "ViewModels/MetaSoundViewModelConversionFunctions.h"
 
@@ -63,7 +62,7 @@ void UMetaSoundViewModel::Initialize(UMetaSoundBuilderBase* InBuilder)
 	Builder = InBuilder;
 
 	const FMetaSoundFrontendDocumentBuilder& DocumentBuilder = InBuilder->GetConstBuilder();
-	UE_MVVM_SET_PROPERTY_VALUE(bIsPreset, DocumentBuilder.IsDocumentConfiguredBy<FPresetDocumentConfiguration>());
+	UE_MVVM_SET_PROPERTY_VALUE(bIsPreset, DocumentBuilder.IsDocumentConfiguredBy<FPresetTemplate>());
 
 	CreateMemberViewModels();
 	SetIsInitialized(true);
diff --git a/Engine/Plugins/Experimental/TechAudioTools/Source/TechAudioToolsMetaSound/Public/ViewModels/MetaSoundViewModel.h b/Engine/Plugins/Experimental/TechAudioTools/Source/TechAudioToolsMetaSound/Public/ViewModels/MetaSoundViewModel.h
--- a/Engine/Plugins/Experimental/TechAudioTools/Source/TechAudioToolsMetaSound/Public/ViewModels/MetaSoundViewModel.h
+++ b/Engine/Plugins/Experimental/TechAudioTools/Source/TechAudioToolsMetaSound/Public/ViewModels/MetaSoundViewModel.h
@@ -4,6 +4,7 @@
 
 #include "MetasoundFrontendDocumentBuilder.h"
 #include "MVVMViewModelBase.h"
+#include "Templates/SubclassOf.h"
 
 #include "MetaSoundViewModel.generated.h"
 
diff --git a/Engine/Plugins/Runtime/Metasound/Config/DefaultMetaSound.ini b/Engine/Plugins/Runtime/Metasound/Config/DefaultMetaSound.ini
--- a/Engine/Plugins/Runtime/Metasound/Config/DefaultMetaSound.ini
+++ b/Engine/Plugins/Runtime/Metasound/Config/DefaultMetaSound.ini
@@ -26,3 +26,17 @@
 
 ; 5.3
 +EnumRedirects=(OldName="MetasoundSourceAudioFormat",NewName="/Script/MetasoundEngine.MetaSoundOutputAudioFormat")
+
+; 5.8
+;Template redirects (renamed from "Configuration")
++StructRedirects=(OldName="/Script/MetasoundFrontend.MetasoundFrontendDocumentConfiguration",NewName="/Script/MetasoundFrontend.MetaSoundFrontendDocumentTemplate")
++StructRedirects=(OldName="/Script/MetasoundFrontend.MetaSoundFrontendDocumentConfigurationVertexMetadata",NewName="/Script/MetasoundFrontend.MetaSoundFrontendDocumentTemplateVertexMetadata")
++StructRedirects=(OldName="/Script/MetasoundFrontend.MetaSoundFrontendDocumentVertexConfiguration",NewName="/Script/MetasoundFrontend.MetaSoundFrontendDocumentVertexTemplate")
++StructRedirects=(OldName="/Script/MetasoundFrontend.MetaSoundFrontendPresetConfigurationVertexMetadata",NewName="/Script/MetasoundFrontend.MetaSoundFrontendPresetVertexMetadata")
++StructRedirects=(OldName="/Script/MetasoundFrontend.MetaSoundFrontendPresetDocumentConfiguration",NewName="/Script/MetasoundFrontend.MetaSoundFrontendPresetTemplate")
++StructRedirects=(OldName="/Script/MetasoundEngine.MetaSoundRandomizerDocumentConfiguration",NewName="/Script/MetasoundEngine.MetaSoundRandomizerTemplate")
++StructRedirects=(OldName="/Script/MetasoundEditor.MetaSoundDocumentConfigurationWidgetTabInfo",NewName="/Script/MetasoundEditor.MetaSoundTemplateWidgetTabInfo")
++ClassRedirects=(OldName="/Script/MetasoundEditor.MetaSoundDocumentConfigurationViewModelBase",NewName="/Script/MetasoundEditor.MetaSoundTemplateViewModelBase")
++ClassRedirects=(OldName="/Script/MetasoundEditor.MetaSoundDocumentConfigurationWidgetInterface",NewName="/Script/MetasoundEditor.MetaSoundTemplateWidgetInterface")
++FunctionRedirects=(OldName="/Script/MetasoundEditor.MetaSoundTemplateWidgetInterface.IsSupportedConfiguration",NewName="/Script/MetasoundEditor.MetaSoundTemplateWidgetInterface.IsSupportedTemplate")
++PropertyRedirects=(OldName="/Script/MetasoundFrontend.MetasoundFrontendDocument.Configuration",NewName="/Script/MetasoundFrontend.MetasoundFrontendDocument.Template")
diff --git a/Engine/Plugins/Runtime/Metasound/Source/MetasoundEditor/Private/DocumentConfigurationViewModels/MetasoundDocumentConfigurationViewModelBase.cpp b/Engine/Plugins/Runtime/Metasound/Source/MetasoundEditor/Private/DocumentConfigurationViewModels/MetasoundDocumentConfigurationViewModelBase.cpp
deleted file mode 100644
diff --git a/Engine/Plugins/Runtime/Metasound/Source/MetasoundEditor/Private/DocumentConfigurationViewModels/RandomizerConfigurationViewModel.cpp b/Engine/Plugins/Runtime/Metasound/Source/MetasoundEditor/Private/DocumentConfigurationViewModels/RandomizerConfigurationViewModel.cpp
deleted file mode 100644
diff --git a/Engine/Plugins/Runtime/Metasound/Source/MetasoundEditor/Private/FindInMetasound.cpp b/Engine/Plugins/Runtime/Metasound/Source/MetasoundEditor/Private/FindInMetasound.cpp
--- a/Engine/Plugins/Runtime/Metasound/Source/MetasoundEditor/Private/FindInMetasound.cpp
+++ b/Engine/Plugins/Runtime/Metasound/Source/MetasoundEditor/Private/FindInMetasound.cpp
@@ -2,12 +2,11 @@
 
 #include "FindInMetasound.h"
 
-#include "DocumentConfigurations/MetasoundFrontendPresetDocumentConfiguration.h"
+#include "DocumentTemplates/MetasoundFrontendPresetTemplate.h"
 #include "Framework/Application/SlateApplication.h"
 #include "Layout/WidgetPath.h"
 #include "MetasoundEditorGraph.h"
 #include "MetasoundEditorGraphBuilder.h"
-#include "MetasoundEditorGraphInputNode.h"
 #include "MetasoundEditorModule.h"
 #include "MetasoundFrontendDocumentBuilder.h"
 #include "MetasoundSettings.h"
@@ -98,13 +97,21 @@ namespace Metasound::Editor
 
 					if (bUseFullValueString)
 					{
+						bool bStringSet = false;
 						if (EnumInterface)
 						{
 							int32 Value;
 							Literal.TryGet(Value);
-							LiteralString += EnumInterface->ToName(Value)->ToString();
+
+							TOptional<FName> Name = EnumInterface->ToName(Value);
+							if (Name.IsSet())
+							{
+								bStringSet = true;
+								LiteralString += Name->ToString();
+							}
 						}
-						else
+
+						if (!bStringSet)
 						{
 							LiteralString += Literal.ToString();
 						}
@@ -133,15 +140,26 @@ namespace Metasound::Editor
 								LiteralString += LOCTEXT("FindMetasound_NoneObject", "(None Object)").ToString();
 							}
 						}
-						else if (EnumInterface)
-						{
-							int32 Value;
-							Literal.TryGet(Value);
-							LiteralString += EnumInterface->ToName(Value)->ToString();
-						}
 						else
 						{
-							LiteralString += Literal.ToString();
+							bool bStringSet = false;
+							if (EnumInterface)
+							{
+								int32 Value;
+								Literal.TryGet(Value);
+
+								TOptional<FName> Name = EnumInterface->ToName(Value);
+								if (Name.IsSet())
+								{
+									bStringSet = true;
+									LiteralString += Name->ToString();
+								}
+							}
+
+							if (!bStringSet)
+							{
+								LiteralString += Literal.ToString();
+							}
 						}
 					}
 
diff --git a/Engine/Plugins/Runtime/Metasound/Source/MetasoundEditor/Private/MetasoundAssetDefinitions.cpp b/Engine/Plugins/Runtime/Metasound/Source/MetasoundEditor/Private/MetasoundAssetDefinitions.cpp
--- a/Engine/Plugins/Runtime/Metasound/Source/MetasoundEditor/Private/MetasoundAssetDefinitions.cpp
+++ b/Engine/Plugins/Runtime/Metasound/Source/MetasoundEditor/Private/MetasoundAssetDefinitions.cpp
@@ -8,7 +8,6 @@
 #include "Components/AudioComponent.h"
 #include "ContentBrowserMenuContexts.h"
 #include "ContentBrowserModule.h"
-#include "DocumentConfigurations/MetasoundFrontendPresetDocumentConfiguration.h"
 #include "Editor/EditorPerProjectUserSettings.h"
 #include "FileHelpers.h"
 #include "IAssetTools.h"
@@ -38,6 +37,7 @@
 #include "Styling/ISlateStyle.h"
 #include "Styling/SlateStyleRegistry.h"
 #include "Styling/SlateTypes.h"
+#include "DocumentTemplates/MetasoundFrontendPresetTemplate.h"
 #include "ToolMenus.h"
 #include "ToolMenuSection.h"
 #include "Widgets/Input/SButton.h"
@@ -86,9 +86,9 @@ namespace Metasound::Editor
 				if (UAssetEditorSubsystem* AssetSubsystem = GEditor->GetEditorSubsystem<UAssetEditorSubsystem>())
 				{
 					FMetaSoundFrontendDocumentBuilder& DocumentBuilder = Frontend::IDocumentBuilderRegistry::GetChecked().FindOrBeginBuilding(MetaSound);
-					if (const FPresetDocumentConfiguration* PresetConfig = DocumentBuilder.GetConstDocumentConfigurationAs<FPresetDocumentConfiguration>())
+					if (const FPresetTemplate* PresetTemplate = DocumentBuilder.GetConstTemplateAs<FPresetTemplate>())
 					{
-						if (UObject* ParentObject = PresetConfig->Parent.GetObject(); GEditor && ParentObject)
+						if (UObject* ParentObject = PresetTemplate->Parent.GetObject(); GEditor && ParentObject)
 						{
 							GEditor->SyncBrowserToObjects(TArray<UObject*>{ ParentObject });
 						}
@@ -106,9 +106,9 @@ namespace Metasound::Editor
 				if (UAssetEditorSubsystem* AssetSubsystem = GEditor->GetEditorSubsystem<UAssetEditorSubsystem>())
 				{
 					FMetaSoundFrontendDocumentBuilder& DocumentBuilder = Frontend::IDocumentBuilderRegistry::GetChecked().FindOrBeginBuilding(MetaSound);
-					if (const FPresetDocumentConfiguration* PresetConfig = DocumentBuilder.GetConstDocumentConfigurationAs<FPresetDocumentConfiguration>())
+					if (const FPresetTemplate* PresetTemplate = DocumentBuilder.GetConstTemplateAs<FPresetTemplate>())
 					{
-						if (UObject* ParentObject = PresetConfig->Parent.GetObject())
+						if (UObject* ParentObject = PresetTemplate->Parent.GetObject())
 						{
 							AssetSubsystem->OpenEditorForAsset(ParentObject);
 						}
@@ -118,18 +118,18 @@ namespace Metasound::Editor
 		}
 
 		template <typename TClass, typename TFactoryClass>
-		void ExecuteCreateMetaSoundWithDocConfig(const FToolMenuContext& MenuContext, TInstancedStruct<FMetaSoundFrontendDocumentConfiguration> DocConfig)
+		void ExecuteCreateMetaSoundWithTemplate(const FToolMenuContext& MenuContext, TInstancedStruct<FMetaSoundFrontendDocumentTemplate> DocTemplate)
 		{
 			if (const UContentBrowserAssetContextMenuContext* Context = UContentBrowserAssetContextMenuContext::FindContextWithAssets(MenuContext))
 			{
-				const UScriptStruct* DocConfigStruct = DocConfig.GetScriptStruct();
-				check(DocConfigStruct);
+				const UScriptStruct* DocTemplateStruct = DocTemplate.GetScriptStruct();
+				check(DocTemplateStruct);
 
-				const FText DisplayName = DocConfigStruct->GetDisplayNameText();
+				const FText DisplayName = DocTemplateStruct->GetDisplayNameText();
 
 				UMetaSoundBaseFactory* Factory = NewObject<TFactoryClass>();
 				check(Factory);
-				Factory->DocumentConfiguration = MoveTemp(DocConfig);
+				Factory->Template = MoveTemp(DocTemplate);
 				Factory->SelectedObjects = Context->LoadSelectedObjects<UObject>();
 
 				if (Factory->SelectedObjects.IsEmpty())
@@ -143,7 +143,8 @@ namespace Metasound::Editor
 
 				FString PackagePath;
 				FString AssetName;
-				IAssetTools::Get().CreateUniqueAssetName(Outermost->GetName(), FString::Format(TEXT("_{0}"), { DisplayName.ToString() }), PackagePath, AssetName);
+				const FString Suffix = TEXT("_") + ObjectTools::SanitizeObjectName(DisplayName.ToString());
+				IAssetTools::Get().CreateUniqueAssetName(Outermost->GetName(), Suffix, PackagePath, AssetName);
 
 				if (UObject* NewMetaSound = IAssetTools::Get().CreateAssetWithDialog(AssetName, FPackageName::GetLongPackagePath(PackagePath), Factory->GetSupportedClass(), Factory))
 				{
@@ -151,7 +152,7 @@ namespace Metasound::Editor
 					ContentBrowserModule.Get().SyncBrowserToAssets(TArray<UObject*>({ NewMetaSound }));
 
 					// Save happened in Create call above, but need to save again after
-					// doc configuration executes to capture initial procedural changes.
+					// Template executes to capture initial procedural changes.
 					if (ISourceControlModule::Get().IsEnabled())
 					{
 						constexpr bool bCheckDirty = false;
@@ -161,7 +162,7 @@ namespace Metasound::Editor
 				}
 				else
 				{
-					UE_LOG(LogMetaSound, Display, TEXT("Error creating new asset when creating '%s' or creation of MetaSound with configuration was canceled by user."), *AssetName);
+					UE_LOG(LogMetaSound, Display, TEXT("Error creating new asset when creating '%s' or creation of MetaSound with template was canceled by user."), *AssetName);
 				}
 			}
 		}
@@ -299,8 +300,8 @@ namespace Metasound::Editor
 					check(ParentPackage);
 					IAssetTools::Get().CreateUniqueAssetName(ParentPackage->GetName(), TEXT("_Preset"), PackagePath, AssetName);
 
-					TInstancedStruct<FPresetDocumentConfiguration> PresetConfig = TInstancedStruct<FPresetDocumentConfiguration>::Make();
-					PresetConfig.GetMutable().Parent = ParentMetaSound;
+					TInstancedStruct<FPresetTemplate> PresetTemplate = TInstancedStruct<FPresetTemplate>::Make();
+					PresetTemplate.GetMutable().Parent = ParentMetaSound;
 
 					IMetasoundEditorModule& MetaSoundEditorModule = FModuleManager::GetModuleChecked<IMetasoundEditorModule>("MetaSoundEditor");
 					if (ParentPackage->bIsCookedForEditor)
@@ -324,7 +325,7 @@ namespace Metasound::Editor
 						UMetaSoundEditorSubsystem& EditorSubsystem = UMetaSoundEditorSubsystem::GetChecked();
 						EditorSubsystem.InitAsset(*NewMetaSound, UMetaSoundEditorSubsystem::FInitAssetArgs
 						{
-							.DocumentConfiguration = MoveTemp(PresetConfig),
+							.Template = MoveTemp(PresetTemplate),
 							.SelectedObjects = { ParentMetaSound }
 						});
 
@@ -389,7 +390,7 @@ namespace Metasound::Editor
 			return false;
 		}
 
-		void AddPresetActions(const FSlateIcon& AssetIcon, const UClass& MetaSoundClass, const UClass& Class, FToolUIAction CreateFromConfigAction, const FText& ConfigDisplayName)
+		void AddPresetActions(const FSlateIcon& AssetIcon, const UClass& MetaSoundClass, const UClass& Class, FToolUIAction CreateFromTemplateAction, const FText& TemplateDisplayName)
 		{
 			const FText ClassName = Class.GetDisplayNameText();
 
@@ -399,15 +400,15 @@ namespace Metasound::Editor
 			check(Section);
 
 			const FText MetaSoundClassDisplayName = MetaSoundClass.GetDisplayNameText();
-			Section->AddDynamicEntry(NAME_None, FNewToolMenuSectionDelegate::CreateLambda([CreateAction = MoveTemp(CreateFromConfigAction), MetaSoundClassDisplayName, MetaSoundClassName, ConfigName = ConfigDisplayName, ClassName, AssetIcon](FToolMenuSection& InSection)
+			Section->AddDynamicEntry(NAME_None, FNewToolMenuSectionDelegate::CreateLambda([CreateAction = MoveTemp(CreateFromTemplateAction), MetaSoundClassDisplayName, MetaSoundClassName, TemplateName = TemplateDisplayName, ClassName, AssetIcon](FToolMenuSection& InSection)
 			{
 				using namespace Metasound::Editor;
 
 				{
-					const TAttribute<FText> Label = FText::Format(LOCTEXT("MetaSound_CreateFromPreset", "Create {0}..."), ConfigName);
+					const TAttribute<FText> Label = FText::Format(LOCTEXT("MetaSound_CreateFromPreset", "Create {0}..."), TemplateName);
 					const TAttribute<FText> ToolTip = FText::Format(LOCTEXT("MetaSound_CreateFromPresetToolTipFormat", "Creates a {0} Preset using the selected MetaSound as a parent."), MetaSoundClassDisplayName);
 
-					const FName MenuEntryName = *FString::Format(TEXT("{0}_CreateFromConfig_{1}"), { MetaSoundClassName.ToString(), ConfigName.ToString()});
+					const FName MenuEntryName = *FString::Format(TEXT("{0}_CreateFromTemplate_{1}"), { MetaSoundClassName.ToString(), TemplateName.ToString()});
 					InSection.AddMenuEntry(MenuEntryName, Label, ToolTip, AssetIcon, CreateAction);
 				}
 				{
@@ -505,15 +506,15 @@ namespace Metasound::Editor
 			return OutSections.Add_GetRef(&MetaSoundSection);
 		}
 
-		void AddConfigActions(
+		void AddTemplateActions(
 			const FSlateIcon& AssetIcon,
 			const UClass& Class,
 			const UClass& MetaSoundClass,
-			FToolUIAction CreateFromConfigAction,
-			TConstStructView<FMetaSoundFrontendDocumentConfiguration> Config)
+			FToolUIAction CreateFromTemplateAction,
+			TConstStructView<FMetaSoundFrontendDocumentTemplate> Template)
 		{
 			const FText ClassName = Class.GetDisplayNameText();
-			const UScriptStruct* ScriptStruct = Config.GetScriptStruct();
+			const UScriptStruct* ScriptStruct = Template.GetScriptStruct();
 			check(ScriptStruct);
 			const UAssetDefinition* AssetDefinition = UAssetDefinitionRegistry::Get()->GetAssetDefinitionForClass(&Class);
 
@@ -532,18 +533,18 @@ namespace Metasound::Editor
 			if (Section)
 			{
 				const FText MetaSoundClassDisplayName = MetaSoundClass.GetDisplayNameText();
-				const FText ConfigToolTip = Config.GetPtr()->GetEditorOptions().ToolTip;
-				auto AddEntryLambda = [CreateAction = MoveTemp(CreateFromConfigAction), ScriptStruct, ConfigToolTip, ClassName, MetaSoundClassDisplayName, AssetIcon](FToolMenuSection& InSection)
+				const FText TemplateToolTip = Template.GetPtr()->GetEditorOptions().ToolTip;
+				auto AddEntryLambda = [CreateAction = MoveTemp(CreateFromTemplateAction), ScriptStruct, TemplateToolTip, ClassName, MetaSoundClassDisplayName, AssetIcon](FToolMenuSection& InSection)
 				{
 					using namespace Metasound::Editor;
 			
-					const FText ConfigName = ScriptStruct->GetDisplayNameText();
-					const TAttribute<FText> Label = FText::Format(LOCTEXT("MetaSound_CreateFromConfig", "Create {0}..."), ConfigName);
+					const FText TemplateName = ScriptStruct->GetDisplayNameText();
+					const TAttribute<FText> Label = FText::Format(LOCTEXT("MetaSound_CreateFromTemplate", "Create {0}..."), TemplateName);
 					const TAttribute<FText> ToolTip = FText::Format(
-						LOCTEXT("MetaSound_CreateFromConfigToolTipFormat", "Creates a {0} with {1} configuration, initialized using the selected assets. {2}"),
-						MetaSoundClassDisplayName, ConfigName, ConfigToolTip);
+						LOCTEXT("MetaSound_CreateFromTemplateToolTipFormat", "Creates a {0} with {1} Template, initialized using the selected assets. {2}"),
+						MetaSoundClassDisplayName, TemplateName, TemplateToolTip);
 			
-					const FName MenuEntryName = *FString::Format(TEXT("MetaSound_CreateFromConfig_{0}"), { ConfigName.ToString() });
+					const FName MenuEntryName = *FString::Format(TEXT("MetaSound_CreateFromTemplate_{0}"), { TemplateName.ToString() });
 					InSection.AddMenuEntry(MenuEntryName, Label, ToolTip, AssetIcon, CreateAction);
 				};
 
@@ -921,7 +922,7 @@ TArray<FToolMenuSection*> UAssetDefinition_MetaSoundSource::RebuildSoundContextM
 namespace MenuExtension_MetaSoundTemplate
 {
 	template <typename TClass, typename TFactoryClass>
-	void RegisterDocumentConfigurationActions(const FSlateIcon& AssetIcon)
+	void RegisterTemplateActions(const FSlateIcon& AssetIcon)
 	{
 		using namespace Metasound;
 		using namespace Metasound::Editor;
@@ -940,14 +941,14 @@ namespace MenuExtension_MetaSoundTemplate
 				continue;
 			}
 
-			if (!Struct->IsChildOf(FMetaSoundFrontendDocumentConfiguration::StaticStruct()))
+			if (!Struct->IsChildOf(FMetaSoundFrontendDocumentTemplate::StaticStruct()))
 			{
 				continue;
 			}
 
 			// Presets are special snowflakes for now, as they require a bit more functionality for back
-			// compat like copying over SoundWave fields which still needs to be migrated to doc config.
-			if (Struct == FPresetDocumentConfiguration::StaticStruct())
+			// compat like copying over SoundWave fields which still needs to be migrated to templates.
+			if (Struct == FPresetTemplate::StaticStruct())
 			{
 				continue;
 			}
@@ -958,9 +959,9 @@ namespace MenuExtension_MetaSoundTemplate
 				continue;
 			}
 
-			TInstancedStruct<FMetaSoundFrontendDocumentConfiguration> ConfigInstance;
-			ConfigInstance.InitializeAsScriptStruct(Struct);
-			TArray<TSoftClassPtr<UObject>> SupportedMetaSoundClasses = ConfigInstance.Get().GetEditorOptions().MetaSoundClasses;
+			TInstancedStruct<FMetaSoundFrontendDocumentTemplate> TemplateInstance;
+			TemplateInstance.InitializeAsScriptStruct(Struct);
+			TArray<TSoftClassPtr<UObject>> SupportedMetaSoundClasses = TemplateInstance.Get().GetEditorOptions().MetaSoundClasses;
 			if (SupportedMetaSoundClasses.IsEmpty())
 			{
 				IMetasoundUObjectRegistry::Get().IterateRegisteredUClasses([&SupportedMetaSoundClasses](UClass& RegClass)
@@ -971,18 +972,18 @@ namespace MenuExtension_MetaSoundTemplate
 
 			if (SupportedMetaSoundClasses.Contains(MetaSoundSoftClass))
 			{
-				FToolUIAction DocConfigCreateAction;
-				DocConfigCreateAction.ExecuteAction = FToolMenuExecuteAction::CreateLambda([Struct](const FToolMenuContext& MenuContext)
+				FToolUIAction DocTemplateCreateAction;
+				DocTemplateCreateAction.ExecuteAction = FToolMenuExecuteAction::CreateLambda([Struct](const FToolMenuContext& MenuContext)
 				{
-					TInstancedStruct<FMetaSoundFrontendDocumentConfiguration> NewInstance;
+					TInstancedStruct<FMetaSoundFrontendDocumentTemplate> NewInstance;
 					NewInstance.InitializeAsScriptStruct(Struct);
-					return AssetDefinitionsPrivate::ExecuteCreateMetaSoundWithDocConfig<TClass, TFactoryClass>(MenuContext, MoveTemp(NewInstance));
+					return AssetDefinitionsPrivate::ExecuteCreateMetaSoundWithTemplate<TClass, TFactoryClass>(MenuContext, MoveTemp(NewInstance));
 				});
 
-				const TArray<TSoftClassPtr<UObject>>& AssetActionClasses = ConfigInstance.GetPtr()->GetEditorOptions().AssetActionClasses;
+				const TArray<TSoftClassPtr<UObject>>& AssetActionClasses = TemplateInstance.GetPtr()->GetEditorOptions().AssetActionClasses;
 				if (AssetActionClasses.IsEmpty())
 				{
-					AssetDefinitionsPrivate::AddConfigActions(AssetIcon, *MetaSoundClass, *MetaSoundClass, MoveTemp(DocConfigCreateAction), ConfigInstance);
+					AssetDefinitionsPrivate::AddTemplateActions(AssetIcon, *MetaSoundClass, *MetaSoundClass, MoveTemp(DocTemplateCreateAction), TemplateInstance);
 				}
 				else
 				{
@@ -995,7 +996,7 @@ namespace MenuExtension_MetaSoundTemplate
 						}
 						check(ClassPtr);
 
-						AssetDefinitionsPrivate::AddConfigActions(AssetIcon, *ClassPtr, *MetaSoundClass, DocConfigCreateAction, ConfigInstance);
+						AssetDefinitionsPrivate::AddTemplateActions(AssetIcon, *ClassPtr, *MetaSoundClass, DocTemplateCreateAction, TemplateInstance);
 					}
 				}
 			}
@@ -1020,14 +1021,14 @@ namespace MenuExtension_MetaSoundTemplate
 			{
 				const FSlateIcon AssetIcon = Metasound::Editor::Style::CreateSlateIcon("ClassIcon.MetasoundSource");
 
-				RegisterDocumentConfigurationActions<UMetaSoundSource, UMetaSoundSourceFactory>(AssetIcon);
+				RegisterTemplateActions<UMetaSoundSource, UMetaSoundSourceFactory>(AssetIcon);
 
 				const UClass* MetaSoundClass = UMetaSoundSource::StaticClass();
 				check(MetaSoundClass);
 
 				FToolUIAction PresetCreateAction;
 				PresetCreateAction.ExecuteAction = FToolMenuExecuteAction::CreateStatic(&AssetDefinitionsPrivate::ExecuteCreateMetaSoundPreset<UMetaSoundSource, UMetaSoundSourceFactory>);
-				AssetDefinitionsPrivate::AddPresetActions(AssetIcon, *MetaSoundClass, *MetaSoundClass, MoveTemp(PresetCreateAction), FPresetDocumentConfiguration::StaticStruct()->GetDisplayNameText());
+				AssetDefinitionsPrivate::AddPresetActions(AssetIcon, *MetaSoundClass, *MetaSoundClass, MoveTemp(PresetCreateAction), FPresetTemplate::StaticStruct()->GetDisplayNameText());
 
 				{
 					const UAssetDefinition* AssetDefinition = UAssetDefinitionRegistry::Get()->GetAssetDefinitionForClass(MetaSoundClass);
@@ -1138,11 +1139,11 @@ namespace MenuExtension_MetaSoundTemplate
 				CreateAction.ExecuteAction = FToolMenuExecuteAction::CreateStatic(&AssetDefinitionsPrivate::ExecuteCreateMetaSoundPreset<UMetaSoundPatch, UMetaSoundFactory>);
 				const FSlateIcon AssetIcon = Metasound::Editor::Style::CreateSlateIcon("ClassIcon.MetasoundPatch");
 
-				RegisterDocumentConfigurationActions<UMetaSoundPatch, UMetaSoundFactory>(AssetIcon);
+				RegisterTemplateActions<UMetaSoundPatch, UMetaSoundFactory>(AssetIcon);
 
 				const UClass* MetaSoundClass = UMetaSoundPatch::StaticClass();
 				check(MetaSoundClass);
-				AssetDefinitionsPrivate::AddPresetActions(AssetIcon, *MetaSoundClass, *MetaSoundClass, MoveTemp(CreateAction), FPresetDocumentConfiguration::StaticStruct()->GetDisplayNameText());
+				AssetDefinitionsPrivate::AddPresetActions(AssetIcon, *MetaSoundClass, *MetaSoundClass, MoveTemp(CreateAction), FPresetTemplate::StaticStruct()->GetDisplayNameText());
 			}
 
 			// Asset actions
@@ -1171,5 +1172,5 @@ namespace MenuExtension_MetaSoundTemplate
 			}
  		}));
 	});
-} // namespace MenuExtension_MetaSoundSourceTemplate
+} // namespace MenuExtension_MetaSoundTemplate
 #undef LOCTEXT_NAMESPACE //MetaSoundEditor
diff --git a/Engine/Plugins/Runtime/Metasound/Source/MetasoundEditor/Private/MetasoundDetailCustomization.cpp b/Engine/Plugins/Runtime/Metasound/Source/MetasoundEditor/Private/MetasoundDetailCustomization.cpp
--- a/Engine/Plugins/Runtime/Metasound/Source/MetasoundEditor/Private/MetasoundDetailCustomization.cpp
+++ b/Engine/Plugins/Runtime/Metasound/Source/MetasoundEditor/Private/MetasoundDetailCustomization.cpp
@@ -114,9 +114,9 @@ namespace Metasound::Editor
 	{
 	}
 
-	FName FMetasoundDetailCustomization::GetDocConfigPropertyPath() const
+	FName FMetasoundDetailCustomization::GetTemplatePropertyPath() const
 	{
-		return BuildChildPath(DocumentPropertyName, GET_MEMBER_NAME_CHECKED(FMetasoundFrontendDocument, Configuration));
+		return BuildChildPath(DocumentPropertyName, GET_MEMBER_NAME_CHECKED(FMetasoundFrontendDocument, Template));
 	}
 
 	FName FMetasoundDetailCustomization::GetInterfaceVersionsPropertyPath() const
@@ -196,8 +196,8 @@ namespace Metasound::Editor
 				TSharedPtr<IPropertyHandle> MinorVersionHandle = DetailLayout.GetProperty(MinorVersionPropertyPath);
 
 				{
-					TSharedPtr<IPropertyHandle> DocConfigHandle = DetailLayout.GetProperty(GetDocConfigPropertyPath());
-					if (DocConfigHandle.IsValid())
+					TSharedPtr<IPropertyHandle> TemplateHandle = DetailLayout.GetProperty(GetTemplatePropertyPath());
+					if (TemplateHandle.IsValid())
 					{
 						FSimpleDelegate OnConfigChanged = FSimpleDelegate::CreateSPLambda(AsShared(), [this]()
 						{
@@ -207,11 +207,11 @@ namespace Metasound::Editor
 								if (DocBuilder.IsValid())
 								{
 									DocBuilder.ConfigureDocument();
-									DocBuilder.GetDocumentDelegates().OnDocumentConfigurationChanged.Broadcast({ EDocumentConfigChangeType::Struct });
+									DocBuilder.GetDocumentDelegates().OnDocumentTemplateChanged.Broadcast({ EDocumentTemplateChangeType::Struct });
 								}
 							}
 						});
-						DocConfigHandle->SetOnPropertyValueChanged(OnConfigChanged);
+						TemplateHandle->SetOnPropertyValueChanged(OnConfigChanged);
 
 						const auto OnConfigPropChanged = TDelegate<void(const FPropertyChangedEvent&)>::CreateSPLambda(AsShared(), [this](const FPropertyChangedEvent& InEvent)
 						{
@@ -220,27 +220,27 @@ namespace Metasound::Editor
 								FMetaSoundFrontendDocumentBuilder& DocBuilder = Builder->GetBuilder();
 								if (DocBuilder.IsValid())
 								{
-									if (FDocumentConfiguration* Config = DocBuilder.GetDocumentConfiguration().GetPtr())
+									if (FMetaSoundTemplate* Template = DocBuilder.GetDocumentTemplate().GetPtr())
 									{
-										Config->OnPropertyChanged(InEvent, DocBuilder);
+										Template->OnPropertyChanged(InEvent, DocBuilder);
 										DocBuilder.ConfigureDocument();
-										DocBuilder.GetDocumentDelegates().OnDocumentConfigurationChanged.Broadcast({ EDocumentConfigChangeType::Property });
+										DocBuilder.GetDocumentDelegates().OnDocumentTemplateChanged.Broadcast({ EDocumentTemplateChangeType::Property });
 									}
 								}
 							}
 						});
-						DocConfigHandle->SetOnChildPropertyValueChangedWithData(OnConfigPropChanged);
-						GeneralCategoryBuilder.AddProperty(DocConfigHandle)
+						TemplateHandle->SetOnChildPropertyValueChangedWithData(OnConfigPropChanged);
+						GeneralCategoryBuilder.AddProperty(TemplateHandle)
 							.Visibility(TAttribute<EVisibility>::CreateSPLambda(this, [this]()
 							{
-								if (Builder.IsValid() && !FEditor::IgnoreDocConfigVisibilitySettings())
+								if (Builder.IsValid() && !FEditor::IgnoreDocTemplateVisibilitySettings())
 								{
 									FMetaSoundFrontendDocumentBuilder& DocBuilder = Builder->GetBuilder();
 									if (DocBuilder.IsValid())
 									{
-										if (const FDocumentConfiguration* Config = DocBuilder.GetConstDocumentConfiguration().GetPtr())
+										if (const FMetaSoundTemplate* Template = DocBuilder.GetConstDocumentTemplate().GetPtr())
 										{
-											return Config->GetEditorOptions().bConfigurationPropertiesVisible
+											return Template->GetEditorOptions().bTemplatePropertiesVisible
 												? EVisibility::Visible
 												: EVisibility::Collapsed;
 										}
@@ -433,9 +433,9 @@ namespace Metasound::Editor
 			FMetaSoundFrontendDocumentBuilder& DocBuilder = Builder->GetBuilder();
 			if (DocBuilder.IsValid())
 			{
-				if (const FDocumentConfiguration* Config = DocBuilder.GetDocumentConfiguration().GetPtr())
+				if (const FMetaSoundTemplate* Template = DocBuilder.GetDocumentTemplate().GetPtr())
 				{
-					return Config->GetEditorOptions().bPageGraphEditingEnabled;
+					return Template->GetEditorOptions().bPageGraphEditingEnabled;
 				}
 			}
 		}
@@ -895,9 +895,9 @@ namespace Metasound::Editor
 			FMetaSoundFrontendDocumentBuilder& DocBuilder = Builder->GetBuilder();
 			if (DocBuilder.IsValid())
 			{
-				if (const FDocumentConfiguration* Config = DocBuilder.GetDocumentConfiguration().GetPtr())
+				if (const FMetaSoundTemplate* Template = DocBuilder.GetDocumentTemplate().GetPtr())
 				{
-					return Config->GetEditorOptions().bInterfaceEditingEnabled;
+					return Template->GetEditorOptions().bInterfaceEditingEnabled;
 				}
 			}
 		}
diff --git a/Engine/Plugins/Runtime/Metasound/Source/MetasoundEditor/Private/MetasoundDetailCustomization.h b/Engine/Plugins/Runtime/Metasound/Source/MetasoundEditor/Private/MetasoundDetailCustomization.h
--- a/Engine/Plugins/Runtime/Metasound/Source/MetasoundEditor/Private/MetasoundDetailCustomization.h
+++ b/Engine/Plugins/Runtime/Metasound/Source/MetasoundEditor/Private/MetasoundDetailCustomization.h
@@ -51,7 +51,7 @@ namespace Metasound
 			// End of IDetailCustomization interface
 
 		private:
-			FName GetDocConfigPropertyPath() const;
+			FName GetTemplatePropertyPath() const;
 			FName GetInterfaceVersionsPropertyPath() const;
 			FName GetRootClassPropertyPath() const;
 			FName GetMetadataPropertyPath() const;
diff --git a/Engine/Plugins/Runtime/Metasound/Source/MetasoundEditor/Private/MetasoundEditor.cpp b/Engine/Plugins/Runtime/Metasound/Source/MetasoundEditor/Private/MetasoundEditor.cpp
--- a/Engine/Plugins/Runtime/Metasound/Source/MetasoundEditor/Private/MetasoundEditor.cpp
+++ b/Engine/Plugins/Runtime/Metasound/Source/MetasoundEditor/Private/MetasoundEditor.cpp
@@ -10,10 +10,6 @@
 #include "Blueprint/UserWidget.h"
 #include "Components/AudioComponent.h"
 #include "DetailLayoutBuilder.h"
-#include "DocumentConfigurationViewModels/MetasoundDocumentConfigurationViewModelBase.h"
-#include "DocumentConfigurations/MetasoundDocumentConfigurator.h"
-#include "DocumentConfigurations/MetasoundFrontendDocumentConfiguration.h"
-#include "DocumentConfigurations/MetasoundFrontendPresetDocumentConfiguration.h"
 #include "EdGraph/EdGraphNode.h"
 #include "EdGraph/EdGraphPin.h"
 #include "EdGraphUtilities.h"
@@ -45,7 +41,6 @@
 #include "MetasoundAssetSubsystem.h"
 #include "MetasoundBuilderSubsystem.h"
 #include "MetasoundDocumentBuilderRegistry.h"
-#include "MetasoundDocumentConfigurationWidgetInterface.h"
 #include "MetasoundDocumentInterface.h"
 #include "MetasoundEditorCommands.h"
 #include "MetasoundEditorDocumentClipboardUtils.h"
@@ -72,6 +67,7 @@
 #include "MetasoundPresetWidgetInterface.h"
 #include "MetasoundSettings.h"
 #include "MetasoundSource.h"
+#include "MetasoundTemplateWidgetInterface.h"
 #include "MetasoundUObjectRegistry.h"
 #include "Misc/Attribute.h"
 #include "Modules/ModuleManager.h"
@@ -87,7 +83,11 @@
 #include "Stats/Stats.h"
 #include "Styling/AppStyle.h"
 #include "Templates/Function.h"
+#include "DocumentTemplates/MetasoundDocumentConfigurator.h"
+#include "DocumentTemplates/MetasoundFrontendPresetTemplate.h"
+#include "DocumentTemplates/MetasoundFrontendDocumentTemplate.h"
 #include "Templates/SharedPointer.h"
+#include "TemplateViewModels/MetasoundTemplateViewModelBase.h"
 #include "ToolMenuEntry.h"
 #include "ToolMenus.h"
 #include "UObject/ScriptInterface.h"
@@ -120,7 +120,7 @@ namespace Metasound
 		namespace AssetEditorPrivate
 		{
 			static bool bEnableLiveAudition = false;
-			static bool bIgnoreDocConfigVisibilitySettings = false;
+			static bool bIgnoreDocTemplateVisibilitySettings = false;
 			static bool bShowGraphOverride = false;
 			static const FText LoadingAssetsFormat = LOCTEXT("MetaSoundScanInProgressNotificationText", "Loading {0} MetaSound Assets...");
 
@@ -132,10 +132,10 @@ namespace Metasound
 				TEXT("Default: false"),
 				ECVF_Default);
 
-			FAutoConsoleVariableRef CVarMetaSoundEditorIgnoreConfigVisibilitySettings(
-				TEXT("au.MetaSound.Editor.IgnoreConfigVisibilitySettings"),
-				bIgnoreDocConfigVisibilitySettings,
-				TEXT("If true, all settings on a Document Configuration for what is and is not visible in editor are ignored and all tabs are invokable.")
+			FAutoConsoleVariableRef CVarMetaSoundEditorIgnoreTemplateVisibilitySettings(
+				TEXT("au.MetaSound.Editor.IgnoreTemplateVisibilitySettings"),
+				bIgnoreDocTemplateVisibilitySettings,
+				TEXT("If true, all settings on a MetaSound Template for what is and is not visible in editor are ignored and all tabs are invokable.")
 				TEXT("Default: false"),
 				ECVF_Default);
 
@@ -231,16 +231,16 @@ namespace Metasound
 				};
 			}
 
-			void SanitizeTabInfo(const UScriptStruct& ConfigStruct, FMetaSoundDocumentConfigurationWidgetTabInfo& InOutTabInfo)
+			void SanitizeTabInfo(const UScriptStruct& TemplateStruct, FMetaSoundTemplateWidgetTabInfo& InOutTabInfo)
 			{
 				if (InOutTabInfo.TabName.IsNone())
 				{
 					FNameBuilder NameBuilder;
-					NameBuilder.Appendf(TEXT("MetasoundEditor_%s"), *ConfigStruct.GetName());
+					NameBuilder.Appendf(TEXT("MetasoundEditor_%s"), *TemplateStruct.GetName());
 					InOutTabInfo.TabName = *NameBuilder;
 					if (InOutTabInfo.TabDisplayName.IsEmpty())
 					{
-						InOutTabInfo.TabDisplayName = ConfigStruct.GetDisplayNameText();
+						InOutTabInfo.TabDisplayName = TemplateStruct.GetDisplayNameText();
 					}
 				}
 
@@ -252,8 +252,8 @@ namespace Metasound
 				if (InOutTabInfo.ToolTip.IsEmpty())
 				{
 					InOutTabInfo.ToolTip = FText::Format(
-						LOCTEXT("ConfigWidgetTab_DefaultToolTipFormat", "Custom control for MetaSound '{0}' configuration"),
-						ConfigStruct.GetDisplayNameText());
+						LOCTEXT("TemplateWidgetTab_DefaultToolTipFormat", "Custom control for MetaSound '{0}' configuration"),
+						TemplateStruct.GetDisplayNameText());
 				}
 
 				if (InOutTabInfo.StyleName.IsNone())
@@ -331,10 +331,10 @@ namespace Metasound
 				return false;
 			}
 
-			TConstStructView<FMetaSoundFrontendDocumentConfiguration> ConfigView = Builder.GetConstDocumentConfiguration();
-			if (ConfigView.IsValid())
+			TConstStructView<FMetaSoundFrontendDocumentTemplate> TemplateView = Builder.GetConstDocumentTemplate();
+			if (TemplateView.IsValid())
 			{
-				return ConfigView.GetPtr()->GetEditorOptions().bPageGraphEditingEnabled;
+				return TemplateView.GetPtr()->GetEditorOptions().bPageGraphEditingEnabled;
 			}
 
 			return true;
@@ -381,7 +381,7 @@ namespace Metasound
 			const FName Pages = "MetasoundEditor_Pages";
 			const FName Find = "MetasoundEditor_Find";
 
-			const FName DocConfigWidgets_Factory = "MetasoundEditor_DocConfigWidgetFactory";
+			const FName DocTemplateWidgets_Factory = "MetasoundEditor_DocTemplateWidgetFactory";
 		} // namespace TabNamesPrivate
 
 		static const TArray<FText> NodeSectionNames
@@ -1085,7 +1085,7 @@ namespace Metasound
 			}
 		};
 
-		FEditor::FDocConfigWidgetSummoner::FDocConfigWidgetSummoner(FName FactoryName, FEditor::FBlueprintId InBlueprintId, TSharedPtr<class FAssetEditorToolkit> InHostingApp)
+		FEditor::FTemplateWidgetSummoner::FTemplateWidgetSummoner(FName FactoryName, FEditor::FBlueprintId InBlueprintId, TSharedPtr<class FAssetEditorToolkit> InHostingApp)
 			: FWorkflowTabFactory(FactoryName, InHostingApp)
 			, BlueprintId(InBlueprintId)
 		{
@@ -1095,21 +1095,21 @@ namespace Metasound
 
 			TabIcon = FSlateIcon(FAppStyle::GetAppStyleSetName(), "LevelEditor.Tabs.Details");
 
-			TabLabel = LOCTEXT("MetaSoundConfigurationView_MenuTitle", "Configuration");
-			TabIdentifier = "Configuration";
-			ViewMenuTooltip = LOCTEXT("MetaSoundConfigurationView_ToolTip", "Custom view for the current MetaSound configuration");
+			TabLabel = LOCTEXT("MetaSoundTemplateView_MenuTitle", "Template");
+			TabIdentifier = "Template";
+			ViewMenuTooltip = LOCTEXT("MetaSoundTemplateView_ToolTip", "Custom view for the current MetaSound Template");
 
 			TSharedPtr<FEditor> Parent = GetParentEditor();
 			if (ensure(Parent.IsValid()))
 			{
-				if (FEditor::FUserWidgetData* WidgetData = Parent->DocConfigWidgetData.Find(BlueprintId))
+				if (FEditor::FUserWidgetData* WidgetData = Parent->DocTemplateWidgetData.Find(BlueprintId))
 				{
 					TabLabel = WidgetData->TabInfo.TabDisplayName;
 					TabIdentifier = WidgetData->TabInfo.TabName;
 					TabIcon = FSlateIcon(WidgetData->TabInfo.StyleName, WidgetData->TabInfo.TabIcon);
-					if (const FDocumentConfiguration* Config = Parent->GetConfiguration())
+					if (const FMetaSoundTemplate* Template = Parent->GetDocumentTemplate())
 					{
-						ViewMenuTooltip = Config->GetEditorOptions().ToolTip;
+						ViewMenuTooltip = Template->GetEditorOptions().ToolTip;
 					}
 				}
 			}
@@ -1117,7 +1117,7 @@ namespace Metasound
 			ViewMenuDescription = TabLabel;
 		}
 
-		TSharedRef<SWidget> FEditor::FDocConfigWidgetSummoner::CreateTabBody(const FWorkflowTabSpawnInfo& Info) const
+		TSharedRef<SWidget> FEditor::FTemplateWidgetSummoner::CreateTabBody(const FWorkflowTabSpawnInfo& Info) const
 		{
 			if (const FUserWidgetData* WidgetData = GetWidgetData())
 			{
@@ -1130,7 +1130,7 @@ namespace Metasound
 			return SNullWidget::NullWidget;
 		}
 
-		TSharedRef<SDockTab> FEditor::FDocConfigWidgetSummoner::SpawnTab(const FWorkflowTabSpawnInfo& Info) const
+		TSharedRef<SDockTab> FEditor::FTemplateWidgetSummoner::SpawnTab(const FWorkflowTabSpawnInfo& Info) const
 		{
 			TSharedRef<SDockTab> DockTab = SNew(SDockTab);
 
@@ -1147,32 +1147,32 @@ namespace Metasound
 			return DockTab;
 		}
 
-		TSharedPtr<FEditor> FEditor::FDocConfigWidgetSummoner::GetParentEditor() const
+		TSharedPtr<FEditor> FEditor::FTemplateWidgetSummoner::GetParentEditor() const
 		{
 			return StaticCastWeakPtr<FEditor>(HostingApp).Pin();
 		}
 
-		const FEditor::FUserWidgetData* FEditor::FDocConfigWidgetSummoner::GetWidgetData() const
+		const FEditor::FUserWidgetData* FEditor::FTemplateWidgetSummoner::GetWidgetData() const
 		{
 			TSharedPtr<FEditor> Parent = GetParentEditor();
 			if (ensure(Parent.IsValid()))
 			{
-				return Parent->DocConfigWidgetData.Find(BlueprintId);
+				return Parent->DocTemplateWidgetData.Find(BlueprintId);
 			}
 
 			return nullptr;
 		}
 
-		FText FEditor::FDocConfigWidgetSummoner::GetTabToolTipText(const FWorkflowTabSpawnInfo& Info) const
+		FText FEditor::FTemplateWidgetSummoner::GetTabToolTipText(const FWorkflowTabSpawnInfo& Info) const
 		{
 			using namespace Frontend;
 
 			TSharedPtr<FEditor> Parent = GetParentEditor();
 			if (Parent.IsValid())
 			{
-				if (const FDocumentConfiguration* Config = Parent->GetConfiguration())
+				if (const FMetaSoundTemplate* Template = Parent->GetDocumentTemplate())
 				{
-					return Config->GetEditorOptions().ToolTip;
+					return Template->GetEditorOptions().ToolTip;
 				}
 			}
 			return { };
@@ -1323,17 +1323,17 @@ namespace Metasound
 			InTabManager->UnregisterTabSpawner(TabNamesPrivate::Interfaces);
 			InTabManager->UnregisterTabSpawner(TabNamesPrivate::Find);
 
-			for (auto FactoryIt = DocConfigTabSet.CreateConstIterator(); FactoryIt; ++FactoryIt)
+			for (auto FactoryIt = DocTemplateTabSet.CreateConstIterator(); FactoryIt; ++FactoryIt)
 			{
 				const FName TabId = FactoryIt.Value()->GetIdentifier();
 				InTabManager->UnregisterTabSpawner(TabId);
 			}
-			DocConfigTabSet.Clear();
+			DocTemplateTabSet.Clear();
 		}
 
-		bool FEditor::IgnoreDocConfigVisibilitySettings()
+		bool FEditor::IgnoreDocTemplateVisibilitySettings()
 		{
-			return AssetEditorPrivate::bIgnoreDocConfigVisibilitySettings;
+			return AssetEditorPrivate::bIgnoreDocTemplateVisibilitySettings;
 		}
 
 		bool FEditor::IsLiveAuditionEnabled()
@@ -1375,7 +1375,7 @@ namespace Metasound
 			FLevelEditorModule& LevelEditor = FModuleManager::LoadModuleChecked<FLevelEditorModule>("LevelEditor");
 			LevelEditor.OnMapChanged().RemoveAll(this);
 
-			for (FBPIdWidgetDataPair& Pair : DocConfigWidgetData)
+			for (FBPIdWidgetDataPair& Pair : DocTemplateWidgetData)
 			{
 				if (UUserWidget* UserWidget = Pair.Value.Widget.Get())
 				{
@@ -1385,7 +1385,7 @@ namespace Metasound
 					}
 				}
 			}
-			DocConfigWidgetData.Reset();
+			DocTemplateWidgetData.Reset();
 
 			check(GEditor);
 			GEditor->UnregisterForUndo(this);
@@ -1470,7 +1470,7 @@ namespace Metasound
 				->SetExtensionId(TabNamesPrivate::GraphCanvas)
 				->AddTab(TabNamesPrivate::GraphCanvas, ETabState::OpenedTab);
 
-			for (const FBPIdWidgetDataPair& Pair : DocConfigWidgetData)
+			for (const FBPIdWidgetDataPair& Pair : DocTemplateWidgetData)
 			{
 				GraphCanvasTabStack->AddTab(Pair.Value.TabInfo.TabName, ETabState::ClosedTab);
 			}
@@ -1544,7 +1544,7 @@ namespace Metasound
 			ExtendToolbarInternal();
 			RegenerateMenusAndToolbars();
 
-			ReloadAllDocConfigTabs();
+			ReloadAllDocTemplateTabs();
 
 			NotifyDocumentVersioned();
 
@@ -1641,9 +1641,9 @@ namespace Metasound
 				if (DocBuilder.IsValid())
 				{
 					const FMetasoundFrontendDocument& Document = DocBuilder.GetConstDocumentChecked();
-					if (const FMetaSoundFrontendDocumentConfiguration* DocConfig = Document.Configuration.GetPtr())
+					if (const FMetaSoundFrontendDocumentTemplate* DocTemplate = Document.Template.GetPtr())
 					{
-						if (!DocConfig->GetEditorOptions().bPageGraphEditingEnabled)
+						if (!DocTemplate->GetEditorOptions().bPageGraphEditingEnabled)
 						{
 							return false;
 						}
@@ -1719,10 +1719,13 @@ namespace Metasound
 				}
 
 				TScriptInterface<const IMetaSoundDocumentInterface> DocInterface = MetaSound;
-				const UStruct* Struct = DocInterface->GetConstDocument().Configuration.GetScriptStruct();
-				if (Struct && Struct->IsChildOf(FPresetDocumentConfiguration::StaticStruct()))
+				if (DocInterface.GetObject() != nullptr)
 				{
-					IconName += TEXT(".Preset");
+					const UStruct* Struct = DocInterface->GetConstDocument().Template.GetScriptStruct();
+					if (Struct && Struct->IsChildOf(FPresetTemplate::StaticStruct()))
+					{
+						IconName += TEXT(".Preset");
+					}
 				}
 
 				IconName += TEXT(".Icon");
@@ -2004,7 +2007,7 @@ namespace Metasound
 			if (Builder.IsValid() && !Builder->GetConstBuilder().GetIsGraphEditable())
 			{
 				FLevelEditorModule& LevelEditor = FModuleManager::LoadModuleChecked<FLevelEditorModule>("LevelEditor");
-				LevelEditor.OnMapChanged().AddSP(this, &FEditor::ChangeDocConfigWidgetWorld);
+				LevelEditor.OnMapChanged().AddSP(this, &FEditor::ChangeDocTemplateWidgetWorld);
 			}
 
 			Palette = SNew(SMetasoundPalette);
@@ -2045,13 +2048,13 @@ namespace Metasound
 			}
 		}
 
-		void FEditor::ReloadAllDocConfigTabs()
+		void FEditor::ReloadAllDocTemplateTabs()
 		{
 			using namespace Frontend;
 
 			if (TabManager.IsValid())
 			{
-				for (auto FactoryIt = DocConfigTabSet.CreateConstIterator(); FactoryIt; ++FactoryIt)
+				for (auto FactoryIt = DocTemplateTabSet.CreateConstIterator(); FactoryIt; ++FactoryIt)
 				{
 					const FName TabId = FactoryIt.Value()->GetIdentifier();
 					if (TSharedPtr<SDockTab> Tab = TabManager->FindExistingLiveTab(TabId))
@@ -2062,9 +2065,9 @@ namespace Metasound
 				}
 			}
 
-			DocConfigTabSet.Clear();
+			DocTemplateTabSet.Clear();
 
-			for (TPair<FBlueprintId, FUserWidgetData>& WidgetData : DocConfigWidgetData)
+			for (TPair<FBlueprintId, FUserWidgetData>& WidgetData : DocTemplateWidgetData)
 			{
 				if (UUserWidget* Widget = WidgetData.Value.Widget.Get())
 				{
@@ -2074,26 +2077,26 @@ namespace Metasound
 					}
 				}
 			}
-			DocConfigWidgetData.Reset();
+			DocTemplateWidgetData.Reset();
 
-			RebuildDocConfigUserWidgets();
+			RebuildDocTemplateUserWidgets();
 
 			TMap<FName, FSlateIcon> TabIcons;
-			for (const FBPIdWidgetDataPair& Pair : DocConfigWidgetData)
+			for (const FBPIdWidgetDataPair& Pair : DocTemplateWidgetData)
 			{
 				const FName TabName = Pair.Value.TabInfo.TabName;
-				FNameBuilder FactoryName(TabNamesPrivate::DocConfigWidgets_Factory);
+				FNameBuilder FactoryName(TabNamesPrivate::DocTemplateWidgets_Factory);
 				if (!TabName.IsNone())
 				{
 					FactoryName.Appendf(TEXT("_%s"), *TabName.ToString());
 				}
-				DocConfigTabSet.RegisterFactory(MakeShared<FDocConfigWidgetSummoner>(*FactoryName, Pair.Key, AsShared()));
+				DocTemplateTabSet.RegisterFactory(MakeShared<FTemplateWidgetSummoner>(*FactoryName, Pair.Key, AsShared()));
 				TabIcons.Add(TabName, FSlateIcon(Pair.Value.TabInfo.StyleName, Pair.Value.TabInfo.TabIcon));
 			}
 
 			if (TabManager.IsValid())
 			{
-				for (auto FactoryIt = DocConfigTabSet.CreateConstIterator(); FactoryIt; ++FactoryIt)
+				for (auto FactoryIt = DocTemplateTabSet.CreateConstIterator(); FactoryIt; ++FactoryIt)
 				{
 					FTabSpawnerEntry& SpawnerEntry = FactoryIt.Value()->RegisterTabSpawner(TabManager.ToSharedRef(), nullptr)
 					.SetGroup(WorkspaceMenuCategory.ToSharedRef());
@@ -2108,7 +2111,7 @@ namespace Metasound
 			}
 		}
 
-		void FEditor::RebuildDocConfigUserWidgets(const FGuid& TargetBlueprintId)
+		void FEditor::RebuildDocTemplateUserWidgets(const FGuid& TargetBlueprintId)
 		{
 			using namespace AssetEditorPrivate;
 			using namespace Frontend;
@@ -2125,8 +2128,8 @@ namespace Metasound
 				return;
 			}
 
-			TConstStructView<FDocumentConfiguration> DocConfig = DocBuilder.GetConstDocumentConfiguration();
-			if (!DocConfig.IsValid() || !DocConfig.GetPtr()->GetEditorOptions().bCustomTabEnabled)
+			TConstStructView<FMetaSoundTemplate> DocTemplate = DocBuilder.GetConstDocumentTemplate();
+			if (!DocTemplate.IsValid() || !DocTemplate.GetPtr()->GetEditorOptions().bCustomTabEnabled)
 			{
 				return;
 			}
@@ -2146,11 +2149,11 @@ namespace Metasound
 			// Preset separated out for back compat
 			if (TFunction<bool(UUserWidget*)> PresetFilterFunction = GetPresetWidgetFilterFunction(DocBuilder); PresetFilterFunction.IsSet())
 			{
-				TFunction<FMetaSoundDocumentConfigurationWidgetTabInfo(UUserWidget*)> OnPresetConstruct = [&DocConfig, &BuilderBase, &TargetBlueprintId](UUserWidget* UserWidget)
+				TFunction<FMetaSoundTemplateWidgetTabInfo(UUserWidget*)> OnPresetConstruct = [&DocTemplate, &BuilderBase, &TargetBlueprintId](UUserWidget* UserWidget)
 				{
 					IMetaSoundPresetWidgetInterface::Execute_OnConstructed(UserWidget, BuilderBase);
 
-					FMetaSoundDocumentConfigurationWidgetTabInfo TabInfo;
+					FMetaSoundTemplateWidgetTabInfo TabInfo;
 					IMetaSoundPresetWidgetInterface::Execute_OnInitializingTab(UserWidget, TabInfo);
 					return TabInfo;
 				};
@@ -2162,52 +2165,52 @@ namespace Metasound
 				return;
 			}
 
-			TFunction<bool(UUserWidget*)> DocConfigFilterFunction = [&DocConfig, BuilderBase](UUserWidget* UserWidget)
+			TFunction<bool(UUserWidget*)> DocTemplateFilterFunction = [&DocTemplate, BuilderBase](UUserWidget* UserWidget)
 			{
-				bool bIsSupportedConfig = true;
-				if (const UScriptStruct* ScriptStruct = DocConfig.GetScriptStruct())
+				bool bIsSupportedTemplate = true;
+				if (const UScriptStruct* ScriptStruct = DocTemplate.GetScriptStruct())
 				{
-					IMetaSoundDocumentConfigurationWidgetInterface::Execute_IsSupportedConfiguration(
+					IMetaSoundTemplateWidgetInterface::Execute_IsSupportedTemplate(
 						UserWidget,
 						ScriptStruct->GetFName(),
 						BuilderBase,
-						bIsSupportedConfig);
+						bIsSupportedTemplate);
 				}
-				return bIsSupportedConfig;
+				return bIsSupportedTemplate;
 			};
 
-			TFunction<FMetaSoundDocumentConfigurationWidgetTabInfo(UUserWidget*)> OnDocConfigTabConstruct = [&DocConfig, BuilderBase](UUserWidget* UserWidget)
+			TFunction<FMetaSoundTemplateWidgetTabInfo(UUserWidget*)> OnDocTemplateTabConstruct = [&DocTemplate, BuilderBase](UUserWidget* UserWidget)
 			{
-				TScriptInterface<IMetaSoundDocumentConfigurationWidgetInterface> WidgetInterface(UserWidget);
+				TScriptInterface<IMetaSoundTemplateWidgetInterface> WidgetInterface(UserWidget);
 				if (ensure(WidgetInterface.GetObject()))
 				{
-					IMetaSoundDocumentConfigurationWidgetInterface::Execute_OnBuilderSet(UserWidget, BuilderBase);
+					IMetaSoundTemplateWidgetInterface::Execute_OnBuilderSet(UserWidget, BuilderBase);
 				}
 
-				TSubclassOf<UMetaSoundDocumentConfigurationViewModelBase> VMClass;
-				IMetaSoundDocumentConfigurationWidgetInterface::Execute_OnInitializingViewModel(UserWidget, VMClass);
-				UMetaSoundDocumentConfigurationViewModelBase::ConstructWidgetViewModel(UserWidget, *BuilderBase, DocConfig, VMClass);
+				TSubclassOf<UMetaSoundTemplateViewModelBase> VMClass;
+				IMetaSoundTemplateWidgetInterface::Execute_OnInitializingViewModel(UserWidget, VMClass);
+				UMetaSoundTemplateViewModelBase::ConstructWidgetViewModel(UserWidget, *BuilderBase, DocTemplate, VMClass);
 
-				FMetaSoundDocumentConfigurationWidgetTabInfo TabInfo;
-				IMetaSoundDocumentConfigurationWidgetInterface::Execute_OnInitializingTab(UserWidget, TabInfo);
+				FMetaSoundTemplateWidgetTabInfo TabInfo;
+				IMetaSoundTemplateWidgetInterface::Execute_OnInitializingTab(UserWidget, TabInfo);
 				return TabInfo;
 			};
 
 			RebuildUserWidgets(*World, *WidgetSubsystem,
-				UMetaSoundDocumentConfigurationWidgetInterface::StaticClass(),
-				MoveTemp(DocConfigFilterFunction),
-				MoveTemp(OnDocConfigTabConstruct),
+				UMetaSoundTemplateWidgetInterface::StaticClass(),
+				MoveTemp(DocTemplateFilterFunction),
+				MoveTemp(OnDocTemplateTabConstruct),
 				TargetBlueprintId);
 		}
 
-		void FEditor::OnDocConfigWidgetBlueprintCompiled(UBlueprint* Blueprint)
+		void FEditor::OnDocTemplateWidgetBlueprintCompiled(UBlueprint* Blueprint)
 		{
 			using namespace AssetEditorPrivate;
 
 			if (Blueprint)
 			{
 				const FGuid BlueprintId = Blueprint->GetBlueprintGuid();
-				if (const FUserWidgetData* WidgetData = DocConfigWidgetData.Find(BlueprintId))
+				if (const FUserWidgetData* WidgetData = DocTemplateWidgetData.Find(BlueprintId))
 				{
 					if (TabManager.IsValid())
 					{
@@ -2218,9 +2221,9 @@ namespace Metasound
 					}
 				}
 
-				RebuildDocConfigUserWidgets(BlueprintId);
+				RebuildDocTemplateUserWidgets(BlueprintId);
 
-				if (const FUserWidgetData* WidgetData = DocConfigWidgetData.Find(BlueprintId))
+				if (const FUserWidgetData* WidgetData = DocTemplateWidgetData.Find(BlueprintId))
 				{
 					if (TabManager.IsValid())
 					{
@@ -2236,7 +2239,7 @@ namespace Metasound
 			}
 		}
 		
-		void FEditor::ChangeDocConfigWidgetWorld(UWorld* World, EMapChangeType MapChangeType)
+		void FEditor::ChangeDocTemplateWidgetWorld(UWorld* World, EMapChangeType MapChangeType)
 		{
 			switch (MapChangeType)
 			{
@@ -2244,12 +2247,12 @@ namespace Metasound
 				{
 					// Moves the UMG widgets to the transient package if the
 					// World they were built with is being torn down.
-					for (FBPIdWidgetDataPair& Pair : DocConfigWidgetData)
+					for (FBPIdWidgetDataPair& Pair : DocTemplateWidgetData)
 					{
-						UUserWidget* DocConfigWidget = Pair.Value.Widget.Get();
-						if (DocConfigWidget && World == DocConfigWidget->GetWorld())
+						UUserWidget* DocTemplateWidget = Pair.Value.Widget.Get();
+						if (DocTemplateWidget && World == DocTemplateWidget->GetWorld())
 						{
-							DocConfigWidget->Rename(nullptr, GetTransientPackage());
+							DocTemplateWidget->Rename(nullptr, GetTransientPackage());
 						}
 					}
 				}
@@ -2258,7 +2261,7 @@ namespace Metasound
 				case EMapChangeType::LoadMap:
 				case EMapChangeType::NewMap:
 				{
-					ReloadAllDocConfigTabs();
+					ReloadAllDocTemplateTabs();
 				}
 				break;
 
@@ -2283,12 +2286,12 @@ namespace Metasound
 			return MenuWidget;
 		}
 
-		TSharedRef<SWidget> FEditor::CreateConfigurationMenuOptions()
+		TSharedRef<SWidget> FEditor::CreateTemplateMenuOptions()
 		{
 			TSharedPtr<FUICommandList> Commands = MakeShared<FUICommandList>();
 			constexpr bool bShouldCloseWindowAfterMenuSelection = true;
 			FMenuBuilder MenuBuilder(bShouldCloseWindowAfterMenuSelection, Commands);
-			CreateConfigurationSubMenuOptions(MenuBuilder);
+			CreateTemplateSubMenuOptions(MenuBuilder);
 			TSharedRef<SWidget> MenuWidget = MenuBuilder.MakeWidget();
 			TWeakObjectPtr<UMetaSoundBuilderBase> WeakBuilderPtr(Builder.Get());
 			return MenuWidget;
@@ -2512,14 +2515,14 @@ namespace Metasound
 			MenuBuilder.EndSection();
 		}
 
-		void FEditor::CreateConfigurationSubMenuOptions(FMenuBuilder& MenuBuilder)
+		void FEditor::CreateTemplateSubMenuOptions(FMenuBuilder& MenuBuilder)
 		{
 			TAttribute<EVisibility> ClearVisibility = TAttribute<EVisibility>::CreateSPLambda(this, [this]()
 			{
 				const IMetasoundEditorModule & MetaSoundEditorModule = FModuleManager::GetModuleChecked<IMetasoundEditorModule>("MetaSoundEditor");
 				if (!MetaSoundEditorModule.IsRestrictedMode())
 				{
-					return GetConfiguration() ? EVisibility::Visible : EVisibility::Collapsed;
+					return GetDocumentTemplate() ? EVisibility::Visible : EVisibility::Collapsed;
 				}
 
 				return EVisibility::Collapsed;
@@ -2529,15 +2532,15 @@ namespace Metasound
 			IMetasoundEditorModule& MetaSoundEditorModule = FModuleManager::GetModuleChecked<IMetasoundEditorModule>("MetaSoundEditor");
 			if (!MetaSoundEditorModule.IsRestrictedMode())
 			{
-				MenuBuilder.BeginSection("ConvertConfigurationHeader", LOCTEXT("ConvertConfigurationDisplayName", "Convert Configuration"), ClearVisibility);
+				MenuBuilder.BeginSection("ConvertTemplateHeader", LOCTEXT("ConvertTemplateDisplayName", "Convert Template"), ClearVisibility);
 				{
 					if (const FMetaSoundFrontendDocumentBuilder* DocBuilder = GetFrontendBuilder())
 					{
-						if (const UScriptStruct* ConfigStruct = DocBuilder->GetConstDocumentConfiguration().GetScriptStruct())
+						if (const UScriptStruct* ConfigStruct = DocBuilder->GetConstDocumentTemplate().GetScriptStruct())
 						{
 							MenuBuilder.AddMenuEntry(
-								FText::Format(LOCTEXT("ClearConfigurationButtonText", "Convert From '{0}' ..."), ConfigStruct->GetDisplayNameText()),
-								LOCTEXT("ConvertFromConfigurationTooltip", "Removes the MetaSound's configuration, leaving the MetaSound state intact and enabling editing of configuration-managed MetaSound functionality."),
+								FText::Format(LOCTEXT("ClearTemplateButtonText", "Convert From '{0}' ..."), ConfigStruct->GetDisplayNameText()),
+								LOCTEXT("ConvertFromTemplateTooltip", "Removes the MetaSound's configuration, leaving the MetaSound state intact and enabling editing of configuration-managed MetaSound functionality."),
 								FSlateIcon(),
 								FUIAction(FExecuteAction::CreateSP(this, &FEditor::ShowConvertFromConfigDialog)),
 								{ },
@@ -2549,7 +2552,7 @@ namespace Metasound
 				MenuBuilder.EndSection();
 			}
 
-			MenuBuilder.BeginSection("SetConfigurationSectionHeader", LOCTEXT("SelectConfigurationSectionName", "Set Configuration"));
+			MenuBuilder.BeginSection("SetTemplateSectionHeader", LOCTEXT("SelectTemplateSectionName", "Set Template"));
 			{
 				auto CreateConfigEntry = [this, &MenuBuilder](UScriptStruct* ConfigStruct)
 				{
@@ -2558,15 +2561,15 @@ namespace Metasound
 					{
 						if (FMetaSoundFrontendDocumentBuilder* DocBuilder = GetFrontendBuilder())
 						{
-							TInstancedStruct<FMetaSoundFrontendDocumentConfiguration> ConfigInstance;
+							TInstancedStruct<FMetaSoundFrontendDocumentTemplate> ConfigInstance;
 
 							check(ConfigStruct);
 							ConfigInstance.InitializeAsScriptStruct(ConfigStruct);
 
-							DocBuilder->SetDocumentConfiguration(MoveTemp(ConfigInstance));
+							DocBuilder->SetTemplate(MoveTemp(ConfigInstance));
 							DocBuilder->ConfigureDocument();
 
-							ReloadAllDocConfigTabs();
+							ReloadAllDocTemplateTabs();
 							UpdateEnabledTabs();
 						}
 					});
@@ -2574,7 +2577,7 @@ namespace Metasound
 					{
 						if (const FMetaSoundFrontendDocumentBuilder* DocBuilder = GetFrontendBuilder())
 						{
-							if (ConfigStruct == DocBuilder->GetConstDocumentConfiguration().GetScriptStruct())
+							if (ConfigStruct == DocBuilder->GetConstDocumentTemplate().GetScriptStruct())
 							{
 								return ECheckBoxState::Checked;
 							}
@@ -2597,7 +2600,7 @@ namespace Metasound
 				{
 					if (UScriptStruct* Struct = Cast<UScriptStruct>(*It))
 					{
-						if (Struct->IsChildOf(FMetaSoundFrontendDocumentConfiguration::StaticStruct())
+						if (Struct->IsChildOf(FMetaSoundFrontendDocumentTemplate::StaticStruct())
 							&& !Struct->HasMetaData("Hidden"))
 						{
 							CreateConfigEntry(Struct);
@@ -2630,9 +2633,9 @@ namespace Metasound
  					{
 						ToolbarBuilder.AddComboButton(
 							FUIAction(),
-							FOnGetContent::CreateSP(this, &FEditor::CreateConfigurationMenuOptions),
-							LOCTEXT("ConfigurationSettingsMenu", "Configuration"),
-							LOCTEXT("ConfigurationSettingsMenu_Tooltip", "Settings related to procedurally-generated MetaSound Configurations"),
+							FOnGetContent::CreateSP(this, &FEditor::CreateTemplateMenuOptions),
+							LOCTEXT("TemplateSettingsMenu", "Template"),
+							LOCTEXT("TemplateSettingsMenu_Tooltip", "Settings related to procedurally-generated MetaSound Templates"),
 							GetExportStatusImage()
 						);
  					}
@@ -2824,8 +2827,8 @@ namespace Metasound
 				FExecuteAction::CreateSP(this, &FEditor::ShowFindInMetaSound));
 
 			ToolkitCommands->MapAction(
-				FEditorCommands::Get().CopySelectionAsConfigurationCodeToClipboard,
-				FExecuteAction::CreateSP(this, &FEditor::CopySelectionAsConfigurationCodeToClipboard));
+				FEditorCommands::Get().CopySelectionAsConfiguratorCodeToClipboard,
+				FExecuteAction::CreateSP(this, &FEditor::CopySelectionAsConfiguratorCodeToClipboard));
 
 			ToolkitCommands->MapAction(
 				FEditorCommands::Get().ToggleShowGraph,
@@ -3009,9 +3012,9 @@ namespace Metasound
 				{
 					IMetaSoundPresetWidgetInterface::Execute_OnAuditionStateChanged(UserWidget, AudioComponent, bIsAuditioning);
 				}
-				else if (ensure(UMetaSoundDocumentConfigurationWidgetInterface::StaticClass() == Class))
+				else if (ensure(UMetaSoundTemplateWidgetInterface::StaticClass() == Class))
 				{
-					IMetaSoundDocumentConfigurationWidgetInterface::Execute_OnAuditionStateChanged(UserWidget, AudioComponent, bIsAuditioning);
+					IMetaSoundTemplateWidgetInterface::Execute_OnAuditionStateChanged(UserWidget, AudioComponent, bIsAuditioning);
 				}
 			}
 		}
@@ -3051,7 +3054,7 @@ namespace Metasound
 						if (IsPlaying())
 						{
 							constexpr bool bIsAuditioning = true;
-							for (FBPIdWidgetDataPair& Pair : DocConfigWidgetData)
+							for (FBPIdWidgetDataPair& Pair : DocTemplateWidgetData)
 							{
 								Pair.Value.SetAuditionState(PreviewComp, bIsAuditioning);
 							}
@@ -3084,7 +3087,7 @@ namespace Metasound
 							GraphConnectionManager = RebuildConnectionManager();
 
 							constexpr bool bIsAuditioning = false;
-							for (FBPIdWidgetDataPair& Pair : DocConfigWidgetData)
+							for (FBPIdWidgetDataPair& Pair : DocTemplateWidgetData)
 							{
 								Pair.Value.SetAuditionState(nullptr, bIsAuditioning);
 							}
@@ -3210,8 +3213,8 @@ namespace Metasound
 			{
 				if (const FMetaSoundFrontendDocumentBuilder& DocBuilder = Builder->GetBuilder(); DocBuilder.IsValid())
 				{
-					TConstStructView<FMetaSoundFrontendDocumentConfiguration> Config = DocBuilder.GetConstDocumentConfiguration();
-					return Config.IsValid() && !Config.GetPtr()->GetEditorOptions().bGraphEditorVisible;
+					TConstStructView<FMetaSoundFrontendDocumentTemplate> Template = DocBuilder.GetConstDocumentTemplate();
+					return Template.IsValid() && !Template.GetPtr()->GetEditorOptions().bGraphEditorVisible;
 				}
 			}
 
@@ -3291,7 +3294,7 @@ namespace Metasound
 			{
 				TSharedPtr<SWindow> DialogWindow =
 					SNew(SWindow)
-					.Title(LOCTEXT("ConvertMetaSoundConfigurationPrompt", "Convert From MetaSound Configuration?"))
+					.Title(LOCTEXT("ConvertMetaSoundTemplatePrompt", "Convert From MetaSound Template?"))
 					.SupportsMinimize(false)
 					.SupportsMaximize(false)
 					.SizingRule(ESizingRule::Autosized)
@@ -3314,15 +3317,15 @@ namespace Metasound
 							{
 								if (FMetaSoundFrontendDocumentBuilder& DocBuilder = Builder->GetBuilder(); DocBuilder.IsValid())
 								{
-									const UScriptStruct* ScriptStruct = DocBuilder.GetConstDocumentConfiguration().GetScriptStruct();
+									const UScriptStruct* ScriptStruct = DocBuilder.GetConstDocumentTemplate().GetScriptStruct();
 									const FText DisplayName = ScriptStruct ? ScriptStruct->GetDisplayNameText() : LOCTEXT("EmptyConfigName", "Unset");
-									const FScopedTransaction Transaction(FText::Format(LOCTEXT("ConvertFromConfigurationFormat", "Convert From '{0}' Configuration"), DisplayName));
+									const FScopedTransaction Transaction(FText::Format(LOCTEXT("ConvertFromTemplateFormat", "Convert From '{0}' Template"), DisplayName));
 
 									GetMetasoundObject()->Modify();
 									EMetaSoundBuilderResult Result = EMetaSoundBuilderResult::Failed;
-									DocBuilder.SetDocumentConfiguration({ });
+									DocBuilder.SetTemplate({ });
 									DocBuilder.ConfigureDocument();
-									ReloadAllDocConfigTabs();
+									ReloadAllDocTemplateTabs();
 									UpdateEnabledTabs();
 								}
 								
@@ -3497,8 +3500,8 @@ namespace Metasound
 					FExecuteAction::CreateLambda([this] { PromoteToDeferredVariable(); }),
 					FCanExecuteAction::CreateLambda([this]() { return CanPromoteToDeferredVariable(); }));
 
-				GraphEditorCommands->MapAction(FEditorCommands::Get().CopySelectionAsConfigurationCodeToClipboard,
-					FExecuteAction::CreateSP(this, &FEditor::CopySelectionAsConfigurationCodeToClipboard));
+				GraphEditorCommands->MapAction(FEditorCommands::Get().CopySelectionAsConfiguratorCodeToClipboard,
+					FExecuteAction::CreateSP(this, &FEditor::CopySelectionAsConfiguratorCodeToClipboard));
 
 				GraphEditorCommands->MapAction(FGraphEditorCommands::Get().HideNoConnectionPins,
 					FExecuteAction::CreateSP(this, &FEditor::HideUnconnectedPins));
@@ -3545,8 +3548,8 @@ namespace Metasound
 				GraphEditorCommands->MapAction(FEditorCommands::Get().UpdateNodeClass,
 					FExecuteAction::CreateSP(this, &FEditor::UpdateSelectedNodeClasses));
 
-				GraphEditorCommands->MapAction(FEditorCommands::Get().CopySelectionAsConfigurationCodeToClipboard,
-					FExecuteAction::CreateSP(this, &FEditor::CopySelectionAsConfigurationCodeToClipboard));
+				GraphEditorCommands->MapAction(FEditorCommands::Get().CopySelectionAsConfiguratorCodeToClipboard,
+					FExecuteAction::CreateSP(this, &FEditor::CopySelectionAsConfiguratorCodeToClipboard));
 			}
 
 			SGraphEditor::FGraphEditorEvents GraphEvents;
@@ -4096,11 +4099,11 @@ namespace Metasound
 			return nullptr;
 		}
 
-		const Frontend::FDocumentConfiguration* FEditor::GetConfiguration() const
+		const Frontend::FMetaSoundTemplate* FEditor::GetDocumentTemplate() const
 		{
 			if (const FMetaSoundFrontendDocumentBuilder* DocBuilder = GetFrontendBuilder())
 			{
-				return DocBuilder->GetConstDocumentConfiguration().GetPtr();
+				return DocBuilder->GetConstDocumentTemplate().GetPtr();
 			}
 
 			return nullptr;
@@ -4410,7 +4413,7 @@ namespace Metasound
 			DeleteSelectedInterfaceItems();
 		}
 
-		void FEditor::CopySelectionAsConfigurationCodeToClipboard() const
+		void FEditor::CopySelectionAsConfiguratorCodeToClipboard() const
 		{
 			if (const FMetaSoundFrontendDocumentBuilder* DocBuilder = GetFrontendBuilder())
 			{
@@ -5260,8 +5263,8 @@ PRAGMA_ENABLE_DEPRECATION_WARNINGS
 
 				if (!MembersModified.IsEmpty() || bForceRefreshViews)
 				{
-					if (const Frontend::FDocumentConfiguration* Config = GetConfiguration();
-						!Config || Config->GetEditorOptions().bMemberEditingEnabled)
+					if (const Frontend::FMetaSoundTemplate* Template = GetDocumentTemplate();
+						!Template || Template->GetEditorOptions().bMemberEditingEnabled)
 					{
 						UMetasoundEditorGraphMember* SelectedMember = RefreshGraphMemberMenu();
 
@@ -5672,11 +5675,11 @@ PRAGMA_ENABLE_DEPRECATION_WARNINGS
 				const bool bEnablePageTab = IsPageGraphEditingEnabled(*DocBuilder, bHasProjectPageValues);
 
 				bool bConfigVisibilitySet = false;
-				if (const FDocumentConfiguration* Config = GetConfiguration())
+				if (const FMetaSoundTemplate* Template = GetDocumentTemplate())
 				{
-					if (!AssetEditorPrivate::bIgnoreDocConfigVisibilitySettings)
+					if (!AssetEditorPrivate::bIgnoreDocTemplateVisibilitySettings)
 					{
-						const FDocumentConfiguration::FEditorOptions& EdOptions = Config->GetEditorOptions();
+						const FMetaSoundTemplate::FEditorOptions& EdOptions = Template->GetEditorOptions();
 						const bool bShowGraph = EdOptions.bGraphEditorVisible || AssetEditorPrivate::bShowGraphOverride;
 						ResolveTabEnablement(TabNamesPrivate::Find, bShowGraph);
 						ResolveTabEnablement(TabNamesPrivate::GraphCanvas, bShowGraph);
@@ -5687,7 +5690,7 @@ PRAGMA_ENABLE_DEPRECATION_WARNINGS
 						// so need to resolve enablement for the pages tab here.
 
 						// Don't show custom config tab if no config widget was constructed, even if enabled
-						for (const FBPIdWidgetDataPair& Pair : DocConfigWidgetData)
+						for (const FBPIdWidgetDataPair& Pair : DocTemplateWidgetData)
 						{
 							ResolveTabEnablement(Pair.Value.TabInfo.TabName, EdOptions.bCustomTabEnabled);
 						}
@@ -5709,8 +5712,8 @@ PRAGMA_ENABLE_DEPRECATION_WARNINGS
 
 					{
 						// Doc config widgets are only visible when visibility is set unless ignore cvar is true.
-						const bool bEnableTab = !AssetEditorPrivate::bIgnoreDocConfigVisibilitySettings;
-						for (const FBPIdWidgetDataPair& Pair : DocConfigWidgetData)
+						const bool bEnableTab = !AssetEditorPrivate::bIgnoreDocTemplateVisibilitySettings;
+						for (const FBPIdWidgetDataPair& Pair : DocTemplateWidgetData)
 						{
 							ResolveTabEnablement(Pair.Value.TabInfo.TabName, bEnableTab);
 						}
@@ -5838,16 +5841,16 @@ PRAGMA_ENABLE_DEPRECATION_WARNINGS
 			OutDelegates.PageDelegates.OnPageSet.AddSP(this, &FDocumentListener::OnPageSet);
 			OutDelegates.InterfaceDelegates.OnInputDefaultChanged.AddSP(this, &FDocumentListener::OnInputDefaultChanged);
 			OutDelegates.InterfaceDelegates.OnMemberMetadataSet.AddSP(this, &FDocumentListener::OnMemberMetadataSet);
-			OutDelegates.OnDocumentConfigurationChanged.AddSP(this, &FDocumentListener::OnDocumentConfigurationChanged);
+			OutDelegates.OnDocumentTemplateChanged.AddSP(this, &FDocumentListener::OnDocumentTemplateChanged);
 		}
 
-		void FEditor::FDocumentListener::OnDocumentConfigurationChanged(const Frontend::FDocumentConfigChangedArgs& Args)
+		void FEditor::FDocumentListener::OnDocumentTemplateChanged(const Frontend::FDocumentTemplateChangedArgs& Args)
 		{
 			if (TSharedPtr<FEditor> ParentPtr = Parent.Pin())
 			{
-				if (Args.Type != Frontend::EDocumentConfigChangeType::Property)
+				if (Args.Type != Frontend::EDocumentTemplateChangeType::Property)
 				{
-					ParentPtr->ReloadAllDocConfigTabs();
+					ParentPtr->ReloadAllDocTemplateTabs();
 					ParentPtr->UpdateEnabledTabs();
 				}
 			}
@@ -5874,7 +5877,7 @@ PRAGMA_ENABLE_DEPRECATION_WARNINGS
 			UAudioWidgetSubsystem& WidgetSubsystem,
 			TSubclassOf<UInterface> InterfaceClass,
 			TFunction<bool(UUserWidget*)>Filter,
-			TFunction<FMetaSoundDocumentConfigurationWidgetTabInfo(UUserWidget*)>OnTabConstruct,
+			TFunction<FMetaSoundTemplateWidgetTabInfo(UUserWidget*)>OnTabConstruct,
 			FGuid TargetBlueprintId)
 		{
 			TFunction<bool(UBlueprint*)> BlueprintFilter;
@@ -5900,7 +5903,7 @@ PRAGMA_ENABLE_DEPRECATION_WARNINGS
 					// existing handle if found. Otherwise, assign a new delegate handler.
 					if (TargetBlueprintId.IsValid())
 					{
-						if (FUserWidgetData* WidgetData = DocConfigWidgetData.Find(TargetBlueprintId))
+						if (FUserWidgetData* WidgetData = DocTemplateWidgetData.Find(TargetBlueprintId))
 						{
 							RecompileHandle = MoveTemp(WidgetData->RecompileHandle);
 						}
@@ -5910,15 +5913,15 @@ PRAGMA_ENABLE_DEPRECATION_WARNINGS
 						OwningBlueprint->OnCompiled().RemoveAll(this);
 					}
 
-					DocConfigWidgetData.Remove(OwningBPId);
+					DocTemplateWidgetData.Remove(OwningBPId);
 
 					if (!RecompileHandle.IsValid())
 					{
-						RecompileHandle = OwningBlueprint->OnCompiled().AddSP(this, &FEditor::OnDocConfigWidgetBlueprintCompiled);
+						RecompileHandle = OwningBlueprint->OnCompiled().AddSP(this, &FEditor::OnDocTemplateWidgetBlueprintCompiled);
 					}
 
 					TStrongObjectPtr<UUserWidget> UserWidgetPtr(UserWidget);
-					FUserWidgetData& NewData = DocConfigWidgetData.Add(OwningBPId, FUserWidgetData
+					FUserWidgetData& NewData = DocTemplateWidgetData.Add(OwningBPId, FUserWidgetData
 					{
 						.Class = InterfaceClass,
 						.RecompileHandle = MoveTemp(RecompileHandle),
@@ -5929,7 +5932,7 @@ PRAGMA_ENABLE_DEPRECATION_WARNINGS
 					FMetaSoundFrontendDocumentBuilder* DocBuilder = GetFrontendBuilder();
 					checkf(DocBuilder, TEXT("Builder must exist by this point in user widget rebuild"));
 
-					const UScriptStruct* ConfigStruct = DocBuilder->GetConstDocumentConfiguration().GetScriptStruct();
+					const UScriptStruct* ConfigStruct = DocBuilder->GetConstDocumentTemplate().GetScriptStruct();
 					checkf(ConfigStruct, TEXT("Builder must have a config set by this point in user widget rebuild"));
 
 					AssetEditorPrivate::SanitizeTabInfo(*ConfigStruct, NewData.TabInfo);
diff --git a/Engine/Plugins/Runtime/Metasound/Source/MetasoundEditor/Private/MetasoundEditor.h b/Engine/Plugins/Runtime/Metasound/Source/MetasoundEditor/Private/MetasoundEditor.h
--- a/Engine/Plugins/Runtime/Metasound/Source/MetasoundEditor/Private/MetasoundEditor.h
+++ b/Engine/Plugins/Runtime/Metasound/Source/MetasoundEditor/Private/MetasoundEditor.h
@@ -9,12 +9,12 @@
 #include "IMetasoundEditor.h"
 #include "Logging/TokenizedMessage.h"
 #include "Math/UnrealMathUtility.h"
-#include "MetasoundDocumentConfigurationWidgetInterface.h"
 #include "MetasoundEditorGraphConnectionManager.h"
 #include "MetasoundEditorGraphNode.h"
 #include "MetasoundFrontend.h"
 #include "MetasoundFrontendController.h"
 #include "MetasoundFrontendDocumentModifyDelegates.h"
+#include "MetasoundTemplateWidgetInterface.h"
 #include "Misc/NotifyHook.h"
 #include "SGraphActionMenu.h"
 #include "SGraphPanel.h"
@@ -140,7 +140,7 @@ namespace Metasound::Editor
 
 		UAudioComponent* GetAudioComponent() const;
 
-		const Frontend::FDocumentConfiguration* GetConfiguration() const;
+		const Frontend::FMetaSoundTemplate* GetDocumentTemplate() const;
 		FMetaSoundFrontendDocumentBuilder* GetFrontendBuilder() const;
 
 		/** IMetasoundEditor interface */
@@ -176,8 +176,8 @@ namespace Metasound::Editor
 		virtual void Play() override;
 		virtual void Stop() override;
 
-		/** Returns whether or not doc config visibility settings should be ignored (for debugging/development) */
-		static bool IgnoreDocConfigVisibilitySettings();
+		/** Returns whether or not template visibility settings should be ignored (for debugging/development) */
+		static bool IgnoreDocTemplateVisibilitySettings();
 
 		/** Returns whether or not live audition is enabled for previewing */
 		static bool IsLiveAuditionEnabled();
@@ -267,8 +267,8 @@ namespace Metasound::Editor
 
 		void ClearDroppedObjectsInfo();
 
-		/** Reloads all document configuration tabs */
-		void ReloadAllDocConfigTabs();
+		/** Reloads all document template tabs */
+		void ReloadAllDocTemplateTabs();
 
 	protected:
 		// Callbacks for action tree
@@ -371,7 +371,7 @@ namespace Metasound::Editor
 		bool CanPasteSelectedMemberItems() const;
 
 		/** Copies the selected nodes, shared connections, and required interfaces to the clipboard as configurator c++ */
-		void CopySelectionAsConfigurationCodeToClipboard() const;
+		void CopySelectionAsConfiguratorCodeToClipboard() const;
 
 		/** Paste selected Member items. */
 		void PasteSelectedMemberItems();
@@ -453,7 +453,7 @@ namespace Metasound::Editor
 			UAudioWidgetSubsystem& WidgetSubsystem,
 			TSubclassOf<UInterface> InterfaceClass,
 			TFunction<bool(UUserWidget*)>Filter,
-			TFunction<FMetaSoundDocumentConfigurationWidgetTabInfo(UUserWidget*)>OnTabConstruct,
+			TFunction<FMetaSoundTemplateWidgetTabInfo(UUserWidget*)>OnTabConstruct,
 			FGuid TargetBlueprintId = FGuid());
 
 		void RefreshExecVisibility(const FGuid& InPageID) const;
@@ -480,14 +480,14 @@ namespace Metasound::Editor
 		/** Creates all internal widgets for the tabs to point at */
 		void CreateInternalWidgets(UObject& MetaSound);
 
-		void RebuildDocConfigUserWidgets(const FGuid& TargetBlueprintId = FGuid());
+		void RebuildDocTemplateUserWidgets(const FGuid& TargetBlueprintId = FGuid());
 
-		/** For teardown and regeneration of user document configuration widget on world change
+		/** For teardown and regeneration of user document template widget on world change
 			Based on UEditorUtilityWidgetBlueprint::ChangeTabWorld */
-		void ChangeDocConfigWidgetWorld(UWorld* World, EMapChangeType MapChangeType);
+		void ChangeDocTemplateWidgetWorld(UWorld* World, EMapChangeType MapChangeType);
 
-		/** Reload document configuration widget tab on associated widget blueprint recompile */
-		void OnDocConfigWidgetBlueprintCompiled(UBlueprint* Blueprint);
+		/** Reload document template widget tab on associated widget blueprint recompile */
+		void OnDocTemplateWidgetBlueprintCompiled(UBlueprint* Blueprint);
 
 		/** Builds the toolbar widget for the Metasound editor */
 		void ExtendToolbarInternal();
@@ -512,7 +512,7 @@ namespace Metasound::Editor
 		/** Either play the Metasound or stop currently playing sound */
 		void TogglePlayback();
 
-		/** Whether or not to show the toggle for showing/hiding a graph (i.e. is or is not a MetaSound with a configuration set) */
+		/** Whether or not to show the toggle for showing/hiding a graph (i.e. is or is not a MetaSound with a template set) */
 		bool CanToggleShowHideGraph() const;
 
 		/** Toggles showing the graph for the given configured MetaSound */
@@ -533,14 +533,14 @@ namespace Metasound::Editor
 		/** Creates audition menu options */
 		TSharedRef<SWidget> CreateAuditionMenuOptions();
 
-		/** Creates configuration menu options */
-		TSharedRef<SWidget> CreateConfigurationMenuOptions();
+		/** Creates template menu options */
+		TSharedRef<SWidget> CreateTemplateMenuOptions();
 
 		/** Creates page menu options */
 		void CreateAuditionPageSubMenuOptions(FMenuBuilder& MenuBuilder);
 
-		/** Creates configuration menu options */
-		void CreateConfigurationSubMenuOptions(FMenuBuilder& MenuBuilder);
+		/** Creates template menu options */
+		void CreateTemplateSubMenuOptions(FMenuBuilder& MenuBuilder);
 
 		/** Show the Metasound object's Source settings in the Details panel */
 		void EditSourceSettings();
@@ -619,22 +619,22 @@ namespace Metasound::Editor
 			TSubclassOf<UInterface> Class;
 			FDelegateHandle RecompileHandle;
 			TStrongObjectPtr<UUserWidget> Widget;
-			FMetaSoundDocumentConfigurationWidgetTabInfo TabInfo;
+			FMetaSoundTemplateWidgetTabInfo TabInfo;
 
 			void SetAuditionState(UAudioComponent* AudioComponent, bool bIsAuditioning) const;
 		};
 		using FBlueprintId = FGuid;
 		using FBPIdWidgetDataPair = TPair<FBlueprintId, FUserWidgetData>;
 
-		/** User defined widget data displayed when editing a doc config with user widgets enabled.*/
-		TMap<FBlueprintId, FUserWidgetData> DocConfigWidgetData;
+		/** User defined widget data displayed when editing a document template with user widgets enabled.*/
+		TMap<FBlueprintId, FUserWidgetData> DocTemplateWidgetData;
 
-		FWorkflowAllowedTabSet DocConfigTabSet;
+		FWorkflowAllowedTabSet DocTemplateTabSet;
 
-		struct FDocConfigWidgetSummoner : public FWorkflowTabFactory
+		struct FTemplateWidgetSummoner : public FWorkflowTabFactory
 		{
 		public:
-			FDocConfigWidgetSummoner(FName FactoryName, FBlueprintId InBlueprintId, TSharedPtr<class FAssetEditorToolkit> InHostingApp);
+			FTemplateWidgetSummoner(FName FactoryName, FBlueprintId InBlueprintId, TSharedPtr<class FAssetEditorToolkit> InHostingApp);
 
 			virtual TSharedRef<SWidget> CreateTabBody(const FWorkflowTabSpawnInfo& Info) const override;
 			virtual TSharedRef<SDockTab> SpawnTab(const FWorkflowTabSpawnInfo& Info) const override;
@@ -691,7 +691,7 @@ namespace Metasound::Editor
 
 		private:
 			virtual void OnBuilderReloaded(Frontend::FDocumentModifyDelegates& OutDelegates) override;
-			void OnDocumentConfigurationChanged(const Frontend::FDocumentConfigChangedArgs& Args);
+			void OnDocumentTemplateChanged(const Frontend::FDocumentTemplateChangedArgs& Args);
 			void OnPageSet(const Frontend::FDocumentMutatePageArgs& Args);
 			void OnInputDefaultChanged(const Frontend::FDocumentArrayPagedInputArgs& Args);
 			void OnMemberMetadataSet(FGuid MemberID);
diff --git a/Engine/Plugins/Runtime/Metasound/Source/MetasoundEditor/Private/MetasoundEditorBuilderListener.cpp b/Engine/Plugins/Runtime/Metasound/Source/MetasoundEditor/Private/MetasoundEditorBuilderListener.cpp
--- a/Engine/Plugins/Runtime/Metasound/Source/MetasoundEditor/Private/MetasoundEditorBuilderListener.cpp
+++ b/Engine/Plugins/Runtime/Metasound/Source/MetasoundEditor/Private/MetasoundEditorBuilderListener.cpp
@@ -2,7 +2,7 @@
 
 #include "MetasoundEditorBuilderListener.h"
 
-#include "DocumentConfigurations/MetasoundFrontendPresetDocumentConfiguration.h"
+#include "DocumentTemplates/MetasoundFrontendPresetTemplate.h"
 
 #include UE_INLINE_GENERATED_CPP_BY_NAME(MetasoundEditorBuilderListener)
 
diff --git a/Engine/Plugins/Runtime/Metasound/Source/MetasoundEditor/Private/MetasoundEditorCommands.cpp b/Engine/Plugins/Runtime/Metasound/Source/MetasoundEditor/Private/MetasoundEditorCommands.cpp
--- a/Engine/Plugins/Runtime/Metasound/Source/MetasoundEditor/Private/MetasoundEditorCommands.cpp
+++ b/Engine/Plugins/Runtime/Metasound/Source/MetasoundEditor/Private/MetasoundEditorCommands.cpp
@@ -37,7 +37,7 @@ namespace Metasound
 			UI_COMMAND(PromoteAllToInput, "Promote All To Graph Input", "Promotes unconnected node inputs to graph inputs.", EUserInterfaceActionType::Button, FInputChord());
 			UI_COMMAND(PromoteAllToCommonInputs, "Promote All To Common Graph Inputs", "Promotes unconnected input pins from selected nodes to graph inputs, sharing inputs if possible.", EUserInterfaceActionType::Button, FInputChord());
 
-			UI_COMMAND(CopySelectionAsConfigurationCodeToClipboard,
+			UI_COMMAND(CopySelectionAsConfiguratorCodeToClipboard,
 				"Copy Selected As Configurator Code to Clipboard",
 				"Reflect configurator C++ code to the clipboard for the given graph selection",
 				EUserInterfaceActionType::Button,
diff --git a/Engine/Plugins/Runtime/Metasound/Source/MetasoundEditor/Private/MetasoundEditorGraphMemberDefaults.cpp b/Engine/Plugins/Runtime/Metasound/Source/MetasoundEditor/Private/MetasoundEditorGraphMemberDefaults.cpp
--- a/Engine/Plugins/Runtime/Metasound/Source/MetasoundEditor/Private/MetasoundEditorGraphMemberDefaults.cpp
+++ b/Engine/Plugins/Runtime/Metasound/Source/MetasoundEditor/Private/MetasoundEditorGraphMemberDefaults.cpp
@@ -6,7 +6,6 @@
 #include "Algo/Transform.h"
 #include "AudioDefines.h"
 #include "AudioParameterControllerInterface.h"
-#include "DocumentConfigurations/MetasoundFrontendPresetDocumentConfiguration.h"
 #include "EdGraph/EdGraphNode.h"
 #include "GraphEditorSettings.h"
 #include "MetasoundDataReference.h"
@@ -26,6 +25,7 @@
 #include "MetasoundSettings.h"
 #include "Misc/Guid.h"
 #include "Sound/SoundWave.h"
+#include "DocumentTemplates/MetasoundFrontendPresetTemplate.h"
 #include "UObject/ObjectMacros.h"
 #include "UObject/SoftObjectPath.h"
 
@@ -631,8 +631,8 @@ bool UMetasoundEditorGraphMemberDefaultFloat::MigrateProperties()
 	bool bIsPreset = Document.RootGraph.PresetOptions.bIsPreset;
 	if (!bIsPreset)
 	{
-		const UStruct* Struct = Document.Configuration.GetScriptStruct();
-		bIsPreset = Struct && Struct->IsChildOf(FPresetDocumentConfiguration::StaticStruct());
+		const UStruct* Struct = Document.Template.GetScriptStruct();
+		bIsPreset = Struct && Struct->IsChildOf(FPresetTemplate::StaticStruct());
 	}
 	// Set default input widget type from settings. Widgets do not apply to presets
 	if (bIsPreset)
diff --git a/Engine/Plugins/Runtime/Metasound/Source/MetasoundEditor/Private/MetasoundEditorGraphSchema.cpp b/Engine/Plugins/Runtime/Metasound/Source/MetasoundEditor/Private/MetasoundEditorGraphSchema.cpp
--- a/Engine/Plugins/Runtime/Metasound/Source/MetasoundEditor/Private/MetasoundEditorGraphSchema.cpp
+++ b/Engine/Plugins/Runtime/Metasound/Source/MetasoundEditor/Private/MetasoundEditorGraphSchema.cpp
@@ -2112,7 +2112,7 @@ void UMetasoundEditorGraphSchema::GetContextMenuActions(class UToolMenu* Menu, c
 	if (Context->Graph)
 	{
 		FToolMenuSection& Section = Menu->AddSection("MetasoundGraphSchemaGraphActions", LOCTEXT("DevActionsMenuHeader", "Dev Actions"));
-		Section.AddMenuEntry(FEditorCommands::Get().CopySelectionAsConfigurationCodeToClipboard);
+		Section.AddMenuEntry(FEditorCommands::Get().CopySelectionAsConfiguratorCodeToClipboard);
 	}
 }
 
diff --git a/Engine/Plugins/Runtime/Metasound/Source/MetasoundEditor/Private/MetasoundEditorSubsystem.cpp b/Engine/Plugins/Runtime/Metasound/Source/MetasoundEditor/Private/MetasoundEditorSubsystem.cpp
--- a/Engine/Plugins/Runtime/Metasound/Source/MetasoundEditor/Private/MetasoundEditorSubsystem.cpp
+++ b/Engine/Plugins/Runtime/Metasound/Source/MetasoundEditor/Private/MetasoundEditorSubsystem.cpp
@@ -2,8 +2,6 @@
 #include "MetasoundEditorSubsystem.h"
 
 #include "AudioPropertiesSheetAssetBase.h"
-#include "DocumentConfigurations/MetasoundFrontendDocumentConfiguration.h"
-#include "DocumentConfigurations/MetasoundFrontendPresetDocumentConfiguration.h"
 #include "FileHelpers.h"
 #include "IAssetTools.h"
 #include "ISourceControlModule.h"
@@ -25,6 +23,8 @@
 #include "Sound/SoundSourceBusSend.h"
 #include "Sound/SoundSubmixSend.h"
 #include "Subsystems/AssetEditorSubsystem.h"
+#include "DocumentTemplates/MetasoundFrontendPresetTemplate.h"
+#include "DocumentTemplates/MetasoundFrontendDocumentTemplate.h"
 
 #include UE_INLINE_GENERATED_CPP_BY_NAME(MetasoundEditorSubsystem)
 
@@ -61,9 +61,9 @@ TScriptInterface<IMetaSoundDocumentInterface> UMetaSoundEditorSubsystem::BuildTo
 		
 		// Duplicate referenced preset object to preserve object properties (ex. quality settings, soundwave properties)
 		UObject* NewMetaSound = nullptr;
-		if (const FPresetDocumentConfiguration* PresetConfig = InBuilder->GetConstBuilder().GetConstDocumentConfigurationAs<FPresetDocumentConfiguration>())
+		if (const FPresetTemplate* PresetTemplate = InBuilder->GetConstBuilder().GetConstTemplateAs<FPresetTemplate>())
 		{
-			NewMetaSound = IAssetTools::Get().DuplicateAsset(AssetName, PackagePath, PresetConfig->Parent.GetObject());
+			NewMetaSound = IAssetTools::Get().DuplicateAsset(AssetName, PackagePath, PresetTemplate->Parent.GetObject());
 		}
 		else
 		{
@@ -325,11 +325,11 @@ void UMetaSoundEditorSubsystem::InitAsset(UObject& InNewMetaSound, UObject* InPa
 {
 	using namespace Metasound::Editor;
 
-	TInstancedStruct<FMetaSoundFrontendDocumentConfiguration> PresetConfig = TInstancedStruct<FMetaSoundFrontendDocumentConfiguration>::Make<FMetaSoundFrontendPresetDocumentConfiguration>();
-	PresetConfig.GetMutable<FMetaSoundFrontendPresetDocumentConfiguration>().Parent = InParentMetaSound;
+	TInstancedStruct<FMetaSoundFrontendDocumentTemplate> PresetTemplate = TInstancedStruct<FMetaSoundFrontendDocumentTemplate>::Make<FMetaSoundFrontendPresetTemplate>();
+	PresetTemplate.GetMutable<FMetaSoundFrontendPresetTemplate>().Parent = InParentMetaSound;
 	InitAsset(InNewMetaSound, FInitAssetArgs
 	{
-		.DocumentConfiguration = MoveTemp(PresetConfig),
+		.Template = MoveTemp(PresetTemplate),
 		.SelectedObjects = { InParentMetaSound }
 	});
 
@@ -378,11 +378,11 @@ void UMetaSoundEditorSubsystem::InitAsset(UObject& InNewMetaSound, UMetaSoundEdi
 	}
 
 	// Must be called in new context after asset is initialized as blank document as configurations use a registered builder
-	if (FMetaSoundFrontendDocumentConfiguration* DocConfig = InitAssetArgs.DocumentConfiguration.GetMutablePtr())
+	if (FMetaSoundFrontendDocumentTemplate* DocTemplate = InitAssetArgs.Template.GetMutablePtr())
 	{
 		FMetaSoundFrontendDocumentBuilder& DocBuilder = Frontend::IDocumentBuilderRegistry::GetChecked().FindOrBeginBuilding(&InNewMetaSound);
-		DocBuilder.SetDocumentConfiguration(MoveTemp(InitAssetArgs.DocumentConfiguration));
-		DocConfig->OnAssetInitialized(MoveTemp(InitAssetArgs.SelectedObjects), DocBuilder);
+		DocBuilder.SetTemplate(MoveTemp(InitAssetArgs.Template));
+		DocTemplate->OnAssetInitialized(MoveTemp(InitAssetArgs.SelectedObjects), DocBuilder);
 		DocBuilder.ConfigureDocument();
 	}
 }
diff --git a/Engine/Plugins/Runtime/Metasound/Source/MetasoundEditor/Private/MetasoundFactory.cpp b/Engine/Plugins/Runtime/Metasound/Source/MetasoundEditor/Private/MetasoundFactory.cpp
--- a/Engine/Plugins/Runtime/Metasound/Source/MetasoundEditor/Private/MetasoundFactory.cpp
+++ b/Engine/Plugins/Runtime/Metasound/Source/MetasoundEditor/Private/MetasoundFactory.cpp
@@ -1,8 +1,8 @@
 // Copyright Epic Games, Inc. All Rights Reserved.
 #include "MetasoundFactory.h"
 
-#include "DocumentConfigurations/MetasoundFrontendDocumentConfiguration.h"
-#include "DocumentConfigurations/MetasoundFrontendPresetDocumentConfiguration.h"
+#include "DocumentTemplates/MetasoundFrontendPresetTemplate.h"
+#include "DocumentTemplates/MetasoundFrontendDocumentTemplate.h"
 #include "Metasound.h"
 #include "MetasoundEditorGraph.h"
 #include "MetasoundEditorGraphBuilder.h"
@@ -40,7 +40,7 @@ UObject* UMetaSoundFactory::FactoryCreateNew(UClass* InClass, UObject* InParent,
 
 	UMetaSoundEditorSubsystem::GetChecked().InitAsset(*NewPatch, UMetaSoundEditorSubsystem::FInitAssetArgs
 	{
-		.DocumentConfiguration = DocumentConfiguration,
+		.Template = Template,
 		.SelectedObjects = SelectedObjects
 	});
 
@@ -63,7 +63,7 @@ UObject* UMetaSoundSourceFactory::FactoryCreateNew(UClass* InClass, UObject* InP
 
 	UMetaSoundEditorSubsystem::GetChecked().InitAsset(*NewSource, UMetaSoundEditorSubsystem::FInitAssetArgs
 	{
-		.DocumentConfiguration = DocumentConfiguration,
+		.Template = Template,
 		.SelectedObjects = SelectedObjects
 	});
 
@@ -71,9 +71,9 @@ UObject* UMetaSoundSourceFactory::FactoryCreateNew(UClass* InClass, UObject* InP
 	const IMetasoundEditorModule& MetaSoundEditorModule = FModuleManager::GetModuleChecked<IMetasoundEditorModule>("MetaSoundEditor");
 	if (MetaSoundEditorModule.IsRestrictedMode())
 	{
-		if (const FMetaSoundFrontendPresetDocumentConfiguration* PresetConfig = DocumentConfiguration.GetPtr<FMetaSoundFrontendPresetDocumentConfiguration>())
+		if (const FMetaSoundFrontendPresetTemplate* PresetTemplate = Template.GetPtr<FMetaSoundFrontendPresetTemplate>())
 		{
-			if (const USoundWave* ParentWave = Cast<const USoundWave>(PresetConfig->Parent.GetObject()))
+			if (const USoundWave* ParentWave = Cast<const USoundWave>(PresetTemplate->Parent.GetObject()))
 			{
 				UMetaSoundEditorSubsystem::GetChecked().SetSoundWaveSettingsFromTemplate(*NewSource, *ParentWave);
 			}
diff --git a/Engine/Plugins/Runtime/Metasound/Source/MetasoundEditor/Private/MetasoundNodeDetailCustomization.cpp b/Engine/Plugins/Runtime/Metasound/Source/MetasoundEditor/Private/MetasoundNodeDetailCustomization.cpp
--- a/Engine/Plugins/Runtime/Metasound/Source/MetasoundEditor/Private/MetasoundNodeDetailCustomization.cpp
+++ b/Engine/Plugins/Runtime/Metasound/Source/MetasoundEditor/Private/MetasoundNodeDetailCustomization.cpp
@@ -5,7 +5,6 @@
 #include "Components/AudioComponent.h"
 #include "Containers/Set.h"
 #include "Delegates/Delegate.h"
-#include "DocumentConfigurations/MetasoundFrontendPresetDocumentConfiguration.h"
 #include "Framework/Notifications/NotificationManager.h"
 #include "HAL/PlatformApplicationMisc.h"
 #include "IDetailChildrenBuilder.h"
@@ -33,6 +32,7 @@
 #include "Widgets/Input/SSearchableComboBox.h"
 #include "Styling/SlateColor.h"
 #include "Templates/Casts.h"
+#include "DocumentTemplates/MetasoundFrontendPresetTemplate.h"
 #include "Templates/SharedPointer.h"
 #include "Templates/UniquePtr.h"
 #include "UObject/WeakObjectPtr.h"
@@ -271,9 +271,9 @@ namespace Metasound
 								return false;
 							}
 
-							if (const FDocumentVertexConfiguration* DocConfig = Input->GetFrontendBuilderChecked().GetConstDocumentConfigurationAs<FDocumentVertexConfiguration>())
+							if (const FDocumentVertexTemplate* DocTemplate = Input->GetFrontendBuilderChecked().GetConstTemplateAs<FDocumentVertexTemplate>())
 							{
-								if (const FDocumentConfigurationVertexMetadata* VertexMetadata = DocConfig->FindConstVertexMetadata(Input->NodeID))
+								if (const FTemplateVertexMetadata* VertexMetadata = DocTemplate->FindConstVertexMetadata(Input->NodeID))
 								{
 									return VertexMetadata->IsDefaultEditable();
 								}
@@ -1421,7 +1421,7 @@ PRAGMA_ENABLE_DEPRECATION_WARNINGS
 			return FText::GetEmpty();
 		}
 
-		const Frontend::FDocumentConfigurationVertexMetadata* FMetasoundVertexDetailCustomization::GetConstConfigMetadata() const
+		const Frontend::FTemplateVertexMetadata* FMetasoundVertexDetailCustomization::GetConstTemplateMetadata() const
 		{
 			using namespace Frontend;
 
@@ -1429,14 +1429,14 @@ PRAGMA_ENABLE_DEPRECATION_WARNINGS
 			{
 				if (const FMetaSoundFrontendDocumentBuilder* DocBuilder = GetBuilder())
 				{
-					return DocBuilder->FindConstVertexMetadata<FDocumentConfigurationVertexMetadata>(MemberPtr->GetMemberID());
+					return DocBuilder->FindConstVertexMetadata<FTemplateVertexMetadata>(MemberPtr->GetMemberID());
 				}
 			}
 
 			return nullptr;
 		}
 
-		Frontend::FDocumentConfigurationVertexMetadata* FMetasoundVertexDetailCustomization::GetConfigMetadata() const
+		Frontend::FTemplateVertexMetadata* FMetasoundVertexDetailCustomization::GetTemplateMetadata() const
 		{
 			using namespace Frontend;
 
@@ -1444,7 +1444,7 @@ PRAGMA_ENABLE_DEPRECATION_WARNINGS
 			{
 				if (const FMetaSoundFrontendDocumentBuilder* DocBuilder = GetBuilder())
 				{
-					return DocBuilder->FindVertexMetadata<FDocumentConfigurationVertexMetadata>(MemberPtr->GetMemberID());
+					return DocBuilder->FindVertexMetadata<FTemplateVertexMetadata>(MemberPtr->GetMemberID());
 				}
 			}
 
@@ -1751,24 +1751,24 @@ PRAGMA_ENABLE_DEPRECATION_WARNINGS
 				return;
 			}
 
-			TStructView<FMetaSoundFrontendDocumentConfiguration> ConfigStruct = DocBuilder->GetDocumentConfiguration();
-			if (FDocumentVertexConfiguration* Config = ConfigStruct.GetPtr<FDocumentVertexConfiguration>())
+			TStructView<FMetaSoundFrontendDocumentTemplate> TemplateStruct = DocBuilder->GetDocumentTemplate();
+			if (FDocumentVertexTemplate* Template = TemplateStruct.GetPtr<FDocumentVertexTemplate>())
 			{
 				const FGuid& MemberID = GraphMember->GetMemberID();
-				if (FDocumentConfigurationVertexMetadata* Metadata = Config->FindVertexMetadata(MemberID))
+				if (FTemplateVertexMetadata* Metadata = Template->FindVertexMetadata(MemberID))
 				{
-					const UScriptStruct* StructClass = Config->GetVertexMetadataStruct(MemberID);
+					const UScriptStruct* StructClass = Template->GetVertexMetadataStruct(MemberID);
 					if (const TSharedPtr<FStructOnScope> DataView = MakeShared<FStructOnScope>(StructClass, reinterpret_cast<uint8*>(Metadata)))
 					{
 						TArray<TSharedPtr<IPropertyHandle>> StructPropHandles = GetDefaultCategoryBuilder(InDetailLayout).AddAllExternalStructureProperties(DataView.ToSharedRef());
 
-						const auto OnConfigPreChange = TDelegate<void()>::CreateSPLambda(AsShared(), [this]()
+						const auto OnTemplatePreChange = TDelegate<void()>::CreateSPLambda(AsShared(), [this]()
 						{
 							if (GEditor && GraphMember.IsValid())
 							{
 								if (UMetasoundEditorGraphMemberDefaultLiteral* MemberDefaultLiteral = GraphMember->GetLiteral())
 								{
-									GEditor->BeginTransaction(LOCTEXT("GraphMemberDocConfig_Transact", "Set GraphMember Config Vertex Metadata"));
+									GEditor->BeginTransaction(LOCTEXT("GraphMemberDocTemplate_Transact", "Set GraphMember Template Vertex Metadata"));
 									GraphMember->GetOutermost()->Modify();
 									GraphMember->GetOutermostObject()->Modify();
 									GraphMember->Modify();
@@ -1777,9 +1777,9 @@ PRAGMA_ENABLE_DEPRECATION_WARNINGS
 							}
 						});
 
-						const auto OnConfigPropChanged = TDelegate<void(const FPropertyChangedEvent&)>::CreateSPLambda(AsShared(), [this](const FPropertyChangedEvent& InEvent)
+						const auto OnTemplatePropChanged = TDelegate<void(const FPropertyChangedEvent&)>::CreateSPLambda(AsShared(), [this](const FPropertyChangedEvent& InEvent)
 						{
-							if (FDocumentConfigurationVertexMetadata* ThisMetadata = GetConfigMetadata())
+							if (FTemplateVertexMetadata* ThisMetadata = GetTemplateMetadata())
 							{
 								if (FMetaSoundFrontendDocumentBuilder* DocBuilder = GetBuilder())
 								{
@@ -1812,8 +1812,8 @@ PRAGMA_ENABLE_DEPRECATION_WARNINGS
 						for (TSharedPtr<IPropertyHandle> StructPropHandle : StructPropHandles)
 						{
 // 							Row->ShouldAutoExpand(true);
-							StructPropHandle->SetOnPropertyValuePreChange(OnConfigPreChange);
-							StructPropHandle->SetOnPropertyValueChangedWithData(OnConfigPropChanged);
+							StructPropHandle->SetOnPropertyValuePreChange(OnTemplatePreChange);
+							StructPropHandle->SetOnPropertyValueChangedWithData(OnTemplatePropChanged);
 						}
 					}
 				}
@@ -1822,7 +1822,7 @@ PRAGMA_ENABLE_DEPRECATION_WARNINGS
 			Enabled = TAttribute<bool>::CreateSPLambda(AsShared(), [this]
 			{
 				using namespace Frontend;
-				if (const FDocumentConfigurationVertexMetadata* VertexMetadata = GetConstConfigMetadata())
+				if (const FTemplateVertexMetadata* VertexMetadata = GetConstTemplateMetadata())
 				{
 					return VertexMetadata->IsDefaultEditable();
 				}
diff --git a/Engine/Plugins/Runtime/Metasound/Source/MetasoundEditor/Private/MetasoundNodeDetailCustomization.h b/Engine/Plugins/Runtime/Metasound/Source/MetasoundEditor/Private/MetasoundNodeDetailCustomization.h
--- a/Engine/Plugins/Runtime/Metasound/Source/MetasoundEditor/Private/MetasoundNodeDetailCustomization.h
+++ b/Engine/Plugins/Runtime/Metasound/Source/MetasoundEditor/Private/MetasoundNodeDetailCustomization.h
@@ -6,7 +6,7 @@
 #include "DetailCategoryBuilder.h"
 #include "DetailLayoutBuilder.h"
 #include "DetailWidgetRow.h"
-#include "DocumentConfigurations/MetasoundFrontendDocumentConfiguration.h"
+#include "DocumentTemplates/MetasoundFrontendDocumentTemplate.h"
 #include "EdGraph/EdGraphSchema.h"
 #include "Framework/Application/SlateApplication.h"
 #include "IDetailCustomization.h"
@@ -401,8 +401,8 @@ namespace Metasound
 			ECheckBoxState OnGetConstructorPinCheckboxState(TWeakObjectPtr<UMetasoundEditorGraphVertex> InGraphMember) const;
 			void OnConstructorPinStateChanged(TWeakObjectPtr<UMetasoundEditorGraphVertex> InGraphMember, ECheckBoxState InNewState);
 
-			const Frontend::FDocumentConfigurationVertexMetadata* GetConstConfigMetadata() const;
-			Frontend::FDocumentConfigurationVertexMetadata* GetConfigMetadata() const;
+			const Frontend::FTemplateVertexMetadata* GetConstTemplateMetadata() const;
+			Frontend::FTemplateVertexMetadata* GetTemplateMetadata() const;
 
 			TSharedPtr<SCheckBox> ConstructorPinCheckbox;
 
diff --git a/Engine/Plugins/Runtime/Metasound/Source/MetasoundEditor/Private/TemplateViewModels/MetasoundTemplateViewModelBase.cpp b/Engine/Plugins/Runtime/Metasound/Source/MetasoundEditor/Private/TemplateViewModels/MetasoundTemplateViewModelBase.cpp
new file mode 100644
--- /dev/null
+++ b/Engine/Plugins/Runtime/Metasound/Source/MetasoundEditor/Private/TemplateViewModels/MetasoundTemplateViewModelBase.cpp
@@ -0,0 +1,98 @@
+// Copyright Epic Games, Inc. All Rights Reserved.
+
+#include "TemplateViewModels/MetasoundTemplateViewModelBase.h"
+#include "MetasoundFrontendDocumentModifyDelegates.h"
+#include "View/MVVMView.h"
+
+
+void UMetaSoundTemplateViewModelBase::ConstructWidgetViewModel(
+	UUserWidget* UserWidget,
+	UMetaSoundBuilderBase& InBuilder,
+	TConstStructView<FMetaSoundFrontendDocumentTemplate> Config,
+	TSubclassOf<UMetaSoundTemplateViewModelBase> VMClass)
+{
+	using namespace Metasound::Engine;
+
+	if (!UserWidget)
+	{
+		return;
+	}
+
+	UMVVMView* View = UserWidget->GetExtension<UMVVMView>();
+	const UScriptStruct* ConfigStruct = Config.GetScriptStruct();
+	
+	if (View && ConfigStruct && VMClass.Get() != nullptr)
+	{
+		if (UMetaSoundTemplateViewModelBase* ViewModel = NewObject<UMetaSoundTemplateViewModelBase>(GetTransientPackage(), VMClass))
+		{
+			if (ViewModel->IsSupportedTemplate(*ConfigStruct))
+			{
+				ViewModel->Initialize(&InBuilder);
+				View->SetViewModelByClass(ViewModel);
+			}
+		}
+	}
+}
+
+bool UMetaSoundTemplateViewModelBase::IsSupportedTemplate(const UScriptStruct& InStruct) const
+{
+	return true;
+}
+
+void UMetaSoundTemplateViewModelBase::Initialize(UMetaSoundBuilderBase* InBuilder)
+{
+	using namespace Metasound::Frontend;
+	Builder = InBuilder;
+
+	BuilderListener = MakeShared<FBuilderListener>(this);
+	Builder->AddTransactionListener(BuilderListener->AsShared());
+
+	Reload();
+}
+
+void UMetaSoundTemplateViewModelBase::Reload()
+{
+}
+
+void UMetaSoundTemplateViewModelBase::FBuilderListener::OnBuilderReloaded(Metasound::Frontend::FDocumentModifyDelegates& OutDelegates)
+{
+	if (TStrongObjectPtr<UMetaSoundTemplateViewModelBase> PinnedParent = Parent.Pin())
+	{
+		PinnedParent->Reload();
+	}
+
+	OutDelegates.OnDocumentTemplateChanged.Remove(TemplateChangeDelegateHandle);
+	TemplateChangeDelegateHandle = OutDelegates.OnDocumentTemplateChanged.AddSP(this, &UMetaSoundTemplateViewModelBase::FBuilderListener::OnTemplateChanged);
+}
+
+void UMetaSoundTemplateViewModelBase::FBuilderListener::OnTemplateChanged(const Metasound::Frontend::FDocumentTemplateChangedArgs& InArgs)
+{
+	using namespace Metasound::Frontend;
+
+	// Property change covered by template's direct property changed delegate
+	// Covers if the whole struct is replaced, in which case just re-initialize
+	// or if editor transaction stack has been rolled forward or back.
+	if (InArgs.Type != Metasound::Frontend::EDocumentTemplateChangeType::Property)
+	{
+		if (TStrongObjectPtr<UMetaSoundTemplateViewModelBase> PinnedParent = Parent.Pin())
+		{
+			if (UMetaSoundBuilderBase* ThisBuilder = PinnedParent->Builder.Get())
+			{
+				FMetaSoundFrontendDocumentBuilder& DocBuilder = ThisBuilder->GetBuilder();
+				TConstStructView<FMetaSoundTemplate> DocTemplate = DocBuilder.GetConstDocumentTemplate();
+				if (const UScriptStruct* ScriptStruct = DocTemplate.GetScriptStruct();
+					ScriptStruct && PinnedParent->IsSupportedTemplate(*ScriptStruct))
+				{
+					PinnedParent->Reload();
+				}
+				else
+				{
+					// Builder is now associated with MetaSound that no longer supports this configuration, so stop listening.
+					// This UObject (and owning object, likely an associated UUserWidget) will be GC'ed at an indeterminate
+					// later point.
+					DocBuilder.GetDocumentDelegates().OnDocumentTemplateChanged.Remove(TemplateChangeDelegateHandle);
+				}
+			}
+		}
+	}
+}
diff --git a/Engine/Plugins/Runtime/Metasound/Source/MetasoundEditor/Private/TemplateViewModels/RandomizerTemplateViewModel.cpp b/Engine/Plugins/Runtime/Metasound/Source/MetasoundEditor/Private/TemplateViewModels/RandomizerTemplateViewModel.cpp
new file mode 100644
--- /dev/null
+++ b/Engine/Plugins/Runtime/Metasound/Source/MetasoundEditor/Private/TemplateViewModels/RandomizerTemplateViewModel.cpp
@@ -0,0 +1,159 @@
+// Copyright Epic Games, Inc. All Rights Reserved.
+
+#include "TemplateViewModels/RandomizerConfigurationViewModel.h"
+
+#include "DocumentTemplates/MetasoundRandomizerTemplate.h"
+
+
+bool UMetaSoundRandomizerViewModel::IsSupportedTemplate(const UScriptStruct& InStruct) const
+{
+	return InStruct.IsChildOf(FMetaSoundRandomizerTemplate::StaticStruct());
+}
+
+bool UMetaSoundRandomizerViewModel::GetIsOneShot() const
+{
+	return bIsOneShot;
+}
+
+const TArray<TObjectPtr<USoundWave>>& UMetaSoundRandomizerViewModel::GetSounds() const
+{
+	return Sounds;
+}
+
+void UMetaSoundRandomizerViewModel::Reload()
+{
+	Super::Reload();
+
+	if (const FMetaSoundRandomizerTemplate* Template = GetConstTemplate<FMetaSoundRandomizerTemplate>())
+	{
+		UE_MVVM_SET_PROPERTY_VALUE(bIsOneShot, Template->bIsOneShot);
+		UE_MVVM_SET_PROPERTY_VALUE(Pitch, Template->Pitch);
+		UE_MVVM_SET_PROPERTY_VALUE(Sounds, Template->Sounds);
+
+		Template->GetPropertyChangedDelegate().Remove(PropertyChangedHandle);
+		PropertyChangedHandle = Template->GetPropertyChangedDelegate().AddUObject(this, &UMetaSoundRandomizerViewModel::OnPropertyChanged);
+	}
+}
+
+void UMetaSoundRandomizerViewModel::OnPropertyChanged(const FPropertyChangedEvent& InEvent)
+{
+	if (const FMetaSoundRandomizerTemplate* Template = GetConstTemplate<FMetaSoundRandomizerTemplate>())
+	{
+		if (InEvent.GetPropertyName() == GET_MEMBER_NAME_CHECKED(FMetaSoundRandomizerTemplate, bIsOneShot))
+		{
+			UE_MVVM_SET_PROPERTY_VALUE(bIsOneShot, Template->bIsOneShot);
+		}
+		else if (InEvent.GetPropertyName() == GET_MEMBER_NAME_CHECKED(FMetaSoundRandomizerTemplate, Pitch))
+		{
+			UE_MVVM_SET_PROPERTY_VALUE(Pitch, Template->Pitch);
+		}
+		else if (InEvent.GetPropertyName() == GET_MEMBER_NAME_CHECKED(FVector2f, X))
+		{
+			UE_MVVM_SET_PROPERTY_VALUE(Pitch, Template->Pitch);
+			BroadcastFieldValueChanged(FFieldNotificationClassDescriptor::GetPitchMin);
+		}
+		else if (InEvent.GetPropertyName() == GET_MEMBER_NAME_CHECKED(FVector2f, Y))
+		{
+			UE_MVVM_SET_PROPERTY_VALUE(Pitch, Template->Pitch);
+			BroadcastFieldValueChanged(FFieldNotificationClassDescriptor::GetPitchMax);
+		}
+		else if (InEvent.GetPropertyName() == GET_MEMBER_NAME_CHECKED(FMetaSoundRandomizerTemplate, Sounds))
+		{
+			UE_MVVM_SET_PROPERTY_VALUE(Sounds, Template->Sounds);
+		}
+	}
+}
+
+FVector2f UMetaSoundRandomizerViewModel::GetPitch() const
+{
+	return Pitch;
+}
+
+float UMetaSoundRandomizerViewModel::GetPitchMin() const
+{
+	return Pitch.X;
+}
+
+float UMetaSoundRandomizerViewModel::GetPitchMax() const
+{
+	return Pitch.Y;
+}
+
+void UMetaSoundRandomizerViewModel::SetIsOneShot(bool bInIsOneShot)
+{
+	if (Builder && UE_MVVM_SET_PROPERTY_VALUE(bIsOneShot, bInIsOneShot))
+	{
+		Builder->SetTemplateProperties<FMetaSoundRandomizerTemplate>([this](FMetaSoundRandomizerTemplate& RandomTemplate)
+		{
+			RandomTemplate.bIsOneShot = bIsOneShot;
+		});
+	}
+}
+
+void UMetaSoundRandomizerViewModel::SetPitch(const FVector2f& InPitch)
+{
+	if (Builder && UE_MVVM_SET_PROPERTY_VALUE(Pitch, InPitch))
+	{
+		Builder->SetTemplateProperties<FMetaSoundRandomizerTemplate>([this](FMetaSoundRandomizerTemplate& RandomTemplate)
+		{
+			RandomTemplate.Pitch = Pitch;
+			BroadcastFieldValueChanged(FFieldNotificationClassDescriptor::GetPitchMin);
+			BroadcastFieldValueChanged(FFieldNotificationClassDescriptor::GetPitchMax);
+		});
+	}
+}
+
+void UMetaSoundRandomizerViewModel::SetPitchMax(float PitchMax)
+{
+	FVector2f NewPitch = Pitch;
+	NewPitch.Y = PitchMax;
+	if (Builder && UE_MVVM_SET_PROPERTY_VALUE(Pitch, NewPitch))
+	{
+		Builder->SetTemplateProperties<FMetaSoundRandomizerTemplate>([this](FMetaSoundRandomizerTemplate& RandomTemplate)
+		{
+			RandomTemplate.Pitch = Pitch;
+			BroadcastFieldValueChanged(FFieldNotificationClassDescriptor::GetPitchMax);
+		});
+	}
+}
+
+void UMetaSoundRandomizerViewModel::SetPitchMin(float PitchMin)
+{
+	FVector2f NewPitch = Pitch;
+	NewPitch.X = PitchMin;
+	if (Builder && UE_MVVM_SET_PROPERTY_VALUE(Pitch, NewPitch))
+	{
+		Builder->SetTemplateProperties<FMetaSoundRandomizerTemplate>([this](FMetaSoundRandomizerTemplate& RandomTemplate)
+		{
+			RandomTemplate.Pitch = Pitch;
+			BroadcastFieldValueChanged(FFieldNotificationClassDescriptor::GetPitchMin);
+		});
+	}
+}
+
+void UMetaSoundRandomizerViewModel::SetNumSounds(int32 NumSounds)
+{
+	if (NumSounds < 0)
+	{
+		return;
+	}
+
+	TArray<TObjectPtr<USoundWave>> NewSounds = Sounds;
+	NewSounds.SetNumZeroed(NumSounds);
+	SetSounds(MoveTemp(NewSounds));
+}
+
+void UMetaSoundRandomizerViewModel::SetSounds(TArray<TObjectPtr<USoundWave>> InSounds)
+{
+	if (Builder)
+	{
+		FMetaSoundRandomizerTemplate::RemoveInvalidSounds(Builder->GetConstBuilder(), true /* bRemoveNullEntries */, InSounds);
+		if (UE_MVVM_SET_PROPERTY_VALUE(Sounds, InSounds))
+		{
+			Builder->SetTemplateProperties<FMetaSoundRandomizerTemplate>([this](FMetaSoundRandomizerTemplate& RandomTemplate)
+			{
+				RandomTemplate.Sounds = Sounds;
+			});
+		}
+	}
+}
diff --git a/Engine/Plugins/Runtime/Metasound/Source/MetasoundEditor/Public/DocumentConfigurationViewModels/MetasoundDocumentConfigurationViewModelBase.h b/Engine/Plugins/Runtime/Metasound/Source/MetasoundEditor/Public/DocumentConfigurationViewModels/MetasoundDocumentConfigurationViewModelBase.h
deleted file mode 100644
diff --git a/Engine/Plugins/Runtime/Metasound/Source/MetasoundEditor/Public/DocumentConfigurationViewModels/RandomizerConfigurationViewModel.h b/Engine/Plugins/Runtime/Metasound/Source/MetasoundEditor/Public/DocumentConfigurationViewModels/RandomizerConfigurationViewModel.h
deleted file mode 100644
diff --git a/Engine/Plugins/Runtime/Metasound/Source/MetasoundEditor/Public/MetasoundDocumentConfigurationWidgetInterface.h b/Engine/Plugins/Runtime/Metasound/Source/MetasoundEditor/Public/MetasoundDocumentConfigurationWidgetInterface.h
deleted file mode 100644
diff --git a/Engine/Plugins/Runtime/Metasound/Source/MetasoundEditor/Public/MetasoundEditorCommands.h b/Engine/Plugins/Runtime/Metasound/Source/MetasoundEditor/Public/MetasoundEditorCommands.h
--- a/Engine/Plugins/Runtime/Metasound/Source/MetasoundEditor/Public/MetasoundEditorCommands.h
+++ b/Engine/Plugins/Runtime/Metasound/Source/MetasoundEditor/Public/MetasoundEditorCommands.h
@@ -80,7 +80,7 @@ namespace Metasound
 			TSharedPtr<FUICommandInfo> PromoteAllToCommonInputs;
 
 			/** Copies the selected nodes, shared connections, and required interfaces to the clipboard as configurator c++ */
-			TSharedPtr<FUICommandInfo> CopySelectionAsConfigurationCodeToClipboard;
+			TSharedPtr<FUICommandInfo> CopySelectionAsConfiguratorCodeToClipboard;
 
 			/** Initialize commands */
 			virtual void RegisterCommands() override;
diff --git a/Engine/Plugins/Runtime/Metasound/Source/MetasoundEditor/Public/MetasoundEditorSubsystem.h b/Engine/Plugins/Runtime/Metasound/Source/MetasoundEditor/Public/MetasoundEditorSubsystem.h
--- a/Engine/Plugins/Runtime/Metasound/Source/MetasoundEditor/Public/MetasoundEditorSubsystem.h
+++ b/Engine/Plugins/Runtime/Metasound/Source/MetasoundEditor/Public/MetasoundEditorSubsystem.h
@@ -31,7 +31,10 @@ class UMetaSoundEditorSubsystem : public UEditorSubsystem
 public:
 	struct FInitAssetArgs
 	{
-		TInstancedStruct<FMetaSoundFrontendDocumentConfiguration> DocumentConfiguration;
+		// (Optional) Template to apply to newly initialized asset
+		TInstancedStruct<FMetaSoundFrontendDocumentTemplate> Template;
+
+		// Currently selected objects in the focused ContentBrowser.
 		TArray<UObject*> SelectedObjects;
 	};
 
@@ -88,7 +91,7 @@ public:
 		UPARAM(DisplayName = "Location") const FVector2D& InLocation,
 		EMetaSoundBuilderResult& OutResult);
 
-	// Initialize the UObject asset with an optional MetaSound document configuration to be applied.
+	// Initialize the UObject asset with optional initialization arguments to be applied.
 	// Clears the existing MetaSound document (for the case of duplicated assets).
 	UE_API void InitAsset(UObject& InNewMetaSound, FInitAssetArgs InitAssetArgs = { });
 
diff --git a/Engine/Plugins/Runtime/Metasound/Source/MetasoundEditor/Public/MetasoundFactory.h b/Engine/Plugins/Runtime/Metasound/Source/MetasoundEditor/Public/MetasoundFactory.h
--- a/Engine/Plugins/Runtime/Metasound/Source/MetasoundEditor/Public/MetasoundFactory.h
+++ b/Engine/Plugins/Runtime/Metasound/Source/MetasoundEditor/Public/MetasoundFactory.h
@@ -2,7 +2,7 @@
 
 #pragma once
 
-#include "DocumentConfigurations/MetasoundFrontendDocumentConfiguration.h"
+#include "DocumentTemplates/MetasoundFrontendDocumentTemplate.h"
 #include "Factories/Factory.h"
 #include "MetasoundFrontendDocument.h"
 #include "StructUtils/InstancedStruct.h"
@@ -20,11 +20,11 @@ class UMetaSoundBaseFactory : public UFactory
 {
 	GENERATED_UCLASS_BODY()
 
-	UPROPERTY(Transient, meta = (Deprecated = 5.8, DeprecationMessage = "Use document configuration instead"))
+	UPROPERTY(Transient, meta = (Deprecated = 5.8, DeprecationMessage = "Use document template instead"))
 	TObjectPtr<UObject> ReferencedMetaSoundObject;
 
 	UPROPERTY(EditAnywhere, Transient, meta = (Category = Factory))
-	TInstancedStruct<FMetaSoundFrontendDocumentConfiguration> DocumentConfiguration;
+	TInstancedStruct<FMetaSoundFrontendDocumentTemplate> Template;
 
 	// If Document configuration is set, selection is provided to the configuration
 	// as an initial argument to be set on the configuration if applicable.
diff --git a/Engine/Plugins/Runtime/Metasound/Source/MetasoundEditor/Public/MetasoundPresetWidgetInterface.h b/Engine/Plugins/Runtime/Metasound/Source/MetasoundEditor/Public/MetasoundPresetWidgetInterface.h
--- a/Engine/Plugins/Runtime/Metasound/Source/MetasoundEditor/Public/MetasoundPresetWidgetInterface.h
+++ b/Engine/Plugins/Runtime/Metasound/Source/MetasoundEditor/Public/MetasoundPresetWidgetInterface.h
@@ -4,7 +4,7 @@
 
 #include "Audio/AudioPanelWidgetInterface.h"
 #include "MetasoundBuilderSubsystem.h"
-#include "MetasoundDocumentConfigurationWidgetInterface.h"
+#include "MetasoundTemplateWidgetInterface.h"
 #include "UObject/Class.h"
 #include "UObject/Interface.h"
 #include "UObject/ObjectMacros.h"
@@ -14,7 +14,7 @@
 #include "MetasoundPresetWidgetInterface.generated.h"
 
 
-// IMetaSoundPresetWidgetInterface has been soft deprecated. Use IMetaSoundDocumentConfigurationWidgetInterface instead.
+// IMetaSoundPresetWidgetInterface has been soft deprecated. Use IMetaSoundTemplateWidgetInterface instead.
 UINTERFACE(Blueprintable, MinimalAPI)
 class UMetaSoundPresetWidgetInterface : public UAudioPanelWidgetInterface
 {
@@ -42,5 +42,5 @@ public:
 
 	// (Optional) Called when initializing a tab owning the given widget to allow implementation to provide basic tab display options and information.
 	UFUNCTION(BlueprintImplementableEvent, Category = "MetaSound Preset Widget", meta = (DisplayName = "On Initializing Tab"))
-	void OnInitializingTab(FMetaSoundDocumentConfigurationWidgetTabInfo& TabInfo);
+	void OnInitializingTab(FMetaSoundTemplateWidgetTabInfo& TabInfo);
 };
diff --git a/Engine/Plugins/Runtime/Metasound/Source/MetasoundEditor/Public/MetasoundTemplateWidgetInterface.h b/Engine/Plugins/Runtime/Metasound/Source/MetasoundEditor/Public/MetasoundTemplateWidgetInterface.h
new file mode 100644
--- /dev/null
+++ b/Engine/Plugins/Runtime/Metasound/Source/MetasoundEditor/Public/MetasoundTemplateWidgetInterface.h
@@ -0,0 +1,75 @@
+// Copyright Epic Games, Inc. All Rights Reserved.
+
+#pragma once
+
+#include "Containers/Set.h"
+#include "MetasoundBuilderBase.h"
+#include "MetasoundBuilderSubsystem.h"
+#include "MetasoundDocumentInterface.h"
+#include "UObject/Class.h"
+#include "UObject/Interface.h"
+#include "UObject/ObjectMacros.h"
+#include "UObject/UObjectGlobals.h"
+#include "UObject/ScriptInterface.h"
+
+
+#include "MetasoundTemplateWidgetInterface.generated.h"
+
+
+// Forward Declarations
+class UMetaSoundTemplateViewModelBase;
+
+
+UINTERFACE(Blueprintable, MinimalAPI, DisplayName = "MetaSound Configuration Widget Interface")
+class UMetaSoundTemplateWidgetInterface : public UInterface
+{
+	GENERATED_BODY()
+};
+
+USTRUCT(BlueprintType, MinimalAPI, DisplayName = "MetaSound Configuration Widget Tab Info")
+struct FMetaSoundTemplateWidgetTabInfo
+{
+	GENERATED_BODY()
+
+	UPROPERTY(BlueprintReadWrite, Category = TabInfo)
+	FName TabName;
+
+	UPROPERTY(BlueprintReadWrite, Category = TabInfo)
+	FText TabDisplayName;
+
+	UPROPERTY(BlueprintReadWrite, Category = TabInfo)
+	FName StyleName = "MetaSoundStyle";
+
+	UPROPERTY(BlueprintReadWrite, Category = TabInfo)
+	FName TabIcon = "MetasoundEditor.Metasound.Icon";
+
+	UPROPERTY(BlueprintReadWrite, Category = TabInfo)
+	FText ToolTip;
+};
+
+class IMetaSoundTemplateWidgetInterface
+{
+	GENERATED_BODY()
+
+public:
+	// Returns whether or not the widget supports the given template. Called when initializing
+	// the MetaSound Editor to determine if an implemented template supports the given widget.
+	UFUNCTION(BlueprintImplementableEvent, BlueprintCallable, Category = "MetaSound|Configuration", meta = (DisplayName = "Is Supported Template"))
+	void IsSupportedTemplate(FName ConfigName, const UMetaSoundBuilderBase* Builder, UPARAM(DisplayName = "Is Supported") bool& bOutIsSupported);
+
+	// Called when the MetaSound starts and stops auditioning. Provides a reference to the audio component when auditioning starts, and returns nullptr when auditioning stops.
+	UFUNCTION(BlueprintImplementableEvent, Category = "MetaSound|Configuration", meta = (DisplayName = "On MetaSound Audition State Changed"))
+	void OnAuditionStateChanged(UAudioComponent* AudioComponent, bool bIsAuditioning);
+
+	// Called when the builder is initialized (prior to PreConstruct).
+	UFUNCTION(BlueprintImplementableEvent, Category = "MetaSound|Configuration", meta = (DisplayName = "On MetaSound Builder Set"))
+	void OnBuilderSet(UMetaSoundBuilderBase* Builder);
+
+	// (Optional) Called when initializing a tab owning the given widget to allow implementation to provide basic tab display options and information.
+	UFUNCTION(BlueprintImplementableEvent, Category = "MetaSound|Configuration", meta = (DisplayName = "On Initializing Tab"))
+	void OnInitializingTab(FMetaSoundTemplateWidgetTabInfo& TabInfo);
+
+	// (Optional) If ViewModel class provided, automatically initializes and sets ViewModel of the given class for the given widget.
+	UFUNCTION(BlueprintImplementableEvent, Category = "MetaSound|Configuration", meta = (DisplayName = "On Initializing ViewModel"))
+	void OnInitializingViewModel(TSubclassOf<UMetaSoundTemplateViewModelBase>& Class);
+};
diff --git a/Engine/Plugins/Runtime/Metasound/Source/MetasoundEditor/Public/TemplateViewModels/MetasoundTemplateViewModelBase.h b/Engine/Plugins/Runtime/Metasound/Source/MetasoundEditor/Public/TemplateViewModels/MetasoundTemplateViewModelBase.h
new file mode 100644
--- /dev/null
+++ b/Engine/Plugins/Runtime/Metasound/Source/MetasoundEditor/Public/TemplateViewModels/MetasoundTemplateViewModelBase.h
@@ -0,0 +1,80 @@
+// Copyright Epic Games, Inc. All Rights Reserved.
+
+#pragma once
+
+#include "Blueprint/UserWidget.h"
+#include "Delegates/IDelegateInstance.h"
+#include "MetasoundBuilderBase.h"
+#include "MetasoundFrontendDocumentModifyDelegates.h"
+#include "MVVMViewModelBase.h"
+#include "StructUtils/StructView.h"
+#include "DocumentTemplates/MetasoundFrontendDocumentTemplate.h"
+#include "UObject/Class.h"
+
+#include "MetasoundTemplateViewModelBase.generated.h"
+
+// Forward Declarations
+class UUserWidget;
+
+#define UE_API METASOUNDEDITOR_API
+
+
+UCLASS(MinimalAPI, Blueprintable, Abstract, DisplayName = "MetaSound Configuration ViewModel Base")
+class UMetaSoundTemplateViewModelBase : public UMVVMViewModelBase
+{
+	GENERATED_BODY()
+
+public:
+	// Constructs a view model
+	static void ConstructWidgetViewModel(
+		UUserWidget* UserWidget,
+		UMetaSoundBuilderBase& InBuilder,
+		TConstStructView<FMetaSoundFrontendDocumentTemplate> Config,
+		TSubclassOf<UMetaSoundTemplateViewModelBase> VMClass);
+
+	UE_API virtual bool IsSupportedTemplate(const UScriptStruct& InStruct) const;
+
+protected:
+	template <typename TemplateType>
+	const TemplateType* GetConstTemplate() const
+	{
+		if (Builder)
+		{
+			return Builder->GetConstBuilder().GetConstTemplateAs<TemplateType>();
+		}
+		return nullptr;
+	}
+
+	// Enables responding to the builder reloading or the model view requiring a full
+	// resynchronization of configuration properties to view (i.e. when user transactions
+	// are applied that require a full resync of the view such as undo, redo, etc.)
+	UE_API virtual void Reload();
+
+	UPROPERTY(Transient)
+	TObjectPtr<UMetaSoundBuilderBase> Builder;
+
+private:
+	void Initialize(UMetaSoundBuilderBase* InBuilder);
+
+	class FBuilderListener : public Metasound::Frontend::IDocumentBuilderTransactionListener
+	{
+		TWeakObjectPtr<UMetaSoundTemplateViewModelBase> Parent;
+
+	public:
+		FBuilderListener() = default;
+		FBuilderListener(UMetaSoundTemplateViewModelBase* InParent)
+			: Parent(InParent)
+		{
+		}
+
+		virtual ~FBuilderListener() = default;
+
+	private:
+		virtual void OnBuilderReloaded(Metasound::Frontend::FDocumentModifyDelegates& OutDelegates) override;
+		void OnTemplateChanged(const Metasound::Frontend::FDocumentTemplateChangedArgs& Args);
+
+		FDelegateHandle TemplateChangeDelegateHandle;
+	};
+	TSharedPtr<FBuilderListener> BuilderListener;
+};
+#undef UE_API // METASOUNDEDITOR_API
diff --git a/Engine/Plugins/Runtime/Metasound/Source/MetasoundEditor/Public/TemplateViewModels/RandomizerConfigurationViewModel.h b/Engine/Plugins/Runtime/Metasound/Source/MetasoundEditor/Public/TemplateViewModels/RandomizerConfigurationViewModel.h
new file mode 100644
--- /dev/null
+++ b/Engine/Plugins/Runtime/Metasound/Source/MetasoundEditor/Public/TemplateViewModels/RandomizerConfigurationViewModel.h
@@ -0,0 +1,71 @@
+// Copyright Epic Games, Inc. All Rights Reserved.
+
+#pragma once
+
+#include "MetasoundTemplateViewModelBase.h"
+#include "MetasoundFrontendDocumentModifyDelegates.h"
+#include "Sound/SoundWave.h"
+
+#include "RandomizerConfigurationViewModel.generated.h"
+
+#define UE_API METASOUNDEDITOR_API
+
+
+// Forward Declarations
+struct FMetaSoundRandomizerTemplate;
+
+
+UCLASS(MinimalAPI, DisplayName = "MetaSound Randomizer ViewModel")
+class UMetaSoundRandomizerViewModel : public UMetaSoundTemplateViewModelBase
+{
+	GENERATED_BODY()
+
+public:
+	UE_API bool GetIsOneShot() const;
+	UE_API FVector2f GetPitch() const;
+
+	UE_API const TArray<TObjectPtr<USoundWave>>& GetSounds() const;
+
+	UFUNCTION(BlueprintPure, FieldNotify, Category = "MetaSound|Randomizer")
+	UE_API float GetPitchMin() const;
+
+	UFUNCTION(BlueprintPure, FieldNotify, Category = "MetaSound|Randomizer")
+	UE_API float GetPitchMax() const;
+
+	UFUNCTION(BlueprintCallable, Category = "MetaSound|Randomizer")
+	UE_API void SetNumSounds(int32 NumSounds);
+
+	UFUNCTION(BlueprintCallable, Category = "MetaSound|Randomizer")
+	UE_API void SetPitchMin(float PitchMin);
+
+	UFUNCTION(BlueprintCallable, Category = "MetaSound|Randomizer")
+	UE_API void SetPitchMax(float PitchMax);
+
+	UE_API void SetIsOneShot(bool bInIsOneShot);
+	UE_API void SetPitch(const FVector2f& InPitch);
+
+	UE_API void SetSounds(TArray<TObjectPtr<USoundWave>> InSounds);
+
+	UE_API virtual bool IsSupportedTemplate(const UScriptStruct& InStruct) const override;
+	UE_API virtual void Reload() override;
+
+private:
+	void OnPropertyChanged(const FPropertyChangedEvent& InEvent);
+
+	FDelegateHandle PropertyChangedHandle;
+
+	// If randomizer is one shot (looping if false).
+	UPROPERTY(BlueprintReadWrite, FieldNotify, Getter = "GetIsOneShot", Setter = "SetIsOneShot", Category = "MetaSound|Randomizer", meta = (AllowPrivateAccess))
+	bool bIsOneShot = false;
+
+	UPROPERTY(BlueprintReadWrite, FieldNotify, Getter = "GetSounds", Setter = "SetSounds", Category = "MetaSound|Randomizer", meta = (
+		AllowPrivateAccess,
+		DisallowedClasses = "/Script/MetasoundEngine.MetaSoundSource, /Script/Engine.SoundSourceBus"
+	))
+	TArray<TObjectPtr<USoundWave>> Sounds;
+
+	// Sound's pitch range in semitones.
+	UPROPERTY(BlueprintReadWrite, FieldNotify, Getter = "GetPitch", Setter = "SetPitch", Category = "MetaSound|Randomizer", meta = (AllowPrivateAccess))
+	FVector2f Pitch { 0.f, 0.f };
+};
+#undef UE_API
diff --git a/Engine/Plugins/Runtime/Metasound/Source/MetasoundEngine/Private/DocumentConfigurations/MetasoundDocumentConfigurator.cpp b/Engine/Plugins/Runtime/Metasound/Source/MetasoundEngine/Private/DocumentConfigurations/MetasoundDocumentConfigurator.cpp
deleted file mode 100644
diff --git a/Engine/Plugins/Runtime/Metasound/Source/MetasoundEngine/Private/DocumentConfigurations/MetasoundRandomizerDocumentConfiguration.cpp b/Engine/Plugins/Runtime/Metasound/Source/MetasoundEngine/Private/DocumentConfigurations/MetasoundRandomizerDocumentConfiguration.cpp
deleted file mode 100644
diff --git a/Engine/Plugins/Runtime/Metasound/Source/MetasoundEngine/Private/DocumentTemplates/MetasoundDocumentConfigurator.cpp b/Engine/Plugins/Runtime/Metasound/Source/MetasoundEngine/Private/DocumentTemplates/MetasoundDocumentConfigurator.cpp
new file mode 100644
--- /dev/null
+++ b/Engine/Plugins/Runtime/Metasound/Source/MetasoundEngine/Private/DocumentTemplates/MetasoundDocumentConfigurator.cpp
@@ -0,0 +1,1100 @@
+// Copyright Epic Games, Inc. All Rights Reserved.
+
+#include "DocumentTemplates/MetasoundDocumentConfigurator.h"
+
+#include "Algo/Transform.h"
+#include "AudioParameter.h"
+#include "Templates/Function.h"
+#include "Interfaces/MetasoundFrontendInterfaceRegistry.h"
+#include "MetasoundAssetSubsystem.h"
+#include "MetasoundDocumentBuilderRegistry.h"
+#include "MetasoundFrontendDocument.h"
+#include "MetasoundFrontendDocumentBuilder.h"
+#include "MetasoundFrontendLiteral.h"
+#include "MetasoundFrontendSearchEngine.h"
+#include "NodeTemplates/MetasoundFrontendNodeTemplateInput.h"
+#include "UObject/SoftObjectPath.h"
+
+
+#if WITH_EDITORONLY_DATA
+namespace Metasound::Engine
+{
+	namespace ConfigCodeReflectorPrivate
+	{
+		struct FConfigClassInfo
+		{
+			const FMetasoundFrontendClass* Class = nullptr;
+		};
+
+		struct FConfigNodeInfo
+		{
+			FString InstanceName;
+			const FMetasoundFrontendNode* Node = nullptr;
+		};
+
+		template <typename EntryType>
+		void FormatLiteralArray(const FMetasoundFrontendLiteral& Literal, FStringBuilderBase& OutBuilder)
+		{
+			if (Literal.GetArrayNum() == 0)
+			{
+				OutBuilder += TEXT(" )");
+				return;
+			}
+
+			OutBuilder += TEXT("\n\t\t{");
+
+			TArray<EntryType> ArrayEntries;
+			TArray<int32> Entries;
+			ensure(Literal.TryGet(ArrayEntries));
+			for (EntryType& Entry : ArrayEntries)
+			{
+				OutBuilder.Appendf(TEXT("\n\t\t\t%s,"), *LexToString(Entry));
+			}
+
+			OutBuilder += TEXT("\n\t\t})\n\t");
+		}
+
+		template <typename EntryType>
+		void FormatLiteralArray(const FMetasoundFrontendLiteral& Literal, FStringBuilderBase& OutBuilder, TFunctionRef<FString(EntryType&)> GetEntryString)
+		{
+			if (Literal.GetArrayNum() == 0)
+			{
+				OutBuilder += TEXT(" )");
+				return;
+			}
+
+			OutBuilder += TEXT("\n\t\t{");
+
+			TArray<EntryType> ArrayEntries;
+			ensure(Literal.TryGet(ArrayEntries));
+			for (EntryType& Entry : ArrayEntries)
+			{
+				const FString EntryString = GetEntryString(Entry);
+				OutBuilder.Appendf(TEXT("\n\t\t\t%s,"), *EntryString);
+			}
+
+			OutBuilder += TEXT("\n\t\t})\n\t");
+		}
+
+		class FConfigReflector
+		{
+			// non-const to enable harvesting input defaults via move semantics.
+			// Otherwise, structure left alone through reflection execution.
+			FMetasoundFrontendGraphLayer Layer;
+
+			TArray<FMetasoundFrontendClassOutput> InterfaceOutputs;
+
+			TSet<FString> NodeInstanceNames;
+
+			TMap<FGuid, FName> TemplateNodeIDToInputName;
+			TMap<FGuid, FConfigClassInfo> ClassIDToInfo;
+			TMap<FGuid, FConfigNodeInfo> NodeIDToInfo;
+
+			TStringBuilder<0> NodeCode;
+			TStringBuilder<0> NodeInputDefaultsCode;
+
+			TStringBuilder<0> InterfaceCode;
+
+			TStringBuilder<0> InputCode;
+
+#if WITH_EDITORONLY_DATA
+			TStringBuilder<0> InputNodeCode; // Input Template nodes
+#endif // WITH_EDITORONLY_DATA
+
+			TStringBuilder<0> OutputCode;
+			TStringBuilder<0> OutputEdgesCode;
+
+			TStringBuilder<0> EdgesCode;
+
+			void CacheDependencies()
+			{
+				Algo::Transform(Layer.Dependencies, ClassIDToInfo, [](const FMetasoundFrontendClass& Dependency)
+				{
+					return TPair<FGuid, FConfigClassInfo>(Dependency.ID, FConfigClassInfo { &Dependency });
+				});
+			}
+
+			FName FindConnectedInputName(const FGuid& InTemplateNodeID)
+			{
+				for (const FMetasoundFrontendEdge& Edge : Layer.Graph.Edges)
+				{
+					if (Edge.ToNodeID == InTemplateNodeID)
+					{
+						auto IsInput = [&Edge](const FMetasoundFrontendClassInput& Input)
+						{
+							return Input.NodeID == Edge.FromNodeID;
+						};
+						if (const FMetasoundFrontendClassInput* Input = Layer.Inputs.FindByPredicate(IsInput))
+						{
+							return Input->Name;
+						}
+					}
+				}
+
+				return NAME_None;
+			}
+
+#if WITH_EDITORONLY_DATA
+			void CacheInputNodes()
+			{
+				using namespace Frontend;
+
+				for (const FMetasoundFrontendNode& Node: Layer.Graph.Nodes)
+				{
+					if (const FConfigClassInfo* ClassInfo = ClassIDToInfo.Find(Node.ClassID))
+					{
+						if (ClassInfo->Class->Metadata.GetClassName() == FInputNodeTemplate::GetChecked().GetClassName())
+						{
+							const FName InputName = FindConnectedInputName(Node.GetID());
+							if (!InputName.IsNone())
+							{
+								TemplateNodeIDToInputName.Add(Node.GetID(), InputName);
+							}
+						}
+					}
+				}
+			}
+#endif // WITH_EDITORONLY_DATA
+
+			const FMetasoundFrontendClassOutput* FindOutput(const FGuid& InNodeID) const
+			{
+				auto IsOutput = [&InNodeID](const FMetasoundFrontendClassOutput& Output) { return Output.NodeID == InNodeID; };
+				if (const FMetasoundFrontendClassOutput* Output = Layer.Outputs.FindByPredicate(IsOutput))
+				{
+					return Output;
+				}
+
+				if (const FMetasoundFrontendClassOutput* Output = InterfaceOutputs.FindByPredicate(IsOutput))
+				{
+					return Output;
+				}
+
+				return nullptr;
+			}
+
+			static void FormatDefaultCode(const FMetasoundFrontendLiteral& Literal, FStringBuilderBase& OutBuilder)
+			{
+				OutBuilder.Append(TEXT("MakeFrontendLiteral("));
+
+				auto FormatAsLoadedPath = [](const FString& InPathString)
+				{
+					return FString::Printf(TEXT("FSoftObjectPath(TEXT(\"%s\")).TryLoad()"), *InPathString);
+				};
+
+				switch (Literal.GetType())
+				{
+					case EMetasoundFrontendLiteralType::Boolean:
+					case EMetasoundFrontendLiteralType::Integer:
+					{
+						OutBuilder.Appendf(TEXT("%s) "), *Literal.ToString());
+					}
+					break;
+
+					case EMetasoundFrontendLiteralType::Float:
+					{
+						float Value = 0.f;
+						ensure(Literal.TryGet(Value));
+						OutBuilder.Appendf(TEXT("%sf) "), *FString::SanitizeFloat(Value));
+					}
+					break;
+
+					case EMetasoundFrontendLiteralType::String:
+					{
+						OutBuilder.Appendf(TEXT("TEXT(\"%s\")) "), *Literal.ToString());
+					}
+					break;
+
+					case EMetasoundFrontendLiteralType::UObject:
+					{
+						OutBuilder.Appendf(TEXT("%s) "), *FormatAsLoadedPath(Literal.ToString()));
+					}
+					break;
+
+					case EMetasoundFrontendLiteralType::BooleanArray:
+					{
+						FormatLiteralArray<bool>(Literal, OutBuilder);
+					}
+					break;
+
+					case EMetasoundFrontendLiteralType::IntegerArray:
+					{
+						FormatLiteralArray<int32>(Literal, OutBuilder);
+					}
+					break;
+
+					case EMetasoundFrontendLiteralType::FloatArray:
+					{
+						FormatLiteralArray<float>(Literal, OutBuilder, [&](float& Value)
+						{
+							return FString::SanitizeFloat(Value);
+						});
+					}
+					break;
+
+					case EMetasoundFrontendLiteralType::StringArray:
+					{
+						FormatLiteralArray<FString>(Literal, OutBuilder);
+					}
+					break;
+
+					case EMetasoundFrontendLiteralType::UObjectArray:
+					{
+						FormatLiteralArray<UObject*>(Literal, OutBuilder, [&](UObject*& Object)
+						{
+							const FSoftObjectPath Path(Object);
+							return FormatAsLoadedPath(Path.ToString());
+						});
+					}
+					break;
+
+					case EMetasoundFrontendLiteralType::None:
+					case EMetasoundFrontendLiteralType::NoneArray:
+					case EMetasoundFrontendLiteralType::Invalid:
+					{
+						OutBuilder += TEXT(")");
+					}
+					default:
+					break;
+				}
+			}
+
+			// Edges from inputs are not distinct in reflected code as inputs are characterized
+			// visually by template nodes which are included in the typical edge addition case.
+			void ReflectEdges()
+			{
+				if (Layer.Graph.Edges.IsEmpty())
+				{
+					return;
+				}
+
+				bool bHasNodeEdges = false;
+				EdgesCode = TEXT("\n// Edges\n");
+				EdgesCode += TEXT(".Add<FDC::FEdge>(\n{\n");
+
+				bool bHasOutputEdges = false;
+				OutputEdgesCode = TEXT("\n// Output Edges\n");
+				OutputEdgesCode += TEXT(".Add<FDC::FEdgeToOutput>(\n{\n");
+
+				for (const FMetasoundFrontendEdge& Edge : Layer.Graph.Edges)
+				{
+					// Edges between template nodes and inputs are "hidden" to mirror
+					// design characterized in editor layer and make config code more
+					// readable. Edges between associated inputs and template nodes
+					// are therefore added via the Add(FInput) function and not here,
+					// so this case early outs.
+					if (TemplateNodeIDToInputName.Contains(Edge.ToNodeID))
+					{
+						continue;
+					}
+
+					if (const FConfigNodeInfo* FromNodeInfo = NodeIDToInfo.Find(Edge.FromNodeID))
+					{
+						const FMetasoundFrontendNode* FromNode = FromNodeInfo->Node;
+						check(FromNode);
+
+						auto IsFromVertex = [&Edge](const FMetasoundFrontendVertex& Vertex) { return Vertex.VertexID == Edge.FromVertexID; };
+						const FMetasoundFrontendVertex* FromVertex = FromNode->Interface.Outputs.FindByPredicate(IsFromVertex);
+						if (!FromVertex)
+						{
+							continue;
+						}
+
+						if (const FConfigNodeInfo* ToNodeInfo = NodeIDToInfo.Find(Edge.ToNodeID))
+						{
+							auto GetClassInfoIfOfType = [this](const FConfigNodeInfo& NodeInfo, EMetasoundFrontendClassType Type) -> const FConfigClassInfo*
+							{
+								check(NodeInfo.Node);
+								if (const FConfigClassInfo* ClassInfo = ClassIDToInfo.Find(NodeInfo.Node->ClassID))
+								{
+									const FMetasoundFrontendClass* Class = ClassInfo->Class;
+									check(Class);
+									if (Class && Class->Metadata.GetType() == Type)
+									{
+										return ClassInfo;
+									}
+								}
+								return nullptr;
+							};
+							if (const FConfigClassInfo* ClassInfo = GetClassInfoIfOfType(*ToNodeInfo, EMetasoundFrontendClassType::Output))
+							{
+								if (const FMetasoundFrontendClassOutput* Output = FindOutput(Edge.ToNodeID))
+								{
+									bHasOutputEdges = true;
+									OutputEdgesCode.Appendf(
+										TEXT("\t{ .Node = \"%s\", .NodeOutput = \"%s\", .GraphOutput = \"%s\" },\n"),
+										*ToNodeInfo->InstanceName,
+										*FromVertex->Name.ToString(),
+										*Output->Name.ToString()
+									);
+								}
+								continue;
+							}
+
+							const FMetasoundFrontendNode* ToNode = ToNodeInfo->Node;
+							check(ToNode);
+
+							auto IsToVertex = [&Edge](const FMetasoundFrontendVertex& Vertex) { return Vertex.VertexID == Edge.ToVertexID; };
+							const FMetasoundFrontendVertex* ToVertex = ToNode->Interface.Inputs.FindByPredicate(IsToVertex);
+							if (!ToVertex)
+							{
+								continue;
+							}
+
+							bHasNodeEdges = true;
+							EdgesCode.Appendf(
+								TEXT("\t{ .OutputNode = \"%s\", .Output = \"%s\", .InputNode = \"%s\", .Input = \"%s\" },\n"),
+								*FromNodeInfo->InstanceName,
+								*FromVertex->Name.ToString(),
+								*ToNodeInfo->InstanceName,
+								*ToVertex->Name.ToString()
+							);
+						}
+					}
+				}
+
+				auto FinishEdgeCode = [](bool bHadContent, FStringBuilderBase& CodeBuilder)
+				{
+					if (bHadContent)
+					{
+						CodeBuilder += TEXT("})");
+					}
+					else
+					{
+						CodeBuilder = { };
+					}
+				};
+
+				FinishEdgeCode(bHasNodeEdges, EdgesCode);
+				FinishEdgeCode(bHasOutputEdges, OutputEdgesCode);
+			}
+
+			void ReflectInterfaces()
+			{
+				using namespace Frontend;
+
+				if (Layer.Inputs.IsEmpty() && Layer.Outputs.IsEmpty())
+				{
+					return;
+				}
+
+				TSet<FMetasoundFrontendVersion> Interfaces;
+
+				IInterfaceRegistry& InterfaceRegistry = IInterfaceRegistry::Get();
+
+				auto GetMemberInterface = [&InterfaceRegistry](FName MemberName) -> TOptional<FMetasoundFrontendVersion>
+				{
+					FName InterfaceNamespace;
+					FName ParamName;
+					Audio::FParameterPath::SplitName(MemberName, InterfaceNamespace, ParamName);
+
+					if (!InterfaceNamespace.IsNone())
+					{
+						FMetasoundFrontendVersion InterfaceVersion;
+						const bool bFoundVersion = ISearchEngine::Get().FindHighestInterfaceVersion(InterfaceNamespace, InterfaceVersion);
+						if (bFoundVersion)
+						{
+							return InterfaceVersion;
+						}
+					}
+
+					return { };
+				};
+
+				for (int32 Index = Layer.Inputs.Num() - 1; Index >= 0; --Index)
+				{
+					FMetasoundFrontendClassInput& Input = Layer.Inputs[Index];
+					TOptional<FMetasoundFrontendVersion> InterfaceVersion = GetMemberInterface(Input.Name);
+					if (InterfaceVersion.IsSet() && InterfaceRegistry.ContainsInput(InterfaceVersion.GetValue(), Input))
+					{
+						Interfaces.Add(InterfaceVersion.GetValue());
+						Layer.Inputs.RemoveAtSwap(Index, EAllowShrinking::No);
+					}
+				}
+				Layer.Inputs.Shrink();
+
+				for (int32 Index = Layer.Outputs.Num() - 1; Index >= 0; --Index)
+				{
+					FMetasoundFrontendClassOutput& Output = Layer.Outputs[Index];
+					TOptional<FMetasoundFrontendVersion> InterfaceVersion = GetMemberInterface(Output.Name);
+					if (InterfaceVersion.IsSet() && InterfaceRegistry.ContainsOutput(InterfaceVersion.GetValue(), Output))
+					{
+						Interfaces.Add(InterfaceVersion.GetValue());
+						InterfaceOutputs.Add(MoveTemp(Layer.Outputs[Index]));
+						Layer.Outputs.RemoveAtSwap(Index, EAllowShrinking::No);
+					}
+				}
+				Layer.Outputs.Shrink();
+
+				if (!Interfaces.IsEmpty())
+				{
+					InterfaceCode = TEXT("\n// Interfaces\n");
+					InterfaceCode += TEXT(".Add<FDC::FInterface>(\n{\n");
+					for (const FMetasoundFrontendVersion& Version : Interfaces)
+					{
+						InterfaceCode += FString::Format(TEXT("\t{ .Version = { .Name = \"{0}\", .Number = { .Major = {1}, .Minor = {2} } } },\n"),
+						{
+							*Version.Name.ToString(),
+							Version.Number.Major,
+							Version.Number.Minor
+						});
+					}
+
+					InterfaceCode += TEXT("})");
+				}
+			}
+
+			void ReflectInputs()
+			{
+				if (Layer.Inputs.IsEmpty())
+				{
+					return;
+				}
+
+				InputCode = TEXT("\n// Inputs\n");
+				InputCode += TEXT(".Add<FDC::FInput>(\n{\n");
+
+				FMetaSoundAssetManager& AssetManager = FMetaSoundAssetManager::GetChecked();
+				for (FMetasoundFrontendClassInput& Input : Layer.Inputs)
+				{
+					// Configurations/MetaSound Builder Subsystem currently do not support pages,
+					// so attempt to harvest literal with Default PageID.  Otherwise, take any value
+					// (i.e. undefined behavior).
+					FMetasoundFrontendLiteral InitValue = FMetasoundFrontendLiteral::GetInvalid();
+					Input.IterateDefaults([&InitValue](const FGuid& PageID, FMetasoundFrontendLiteral& Literal)
+					{
+						if (PageID == Frontend::DefaultPageID || !InitValue.IsValid())
+						{
+							InitValue = MoveTemp(Literal);
+						}
+					});
+
+					InputCode += TEXT("\t{ ");
+					{
+						InputCode.Appendf(TEXT(".Name = \"%s\", "), *Input.Name.ToString());
+						InputCode.Appendf(TEXT(".DataType = \"%s\", "), *Input.TypeName.ToString());
+						InputCode.Appendf(TEXT(".DefaultValue = ")); FormatDefaultCode(InitValue, InputCode);
+						InputCode.Appendf(TEXT(".bIsConstructorInput = %s"), Input.AccessType == EMetasoundFrontendVertexAccessType::Value ? TEXT("true") : TEXT("false"));
+					}
+					InputCode += TEXT("},\n");
+
+				}
+
+				InputCode += TEXT("})");
+			}
+
+			void ReflectNodeInputDefaults()
+			{
+				TStringBuilder<0> LiteralsCode;
+
+				for (const FMetasoundFrontendNode& Node : Layer.Graph.Nodes)
+				{
+					if (const FConfigNodeInfo* NodeInfo = NodeIDToInfo.Find(Node.GetID()))
+					{
+						for (const FMetasoundFrontendVertexLiteral& VertexLiteral : Node.InputLiterals)
+						{
+							auto IsVertexWithID = [&VertexLiteral](const FMetasoundFrontendVertex& Vertex) { return Vertex.VertexID == VertexLiteral.VertexID; };
+							if (const FMetasoundFrontendVertex* InputVertex = Node.Interface.Inputs.FindByPredicate(IsVertexWithID))
+							{
+								LiteralsCode += TEXT("\t{ ");
+								{
+									LiteralsCode.Appendf(TEXT(".Node = \"%s\", "), *NodeInfo->InstanceName);
+									LiteralsCode.Appendf(TEXT(".Input = \"%s\", "), *InputVertex->Name.ToString());
+									LiteralsCode.Appendf(TEXT(".Value = ")); FormatDefaultCode(VertexLiteral.Value, LiteralsCode);
+								}
+								LiteralsCode += TEXT("},\n");
+							}
+						}
+					}
+				}
+
+				FStringView LiteralsView = LiteralsCode.ToView();
+				if (!LiteralsView.IsEmpty())
+				{
+					NodeInputDefaultsCode = TEXT("\n// Node Input Defaults\n");
+					NodeInputDefaultsCode += TEXT(".Set<FDC::FNodeInputDefault>(\n{\n");
+					NodeInputDefaultsCode += LiteralsView;
+					NodeInputDefaultsCode += TEXT("})");
+				}
+			}
+
+			void ReflectNodes()
+			{
+				using namespace Frontend;
+
+				if (Layer.Graph.Nodes.IsEmpty())
+				{
+					return;
+				}
+
+				auto CacheNodeInfo = [this](const FString& BaseName, const FMetasoundFrontendNode& Node)
+				{
+					int32 Index = 0;
+					FString FullName = BaseName;
+					while (NodeInstanceNames.Contains(FullName))
+					{
+						FullName = FString::Printf(TEXT("%s_%i"), *BaseName, ++Index);
+					}
+					NodeInstanceNames.Add(FullName);
+					NodeIDToInfo.Add(Node.GetID(), FConfigNodeInfo { .InstanceName = FullName, .Node = &Node });
+
+					return FullName;
+				};
+
+				NodeCode = TEXT("\n// Nodes\n");
+				NodeCode += TEXT(".Add<FDC::FNode>(\n{\n");
+
+#if WITH_EDITORONLY_DATA
+				InputNodeCode = TEXT("\n// Input (Template) Nodes\n");
+				InputNodeCode += TEXT(".Add<FDC::FInputNode>(\n{\n");
+#endif // WITH_EDITORONLY_DATA
+
+				bool bHasNodes = false;
+				bool bHasInputNodes = false;
+				FMetaSoundAssetManager& AssetManager = FMetaSoundAssetManager::GetChecked();
+				for (const FMetasoundFrontendNode& Node : Layer.Graph.Nodes)
+				{
+					if (const FConfigClassInfo* ClassInfo = ClassIDToInfo.Find(Node.ClassID))
+					{
+						const FMetasoundFrontendClass* Class = ClassInfo->Class;
+						check(Class);
+
+						FVector2D Location = FVector2D::Zero();
+						if (!Node.Style.Display.Locations.IsEmpty())
+						{
+							Location = Node.Style.Display.Locations.Array().Last().Value;
+						}
+
+						const FMetasoundFrontendClassName& ClassName = Class->Metadata.GetClassName();
+						const EMetasoundFrontendClassType ClassType = Class->Metadata.GetType();
+						switch (ClassType)
+						{
+							case EMetasoundFrontendClassType::Template:
+							{
+								if (ClassName != FInputNodeTemplate::GetChecked().ClassName)
+								{
+									UE_LOG(LogMetaSound, Error,
+										TEXT("Template node of class '%s' currently not supported for configuration reflection")
+										, *Class->Metadata.GetClassName().ToString());
+									continue;
+								}
+
+								if (const FName* InputName = TemplateNodeIDToInputName.Find(Node.GetID()))
+								{
+									bHasInputNodes = true;
+									InputNodeCode.Appendf(
+										TEXT("\t{ .NodeName = \"%s\", .InputName = \"%s\", .Location = FVector2D(%sf, %sf) },\n"),
+										*CacheNodeInfo(InputName->ToString(), Node),
+										*InputName->ToString(),
+										*FString::SanitizeFloat(Location.X),
+										*FString::SanitizeFloat(Location.Y)
+									);
+								}
+							}
+							break;
+
+							case EMetasoundFrontendClassType::External:
+							{
+								FString ClassString;
+								FString BaseName;
+								if (AssetManager.IsAssetClass(Class->Metadata))
+								{
+									const FTopLevelAssetPath AssetPath = AssetManager.FindAssetPath(FMetaSoundAssetKey(Class->Metadata));
+									check(AssetPath.IsValid());
+
+									BaseName = AssetPath.GetAssetName().ToString();
+									ClassString = FString::Printf(TEXT(".AssetPath = FTopLevelAssetPath(TEXT(\"%s\"))"), *AssetPath.ToString());
+								}
+								else
+								{
+									BaseName = ClassName.Name.ToString();
+									ClassString = FString::Printf(TEXT(".ClassName = { \"%s\", \"%s\", \"%s\" }"),
+										*ClassName.Namespace.ToString(),
+										*ClassName.Name.ToString(),
+										*ClassName.Variant.ToString()
+									);
+								}
+
+								bHasNodes = true;
+								NodeCode.Appendf(
+									TEXT("\t{ .Name = \"%s\", %s, .Location = FVector2D(%sf, %sf) },\n"),
+									*CacheNodeInfo(BaseName, Node),
+									*ClassString,
+									*FString::SanitizeFloat(Location.X),
+									*FString::SanitizeFloat(Location.Y)
+								);
+								break;
+							}
+
+							case EMetasoundFrontendClassType::Input:
+							case EMetasoundFrontendClassType::Output:
+							{
+								FString NodeName = FString::Printf(TEXT("%s_%s"), LexToString(ClassType), *Node.Name.ToString());
+								// Don't reflect as code as input/output is included in either a. the analogous
+								// reflection call OR reflecting associated interfaces. However, node info is
+								// cached for use when reflecting inputs/outputs.
+								CacheNodeInfo(NodeName, Node);
+							}
+							break;
+
+							default:
+							{
+								break;
+							}
+						}
+					}
+				}
+
+				auto FinishNodeCode = [](bool bHadContent, FStringBuilderBase& CodeBuilder)
+				{
+					if (bHadContent)
+					{
+						CodeBuilder += TEXT("})");
+					}
+					else
+					{
+						CodeBuilder = { };
+					}
+				};
+
+				FinishNodeCode(bHasNodes, NodeCode);
+				FinishNodeCode(bHasInputNodes, InputNodeCode);
+			}
+
+			void ReflectOutputs()
+			{
+				if (Layer.Outputs.IsEmpty())
+				{
+					return;
+				}
+
+				OutputCode = TEXT("\n// Outputs\n");
+				OutputCode += TEXT(".Add<FDC::FOutput>(\n{\n");
+
+				FMetaSoundAssetManager& AssetManager = FMetaSoundAssetManager::GetChecked();
+				for (const FMetasoundFrontendClassOutput& Output : Layer.Outputs)
+				{
+					FVector2D Location = FVector2D::Zero();
+
+					if (const FConfigNodeInfo* NodeInfo = NodeIDToInfo.Find(Output.NodeID))
+					{
+						const FMetasoundFrontendNode* Node = NodeInfo->Node;
+						check(Node);
+
+						if (!Node->Style.Display.Locations.IsEmpty())
+						{
+							Location = Node->Style.Display.Locations.Array().Last().Value;
+						}
+
+						OutputCode.Appendf(
+							TEXT("\t{ .Name = \"%s\", .DataType = \"%s\", .Location = FVector2D(%sf, %sf), .bIsConstructorOutput = %s }\n"),
+							*Output.Name.ToString(),
+							*Output.TypeName.ToString(),
+							*FString::SanitizeFloat(FMath::RoundHalfToEven(Location.X)),
+							*FString::SanitizeFloat(FMath::RoundHalfToEven(Location.Y)),
+							Output.AccessType == EMetasoundFrontendVertexAccessType::Value ? TEXT("true") : TEXT("false")
+						);
+					}
+				}
+
+				OutputCode += TEXT("})");
+			}
+
+		public:
+			FConfigReflector(FMetasoundFrontendGraphLayer InLayer)
+				: Layer(MoveTemp(InLayer))
+			{
+			}
+
+			FString Execute()
+			{
+				CacheDependencies();
+
+#if WITH_EDITORONLY_DATA
+				CacheInputNodes();
+#endif // WITH_EDITORONLY_DATA
+
+				// Strips inputs/outputs found to be registered to interfaces and directly adds applicable interfaces
+				ReflectInterfaces();
+
+				// Nodes reflected first to cache node data used by input/output reflection
+				ReflectNodes();
+
+				ReflectInputs();
+				ReflectOutputs();
+				ReflectEdges();
+
+				ReflectNodeInputDefaults();
+
+				TStringBuilder<0> LayerCode;
+				auto AddSection = [&LayerCode](FStringBuilderBase& SectionCode, bool bNextLine = true)
+				{
+					FStringView SectionView = SectionCode.ToView();
+					if (!SectionView.IsEmpty())
+					{
+						LayerCode.Append(SectionView);
+						if (bNextLine)
+						{
+							LayerCode.AppendChar('\n');
+						}
+					}
+					SectionCode = { };
+				};
+
+				LayerCode.Append(TEXT("using namespace Metasound::Frontend;\n"));
+				LayerCode.Append(TEXT("using FDC = Metasound::Engine::FDocumentConfigurator;\n"));
+				LayerCode.Append(TEXT("FDC Configurator(FDC::FArgs { .Builder = &OutBuilder });\n"));
+				LayerCode.Append(TEXT("return Configurator\n"));
+
+				AddSection(InterfaceCode);
+				AddSection(InputCode);
+				AddSection(OutputCode);
+
+				AddSection(InputNodeCode);
+				AddSection(NodeCode);
+				AddSection(NodeInputDefaultsCode);
+
+				AddSection(EdgesCode);
+				AddSection(OutputEdgesCode);
+
+				LayerCode.Append(TEXT(".Succeeded();\n"));
+				return FString(LayerCode);
+			}
+		};
+	} // namespace ConfigCodeReflectorPrivate
+
+	FDocumentConfigurator::FDocumentConfigurator(const FDocumentConfigurator::FArgs& Args)
+#if WITH_EDITORONLY_DATA
+		: Offset(Args.Offset)
+#endif // WITH_EDITORONLY_DATA
+	{
+		checkf(Args.Builder, TEXT("Builder must be supplied via document configurator arguments"));
+
+		UObject& MetaSoundObject = Args.Builder->CastDocumentObjectChecked<UObject>();
+		Builder.Reset(&FDocumentBuilderRegistry::GetChecked().FindOrBeginBuilding(MetaSoundObject));
+
+		checkf(Args.Builder == &Builder->GetConstBuilder(), TEXT("Builder provided in arguments differs from that registered with the DocumentBuilderRegistry."));
+
+		if (Args.bResetDoc)
+		{
+			ResetDocument();
+		}
+	}
+
+	FDocumentConfigurator::~FDocumentConfigurator()
+	{
+		// Inject result is ignored because 1. the configurator has already reported success if its destroying,
+		// and 2. the inject can fail if there are no inputs which is technically a valid situation.
+		EMetaSoundBuilderResult InjectResult = EMetaSoundBuilderResult::Failed;
+		constexpr bool bForceNodeCreation = false;
+		Builder->InjectInputTemplateNodes(bForceNodeCreation, InjectResult);
+	}
+
+	FDocumentConfigurator& FDocumentConfigurator::Add(FDocumentConfigurator::FInterface Interface)
+	{
+		const bool bSucceeded = Builder->GetBuilder().ModifyInterfaces(Frontend::FModifyInterfaceOptions(
+			{ }, // None removed
+			{ Interface.Version }
+		));
+
+		if (Result == EMetaSoundBuilderResult::Succeeded)
+		{
+			Result = bSucceeded ? EMetaSoundBuilderResult::Succeeded : EMetaSoundBuilderResult::Failed;
+		}
+
+		return *this;
+	}
+
+	FDocumentConfigurator& FDocumentConfigurator::Add(FDocumentConfigurator::FInput Input)
+	{
+		using namespace Frontend;
+
+		EMetaSoundBuilderResult StepResult = EMetaSoundBuilderResult::Failed;
+
+		FMetaSoundBuilderNodeOutputHandle OutputHandle = Builder->AddGraphInputNode(Input.Name, Input.DataType, MoveTemp(Input.DefaultValue), StepResult, Input.bIsConstructorInput);
+		if (StepResult == EMetaSoundBuilderResult::Succeeded)
+		{
+			Builder->SetNodeLocation(OutputHandle.NodeID, Input.Location + Offset, StepResult);
+			if (StepResult == EMetaSoundBuilderResult::Failed)
+			{
+				Result = StepResult;
+			}
+		}
+		else
+		{
+			Result = EMetaSoundBuilderResult::Failed;
+		}
+
+		return *this;
+	}
+
+#if WITH_EDITORONLY_DATA
+	FDocumentConfigurator& FDocumentConfigurator::Add(FDocumentConfigurator::FInputNode InputNode)
+	{
+		EMetaSoundBuilderResult StepResult = EMetaSoundBuilderResult::Failed;
+
+		FMetaSoundNodeHandle NodeHandle;
+		if (NodeNameToHandle.Contains(InputNode.NodeName))
+		{
+			UE_LOG(LogMetaSound, Error,
+				TEXT("Node with name '%s' already exists. Failed to add node instance to MetaSound graph"),
+				*InputNode.InputName.ToString());
+			StepResult = EMetaSoundBuilderResult::Failed;
+		}
+		else
+		{
+			using namespace Frontend;
+
+			if (const INodeTemplate* Template = INodeTemplateRegistry::Get().FindTemplate(FInputNodeTemplate::ClassName); ensure(Template))
+			{
+				if (const FMetasoundFrontendClassInput* Input = Builder->GetConstBuilder().FindGraphInput(InputNode.InputName))
+				{
+					FNodeTemplateGenerateInterfaceParams Params { .InputsToConnect = { Input->TypeName }, .OutputsToConnect = { Input->TypeName } };
+					if (const FMetasoundFrontendNode* TemplateNode = Builder->GetBuilder().AddNodeByTemplate(*Template, MoveTemp(Params)))
+					{
+						Builder->ConnectGraphInputToNode(InputNode.InputName, FMetaSoundNodeHandle(TemplateNode->GetID()), "Value", StepResult);
+						if (StepResult == EMetaSoundBuilderResult::Succeeded)
+						{
+							NodeHandle.NodeID = TemplateNode->GetID();
+							NodeNameToHandle.Add(InputNode.NodeName, NodeHandle);
+							Builder->SetNodeLocation(NodeHandle, InputNode.Location + Offset, StepResult);
+						}
+					}
+				}
+				else
+				{
+					UE_LOG(LogMetaSound, Error,
+						TEXT("Input Node could not be added: Input '%s' not found (must add input first)"),
+						*InputNode.InputName.ToString());
+				}
+			}
+		}
+
+		if (Result == EMetaSoundBuilderResult::Succeeded)
+		{
+			Result = StepResult;
+		}
+
+		return *this;
+	}
+#endif // WITH_EDITORONLY_DATA
+
+	FDocumentConfigurator& FDocumentConfigurator::Add(FDocumentConfigurator::FOutput Output)
+	{
+		using namespace Frontend;
+
+		EMetaSoundBuilderResult StepResult = EMetaSoundBuilderResult::Failed;
+
+		Builder->AddGraphOutputNode(Output.Name, Output.DataType, MoveTemp(Output.DefaultValue), StepResult, Output.bIsConstructorOutput);
+		if (StepResult == EMetaSoundBuilderResult::Failed)
+		{
+			Result = StepResult;
+		}
+
+		return *this;
+	}
+
+	FDocumentConfigurator& FDocumentConfigurator::Add(FDocumentConfigurator::FNode Node)
+	{
+		EMetaSoundBuilderResult StepResult = EMetaSoundBuilderResult::Failed;
+
+		FMetaSoundNodeHandle NodeHandle;
+		if (NodeNameToHandle.Contains(Node.Name))
+		{
+			UE_LOG(LogMetaSound, Error, TEXT("Node with name '%s' already exists. Failed to add node instance to MetaSound graph"), *Node.Name.ToString());
+			Result = EMetaSoundBuilderResult::Failed;
+		}
+		else if (Node.ClassName.IsValid())
+		{
+			NodeHandle = Builder->AddNodeByClassName(Node.ClassName, StepResult);
+		}
+		else if (Node.AssetPath.IsValid())
+		{
+			const FSoftObjectPath MetaSoundPath(Node.AssetPath);
+			if (TScriptInterface<IMetaSoundDocumentInterface> DocInterface = MetaSoundPath.TryLoad(); DocInterface.GetObject())
+			{
+				NodeHandle = Builder->AddNode(DocInterface, StepResult);
+			}
+		}
+
+		if (StepResult == EMetaSoundBuilderResult::Succeeded)
+		{
+			NodeNameToHandle.Add(Node.Name, NodeHandle);
+			Builder->SetNodeLocation(NodeHandle, Node.Location + Offset, StepResult);
+		}
+
+		if (Result == EMetaSoundBuilderResult::Succeeded)
+		{
+			Result = StepResult;
+		}
+
+		return *this;
+	}
+
+	FDocumentConfigurator& FDocumentConfigurator::Add(FDocumentConfigurator::FEdge Edge)
+	{
+		EMetaSoundBuilderResult StepResult = EMetaSoundBuilderResult::Failed;
+
+		if (const FMetaSoundNodeHandle* OutputNodeHandle = NodeNameToHandle.Find(Edge.OutputNode))
+		{
+			if (const FMetaSoundNodeHandle* InputNodeHandle = NodeNameToHandle.Find(Edge.InputNode))
+			{
+				Builder->ConnectNodes(*OutputNodeHandle, Edge.Output, *InputNodeHandle, Edge.Input, StepResult);
+			}
+			else
+			{
+				UE_LOG(LogMetaSound, Error, TEXT("Failed to find input node with name '%s' to connect"), *Edge.InputNode.ToString());
+				StepResult = EMetaSoundBuilderResult::Failed;
+			}
+		}
+		else
+		{
+			UE_LOG(LogMetaSound, Error, TEXT("Failed to find output node with name '%s' to connect"), *Edge.OutputNode.ToString());
+			StepResult = EMetaSoundBuilderResult::Failed;
+		}
+
+		if (Result == EMetaSoundBuilderResult::Succeeded)
+		{
+			Result = StepResult;
+		}
+
+		return *this;
+	}
+
+	FDocumentConfigurator& FDocumentConfigurator::Add(FDocumentConfigurator::FEdgeFromInput Edge)
+	{
+		EMetaSoundBuilderResult StepResult = EMetaSoundBuilderResult::Failed;
+
+		if (const FMetaSoundNodeHandle* OutputNodeHandle = NodeNameToHandle.Find(Edge.Node))
+		{
+			Builder->ConnectGraphInputToNode(Edge.GraphInput, *OutputNodeHandle, Edge.NodeInput, StepResult);
+			if (Result == EMetaSoundBuilderResult::Succeeded)
+			{
+				Result = StepResult;
+			}
+		}
+		else
+		{
+			UE_LOG(LogMetaSound, Error, TEXT("Failed to find node with name '%s': Graph input to node input connection failed"), *Edge.Node.ToString());
+			StepResult = EMetaSoundBuilderResult::Failed;
+		}
+
+		if (Result == EMetaSoundBuilderResult::Succeeded)
+		{
+			Result = StepResult;
+		}
+
+		return *this;
+	}
+
+	FDocumentConfigurator& FDocumentConfigurator::Add(FDocumentConfigurator::FEdgeToOutput Edge)
+	{
+		EMetaSoundBuilderResult StepResult = EMetaSoundBuilderResult::Failed;
+
+		if (const FMetaSoundNodeHandle* OutputNodeHandle = NodeNameToHandle.Find(Edge.Node))
+		{
+			Builder->ConnectNodeToGraphOutput(*OutputNodeHandle, Edge.NodeOutput, Edge.GraphOutput, StepResult);
+			if (Result == EMetaSoundBuilderResult::Succeeded)
+			{
+				Result = StepResult;
+			}
+		}
+		else
+		{
+			UE_LOG(LogMetaSound, Error, TEXT("Failed to find node with name '%s': Graph output to node intput connection failed"), *Edge.Node.ToString());
+			StepResult = EMetaSoundBuilderResult::Failed;
+		}
+
+		if (Result == EMetaSoundBuilderResult::Succeeded)
+		{
+			Result = StepResult;
+		}
+
+		return *this;
+	}
+
+	FDocumentConfigurator& FDocumentConfigurator::AddInternal(TArray<FDocumentConfigurator::FInterface> Interfaces)
+	{
+		TArray<FMetasoundFrontendVersion> Versions;
+		Algo::Transform(Interfaces, Versions, [](const FInterface& Interface) { return Interface.Version; });
+		const bool bSucceeded = Builder->GetBuilder().ModifyInterfaces(Frontend::FModifyInterfaceOptions(
+			{ }, // None removed
+			{ Versions }
+		));
+
+		if (!bSucceeded)
+		{
+			Result = EMetaSoundBuilderResult::Failed;
+		}
+
+		return *this;
+	}
+
+	UMetaSoundBuilderBase& FDocumentConfigurator::GetBuilder()
+	{
+		return *Builder.Get();
+	}
+	
+#if WITH_EDITORONLY_DATA
+	void FDocumentConfigurator::Reflect(FMetasoundFrontendGraphLayer InLayer, FString& OutCode)
+	{
+		using namespace ConfigCodeReflectorPrivate;
+		FConfigReflector Reflector(MoveTemp(InLayer));
+		OutCode = Reflector.Execute();
+	}
+#endif // WITH_EDITORONLY_DATA
+
+	FDocumentConfigurator& FDocumentConfigurator::ResetDocument()
+	{
+		using namespace Frontend;
+
+		FMetaSoundFrontendDocumentBuilder& FrontendBuilder = Builder->GetBuilder();
+		TSharedRef<FDocumentModifyDelegates> ModifyDelegates = MakeShared<FDocumentModifyDelegates>(FrontendBuilder.GetDocumentDelegates());
+		FrontendBuilder.ClearDocument(ModifyDelegates);
+
+		FrontendBuilder.SetIsGraphEditable(false);
+
+		return *this;
+	}
+
+	FDocumentConfigurator& FDocumentConfigurator::Set(FDocumentConfigurator::FInputLocation Location)
+	{
+		EMetaSoundBuilderResult StepResult = EMetaSoundBuilderResult::Failed;
+		if (const FMetaSoundNodeHandle InputNode = Builder->FindGraphInputNode(Location.Name, StepResult); StepResult == EMetaSoundBuilderResult::Succeeded)
+		{
+			Builder->SetNodeLocation(InputNode, Location.Location + Offset, StepResult);
+		}
+
+		if (Result == EMetaSoundBuilderResult::Succeeded)
+		{
+			Result = StepResult;
+		}
+
+		return *this;
+	}
+
+	FDocumentConfigurator& FDocumentConfigurator::Set(FDocumentConfigurator::FNodeInputDefault Input)
+	{
+		EMetaSoundBuilderResult StepResult = EMetaSoundBuilderResult::Failed;
+		if (const FMetaSoundNodeHandle* NodeHandle = NodeNameToHandle.Find(Input.Node))
+		{
+			FMetaSoundBuilderNodeInputHandle InputHandle = Builder->FindNodeInputByName(*NodeHandle, Input.Input, StepResult);
+			if (StepResult == EMetaSoundBuilderResult::Succeeded)
+			{
+				Builder->SetNodeInputDefault(InputHandle, Input.Value, StepResult);
+			}
+		}
+
+		if (Result == EMetaSoundBuilderResult::Succeeded)
+		{
+			Result = StepResult;
+		}
+
+		return *this;
+	}
+
+	bool FDocumentConfigurator::Succeeded() const
+	{
+		return Result == EMetaSoundBuilderResult::Succeeded;
+	}
+} // namespace Metasound::Engine
+#endif // WITH_EDITORONLY_DATA
diff --git a/Engine/Plugins/Runtime/Metasound/Source/MetasoundEngine/Private/DocumentTemplates/MetasoundRandomizerTemplate.cpp b/Engine/Plugins/Runtime/Metasound/Source/MetasoundEngine/Private/DocumentTemplates/MetasoundRandomizerTemplate.cpp
new file mode 100644
--- /dev/null
+++ b/Engine/Plugins/Runtime/Metasound/Source/MetasoundEngine/Private/DocumentTemplates/MetasoundRandomizerTemplate.cpp
@@ -0,0 +1,323 @@
+// Copyright Epic Games, Inc. All Rights Reserved.
+
+#include "DocumentTemplates/MetasoundRandomizerTemplate.h"
+
+#include "Algo/MaxElement.h"
+#include "Algo/Transform.h"
+#include "Containers/ContainersFwd.h"
+#include "DocumentTemplates/MetasoundDocumentConfigurator.h"
+#include "Interfaces/MetasoundOutputFormatInterfaces.h"
+#include "MetasoundBuilderSubsystem.h"
+#include "MetasoundDataTypeGetTraits.h"
+#include "MetasoundDocumentInterface.h"
+#include "MetasoundEngineNodesNames.h"
+#include "MetasoundFrontendDocument.h"
+#include "MetasoundFrontendDocumentBuilder.h"
+#include "MetasoundFrontendDocumentBuilderRegistry.h"
+#include "MetasoundFrontendDocumentIdGenerator.h"
+#include "MetasoundWave.h"
+#include "Sound/SoundWave.h"
+#include "UObject/SoftObjectPath.h"
+
+
+#define LOCTEXT_NAMESPACE "MetaSound"
+
+namespace Metasound::Engine
+{
+	namespace RandomizerPrivate
+	{
+		const FMetasoundFrontendNode* CreateWaveInput(FMetaSoundFrontendDocumentBuilder& OutBuilder, const TArray<TObjectPtr<USoundWave>>& InSounds)
+		{
+			using namespace Frontend;
+
+			const FMetasoundFrontendDocument& Doc = OutBuilder.GetConstDocumentChecked();
+
+			FMetasoundFrontendClassInput WaveInput;
+			WaveInput.Name = "Sounds";
+			WaveInput.TypeName = GetMetasoundDataTypeName<TArray<FWaveAsset>>();
+			WaveInput.NodeID = FDocumentIDGenerator::Get().CreateNodeID(Doc);
+			WaveInput.VertexID = FDocumentIDGenerator::Get().CreateVertexID(Doc);
+			WaveInput.AccessType = EMetasoundFrontendVertexAccessType::Value;
+
+			TArray<UObject*> SoundObjects;
+			Algo::TransformIf(InSounds, SoundObjects,
+				[](const TObjectPtr<USoundWave> Sound) { return Sound != nullptr; },
+				[](const TObjectPtr<USoundWave> Sound) { return Sound.Get(); });
+
+			FMetasoundFrontendLiteral DefaultValue;
+			DefaultValue.Set(MoveTemp(SoundObjects));
+			WaveInput.InitDefault(MoveTemp(DefaultValue));
+
+			return OutBuilder.AddGraphInput(MoveTemp(WaveInput));
+		}
+
+		int32 GetMaxChannels(const TArray<TObjectPtr<USoundWave>>& InSounds)
+		{
+			auto GetNumChannels = [](const TObjectPtr<USoundWave>& Sound)
+			{
+				if (const USoundWave* SoundPtr = Sound.Get())
+				{
+					return SoundPtr->NumChannels;
+				}
+
+				return 0;
+			};
+
+			if (const TObjectPtr<USoundWave>* MaxElement = Algo::MaxElementBy(InSounds, GetNumChannels))
+			{
+				return GetNumChannels(*MaxElement);
+			}
+
+			return 0;
+		}
+
+		FMetasoundFrontendClassName GetPlayerClassName(int32 NumChannels)
+		{
+			FMetasoundFrontendClassName ClassName(EngineNodes::Namespace, "Wave Player");
+			if (NumChannels == 2)
+			{
+				ClassName.Variant = EngineNodes::StereoVariant;
+			}
+			else if (NumChannels == 4)
+			{
+				ClassName.Variant = EngineNodes::QuadVariant;
+			}
+			else if (NumChannels == 6)
+			{
+				ClassName.Variant = EngineNodes::FiveDotOneVariant;
+			}
+			else if (NumChannels == 8)
+			{
+				ClassName.Variant = EngineNodes::SevenDotOneVariant;
+			}
+			else
+			{
+				ClassName.Variant = EngineNodes::MonoVariant;
+
+				if (NumChannels != 1)
+				{
+					UE_LOG(LogMetaSound, Warning, TEXT("SoundWave format with %i NumChannels not recognized: using Mono MetaSound WavePlayer"), NumChannels);
+				}
+			}
+			return ClassName;
+		}
+
+		FMetasoundFrontendVersion GetOutputVersion(int32 NumChannels)
+		{
+			if (NumChannels == 2)
+			{
+				return OutputFormatStereoInterface::GetVersion();
+			}
+
+			if (NumChannels == 4)
+			{
+				return OutputFormatQuadInterface::GetVersion();
+			}
+
+			if (NumChannels == 6)
+			{
+				return OutputFormatFiveDotOneInterface::GetVersion();
+			}
+
+			if (NumChannels == 8)
+			{
+				return OutputFormatSevenDotOneInterface::GetVersion();
+			}
+
+			return OutputFormatMonoInterface::GetVersion();
+		}
+
+		TArray<FDocumentConfigurator::FEdgeToOutput> GetAudioOutputEdges(int32 NumChannels)
+		{
+			using FDC = FDocumentConfigurator;
+
+			if (NumChannels == 2)
+			{
+				return
+				{
+					FDC::FEdgeToOutput { .Node = "WavePlayer", .NodeOutput = "Out Left", .GraphOutput = OutputFormatStereoInterface::Outputs::LeftOut },
+					FDC::FEdgeToOutput { .Node = "WavePlayer", .NodeOutput = "Out Right", .GraphOutput = OutputFormatStereoInterface::Outputs::RightOut }
+				};
+			}
+
+			if (NumChannels == 4)
+			{
+				return
+				{
+					FDC::FEdgeToOutput { .Node = "WavePlayer", .NodeOutput = "Out Front Left", .GraphOutput = OutputFormatQuadInterface::Outputs::FrontLeftOut },
+					FDC::FEdgeToOutput { .Node = "WavePlayer", .NodeOutput = "Out Front Right", .GraphOutput = OutputFormatQuadInterface::Outputs::FrontRightOut },
+					FDC::FEdgeToOutput { .Node = "WavePlayer", .NodeOutput = "Out Side Left", .GraphOutput = OutputFormatQuadInterface::Outputs::SideLeftOut },
+					FDC::FEdgeToOutput { .Node = "WavePlayer", .NodeOutput = "Out Side Right", .GraphOutput = OutputFormatQuadInterface::Outputs::SideRightOut }
+				};
+			}
+
+			if (NumChannels == 6)
+			{
+				return
+				{
+					FDC::FEdgeToOutput { .Node = "WavePlayer", .NodeOutput = "Out Front Left", .GraphOutput = OutputFormatFiveDotOneInterface::Outputs::FrontLeftOut },
+					FDC::FEdgeToOutput { .Node = "WavePlayer", .NodeOutput = "Out Front Right", .GraphOutput = OutputFormatFiveDotOneInterface::Outputs::FrontRightOut },
+					FDC::FEdgeToOutput { .Node = "WavePlayer", .NodeOutput = "Out Front Center", .GraphOutput = OutputFormatFiveDotOneInterface::Outputs::FrontCenterOut },
+					FDC::FEdgeToOutput { .Node = "WavePlayer", .NodeOutput = "Out Low Frequency", .GraphOutput = OutputFormatFiveDotOneInterface::Outputs::LowFrequencyOut },
+					FDC::FEdgeToOutput { .Node = "WavePlayer", .NodeOutput = "Out Side Left", .GraphOutput = OutputFormatFiveDotOneInterface::Outputs::SideLeftOut },
+					FDC::FEdgeToOutput { .Node = "WavePlayer", .NodeOutput = "Out Side Right", .GraphOutput = OutputFormatFiveDotOneInterface::Outputs::SideRightOut }
+				};
+			}
+
+			if (NumChannels == 8)
+			{
+				return
+				{
+					FDC::FEdgeToOutput { .Node = "WavePlayer", .NodeOutput = "Out Front Left", .GraphOutput = OutputFormatSevenDotOneInterface::Outputs::FrontLeftOut },
+					FDC::FEdgeToOutput { .Node = "WavePlayer", .NodeOutput = "Out Front Right", .GraphOutput = OutputFormatSevenDotOneInterface::Outputs::FrontRightOut },
+					FDC::FEdgeToOutput { .Node = "WavePlayer", .NodeOutput = "Out Front Center", .GraphOutput = OutputFormatSevenDotOneInterface::Outputs::FrontCenterOut },
+					FDC::FEdgeToOutput { .Node = "WavePlayer", .NodeOutput = "Out Low Frequency", .GraphOutput = OutputFormatSevenDotOneInterface::Outputs::LowFrequencyOut },
+					FDC::FEdgeToOutput { .Node = "WavePlayer", .NodeOutput = "Out Side Left", .GraphOutput = OutputFormatSevenDotOneInterface::Outputs::SideLeftOut },
+					FDC::FEdgeToOutput { .Node = "WavePlayer", .NodeOutput = "Out Side Right", .GraphOutput = OutputFormatSevenDotOneInterface::Outputs::SideRightOut },
+					FDC::FEdgeToOutput { .Node = "WavePlayer", .NodeOutput = "Out Back Left", .GraphOutput = OutputFormatSevenDotOneInterface::Outputs::BackLeftOut },
+					FDC::FEdgeToOutput { .Node = "WavePlayer", .NodeOutput = "Out Back Right", .GraphOutput = OutputFormatSevenDotOneInterface::Outputs::BackRightOut }
+				};
+			}
+
+			if (NumChannels != 1)
+			{
+				UE_LOG(LogMetaSound, Warning, TEXT("SoundWave format with %i NumChannels not recognized: Connecting only first channel as mono."), NumChannels);
+			}
+
+			return
+			{
+				FDC::FEdgeToOutput { .Node = "WavePlayer", .NodeOutput = "Out Mono", .GraphOutput = OutputFormatMonoInterface::Outputs::MonoOut }
+			};
+		}
+	} // namespace RandomizerPrivate
+} // namespace Metasound::Engine
+
+#if WITH_EDITORONLY_DATA
+bool FMetaSoundRandomizerTemplate::ConfigureDocument(FMetaSoundFrontendDocumentBuilder& OutBuilder)
+{
+	using namespace Metasound;
+	using namespace Metasound::Engine;
+	using namespace Metasound::Frontend;
+
+	using FDC = FDocumentConfigurator;
+
+	// If max channels cannot be deduced, then return clear document and consider it valid as
+	// user may not be supplying waves for a reason (or may be amidst editing array of sounds).
+	const int32 MaxNumChannels = RandomizerPrivate::GetMaxChannels(Sounds);
+
+	FDC Configurator(FDC::FArgs { .Builder = &OutBuilder });
+
+	// Add interfaces first.  Done even if no sounds supplied to make sure resulting document is valid source.
+	Configurator.Add<FDC::FInterface>(
+	{
+		{ .Version = SourceInterface::GetVersion() },
+		{ .Version = RandomizerPrivate::GetOutputVersion(MaxNumChannels) }
+	});
+
+	if (bIsOneShot)
+	{
+		Configurator.Add(FDC::FInterface { .Version = SourceOneShotInterface::GetVersion() });
+	}
+
+	if (MaxNumChannels <= 0)
+	{
+		return true;
+	}
+
+	const FName WaveArrayTypeName = GetMetasoundDataTypeName<TArray<FWaveAsset>>();
+	const FMetasoundFrontendClassName PlayerClassName = RandomizerPrivate::GetPlayerClassName(MaxNumChannels);
+	const FMetasoundFrontendClassName RandomClassName = FMetasoundFrontendClassName("Array", "Random Get", WaveArrayTypeName);
+
+	Configurator
+
+	// Add Non-Interface Graph Inputs
+	.Add(FDC::FInput { .Name = "Sounds", .DataType = WaveArrayTypeName, .DefaultValue = MakeFrontendLiteral<USoundWave>(Sounds), .Location = FVector2D(-500, 250) })
+
+	// Set Interface Graph Input locations
+	.Set(FDC::FInputLocation { .Name = "UE.Source.OnPlay", .Location = FVector2D(-500, 0)})
+
+	// Add Nodes
+	.Add<FDC::FNode>(
+	{
+		{ .Name = "Pitch", .ClassName = FMetasoundFrontendClassName("UE", "RandomFloat"), .Location = FVector2D(250, 400) },
+		{ .Name = "RandomWave", .ClassName = RandomClassName, .Location = FVector2D(0, 0) },
+		{ .Name = "WavePlayer", .ClassName = PlayerClassName, .Location = FVector2D(500, 0) }
+	})
+
+	.Set(FDC::FNodeInputDefault { .Node = "Pitch", .Input = "Min", .Value = MakeFrontendLiteral(FMath::Min(Pitch.X, Pitch.Y)) })
+	.Set(FDC::FNodeInputDefault { .Node = "Pitch", .Input = "Max", .Value = MakeFrontendLiteral(FMath::Max(Pitch.X, Pitch.Y)) })
+
+	// Connect Graph Inputs to Nodes
+	.Add<FDC::FEdgeFromInput>(
+	{
+		{ .GraphInput = "Sounds",			.Node = "RandomWave", .NodeInput = "In Array" },
+		{ .GraphInput = "UE.Source.OnPlay", .Node = "RandomWave", .NodeInput = "Next" }
+	})
+
+	// Connect Nodes Together
+	.Add<FDC::FEdge>(
+	{
+		{ .OutputNode = "RandomWave", .Output = "On Next",	.InputNode = "WavePlayer", .Input = "Play" },
+		{ .OutputNode = "RandomWave", .Output = "Value",	.InputNode = "WavePlayer", .Input = "Wave Asset" },
+		{ .OutputNode = "Pitch",	  .Output = "Value",	.InputNode = "WavePlayer", .Input = "Pitch Shift" },
+	})
+
+	// Connect Nodes to Graph Outputs
+	.Add<FDC::FEdgeToOutput>(RandomizerPrivate::GetAudioOutputEdges(MaxNumChannels));
+
+	if (bIsOneShot)
+	{
+		Configurator.Add(FDC::FEdgeToOutput { .Node = "WavePlayer", .NodeOutput = "On Finished", .GraphOutput = SourceOneShotInterface::Outputs::OnFinished });
+	}
+	else
+	{
+		Configurator.Set(FDC::FNodeInputDefault { .Node = "WavePlayer", .Input = "Loop", .Value = MakeFrontendLiteral(true) });
+	}
+
+	return Configurator.Succeeded();
+}
+#endif // WITH_EDITORONLY_DATA
+
+#if WITH_EDITOR
+const FMetaSoundFrontendDocumentTemplate::FEditorOptions& FMetaSoundRandomizerTemplate::GetEditorOptions() const
+{
+	static const FEditorOptions RandomizerOptions
+	{
+		.AssetActionClasses = { USoundWave::StaticClass() },
+		.bTemplatePropertiesVisible = true
+	};
+	return RandomizerOptions;
+}
+
+void FMetaSoundRandomizerTemplate::OnAssetInitialized(TArray<UObject*> SelectedObjects, FMetaSoundFrontendDocumentBuilder& OutBuilder)
+{
+	Sounds.Reset();
+	Algo::TransformIf(SelectedObjects, Sounds,
+		[](const UObject* Object) { return Object->GetClass() == USoundWave::StaticClass(); }, // Has to match type exactly so procedural sounds, SourceBuses, MetaSounds, etc. aren't used
+		[](UObject* Object) { return TObjectPtr<USoundWave>(CastChecked<USoundWave>(Object)); });
+}
+
+void FMetaSoundRandomizerTemplate::OnPropertyChanged(const FPropertyChangedEvent& InEvent, FMetaSoundFrontendDocumentBuilder& OutBuilder)
+{
+	if (InEvent.Property && InEvent.Property->GetName() == GET_MEMBER_NAME_CHECKED(FMetaSoundRandomizerTemplate, Sounds))
+	{
+		RemoveInvalidSounds(OutBuilder, false /* bRemoveNullEntries */, Sounds);
+	}
+
+	Super::OnPropertyChanged(InEvent, OutBuilder);
+}
+
+bool FMetaSoundRandomizerTemplate::RemoveInvalidSounds(const FMetaSoundFrontendDocumentBuilder& ParentBuilder, bool bRemoveNullEntries, TArray<TObjectPtr<USoundWave>>& OutSounds)
+{
+	return OutSounds.RemoveAll([&bRemoveNullEntries](const TObjectPtr<USoundWave>& Wave)
+	{
+		if (bRemoveNullEntries && Wave == nullptr)
+		{
+			return true;
+		}
+
+		return Wave && Wave->GetClass() != USoundWave::StaticClass();
+	}) > 0;
+}
+#endif // WITH_EDITOR
+#undef LOCTEXT_NAMESPACE
diff --git a/Engine/Plugins/Runtime/Metasound/Source/MetasoundEngine/Private/MetasoundAssetSubsystem.cpp b/Engine/Plugins/Runtime/Metasound/Source/MetasoundEngine/Private/MetasoundAssetSubsystem.cpp
--- a/Engine/Plugins/Runtime/Metasound/Source/MetasoundEngine/Private/MetasoundAssetSubsystem.cpp
+++ b/Engine/Plugins/Runtime/Metasound/Source/MetasoundEngine/Private/MetasoundAssetSubsystem.cpp
@@ -5,7 +5,6 @@
 #include "Algo/Transform.h"
 #include "AssetRegistry/AssetRegistryModule.h"
 #include "Async/Async.h"
-#include "DocumentConfigurations/MetasoundFrontendPresetDocumentConfiguration.h"
 #include "Engine/AssetManager.h"
 #include "HAL/CriticalSection.h"
 #include "Metasound.h"
@@ -23,6 +22,7 @@
 #include "MetasoundUObjectRegistry.h"
 #include "Misc/CoreDelegates.h"
 #include "Modules/ModuleManager.h"
+#include "DocumentTemplates/MetasoundFrontendPresetTemplate.h"
 #include "UObject/NoExportTypes.h"
 
 
@@ -771,8 +771,8 @@ PRAGMA_ENABLE_DEPRECATION_WARNINGS
 		TScriptInterface<const IMetaSoundDocumentInterface> DocInterface = InAsset.GetOwningAsset();
 		check(DocInterface.GetObject());
 
-		const UStruct* Struct = DocInterface->GetConstDocument().Configuration.GetScriptStruct();
-		if (!Struct || !Struct->IsChildOf(FPresetDocumentConfiguration::StaticStruct()))
+		const UStruct* Struct = DocInterface->GetConstDocument().Template.GetScriptStruct();
+		if (!Struct || !Struct->IsChildOf(FPresetTemplate::StaticStruct()))
 		{
 			return false;
 		}
@@ -795,8 +795,8 @@ PRAGMA_ENABLE_DEPRECATION_WARNINGS
 		TScriptInterface<const IMetaSoundDocumentInterface> DocInterface = InAsset.GetOwningAsset();
 
 		// Base case (found the first non preset referenced asset)
-		const UStruct* Struct = DocInterface->GetConstDocument().Configuration.GetScriptStruct();
-		if (!Struct || !Struct->IsChildOf(FPresetDocumentConfiguration::StaticStruct()))
+		const UStruct* Struct = DocInterface->GetConstDocument().Template.GetScriptStruct();
+		if (!Struct || !Struct->IsChildOf(FPresetTemplate::StaticStruct()))
 		{
 			return true;
 		}
diff --git a/Engine/Plugins/Runtime/Metasound/Source/MetasoundEngine/Private/MetasoundBuilderBase.cpp b/Engine/Plugins/Runtime/Metasound/Source/MetasoundEngine/Private/MetasoundBuilderBase.cpp
--- a/Engine/Plugins/Runtime/Metasound/Source/MetasoundEngine/Private/MetasoundBuilderBase.cpp
+++ b/Engine/Plugins/Runtime/Metasound/Source/MetasoundEngine/Private/MetasoundBuilderBase.cpp
@@ -6,7 +6,6 @@
 #include "Algo/Transform.h"
 #include "AudioDevice.h"
 #include "Components/AudioComponent.h"
-#include "DocumentConfigurations/MetasoundFrontendPresetDocumentConfiguration.h"
 #include "Engine/Engine.h"
 #include "HAL/IConsoleManager.h"
 #include "Interfaces/MetasoundOutputFormatInterfaces.h"
@@ -32,6 +31,7 @@
 #include "MetasoundUObjectRegistry.h"
 #include "MetasoundVertex.h"
 #include "NodeTemplates/MetasoundFrontendNodeTemplateInput.h"
+#include "DocumentTemplates/MetasoundFrontendPresetTemplate.h"
 
 #include UE_INLINE_GENERATED_CPP_BY_NAME(MetasoundBuilderBase)
 
@@ -606,9 +606,9 @@ void UMetaSoundBuilderBase::ConnectGraphInputToNode(const FName& GraphInputName,
 void UMetaSoundBuilderBase::ConvertFromPreset(EMetaSoundBuilderResult& OutResult)
 {
 	using namespace Metasound::Frontend;
-	if (Builder.IsDocumentConfiguredBy<FPresetDocumentConfiguration>())
+	if (Builder.IsDocumentConfiguredBy<FPresetTemplate>())
 	{
-		Builder.SetDocumentConfiguration({ });
+		Builder.SetTemplate({ });
 		Builder.ConfigureDocument();
 
 #if WITH_EDITORONLY_DATA
@@ -662,10 +662,10 @@ void UMetaSoundBuilderBase::ConvertToPreset(const TScriptInterface<IMetaSoundDoc
 	TSharedRef<FDocumentModifyDelegates> DocumentDelegates = MakeShared<FDocumentModifyDelegates>(ParentInterface->GetConstDocument());
 	InitDelegates(*DocumentDelegates);
 
-	TInstancedStruct<FMetaSoundFrontendPresetDocumentConfiguration> NewPresetConfig = TInstancedStruct<FMetaSoundFrontendPresetDocumentConfiguration>::Make();
+	TInstancedStruct<FMetaSoundFrontendPresetTemplate> NewPresetConfig = TInstancedStruct<FMetaSoundFrontendPresetTemplate>::Make();
 	NewPresetConfig.GetMutable().Parent = ParentObject;
 
-	Builder.SetDocumentConfiguration(MoveTemp(NewPresetConfig));
+	Builder.SetTemplate(MoveTemp(NewPresetConfig));
 	OutResult = Builder.ConfigureDocument() ? EMetaSoundBuilderResult::Succeeded : EMetaSoundBuilderResult::Failed;
 }
 
@@ -1014,11 +1014,11 @@ int32 UMetaSoundBuilderBase::GetLastTransactionRegistered() const
 UObject* UMetaSoundBuilderBase::GetReferencedPresetAsset() const
 {
 	using namespace Metasound::Frontend;
-	const FPresetDocumentConfiguration* PresetConfig = GetConstBuilder()
-		.GetConstDocumentConfigurationAs<FPresetDocumentConfiguration>();
-	if (PresetConfig)
+	const FPresetTemplate* PresetTemplate = GetConstBuilder()
+		.GetConstTemplateAs<FPresetTemplate>();
+	if (PresetTemplate)
 	{
-		return PresetConfig->Parent.GetObject();
+		return PresetTemplate->Parent.GetObject();
 	}
 
 	return nullptr;
@@ -1073,7 +1073,7 @@ bool UMetaSoundBuilderBase::InterfaceIsDeclared(FName InterfaceName) const
 
 bool UMetaSoundBuilderBase::IsPreset() const
 {
-	return Builder.IsDocumentConfiguredBy<FMetaSoundFrontendPresetDocumentConfiguration>();
+	return Builder.IsDocumentConfiguredBy<FMetaSoundFrontendPresetTemplate>();
 }
 
 bool UMetaSoundBuilderBase::NodesAreConnected(const FMetaSoundBuilderNodeOutputHandle& OutputHandle, const FMetaSoundBuilderNodeInputHandle& InputHandle) const
@@ -1349,18 +1349,18 @@ void UMetaSoundBuilderBase::SetMemberMetadata(UMetaSoundFrontendMemberMetadata&
 }
 #endif // WITH_EDITOR
 
-void UMetaSoundBuilderBase::SetConfiguration(TInstancedStruct<FMetaSoundFrontendDocumentConfiguration> Configuration)
+void UMetaSoundBuilderBase::SetTemplate(TInstancedStruct<FMetaSoundFrontendDocumentTemplate> Configuration)
 {
-	Builder.SetDocumentConfiguration(MoveTemp(Configuration));
+	Builder.SetTemplate(MoveTemp(Configuration));
 
 PRAGMA_DISABLE_EXPERIMENTAL_WARNINGS
 	Builder.ConfigureDocument();
 PRAGMA_ENABLE_EXPERIMENTAL_WARNINGS
 }
 
-void UMetaSoundBuilderBase::SetConfigurationProperties(TFunctionRef<void(TStructView<FMetaSoundFrontendDocumentConfiguration>)> SetPropertyFunction)
+void UMetaSoundBuilderBase::SetTemplateProperties(TFunctionRef<void(TStructView<FMetaSoundFrontendDocumentTemplate>)> SetPropertyFunction)
 {
-	Builder.SetDocumentConfigurationProperties(SetPropertyFunction);
+	Builder.SetTemplateProperties(SetPropertyFunction);
 
 	PRAGMA_DISABLE_EXPERIMENTAL_WARNINGS
 	Builder.ConfigureDocument();
@@ -1439,19 +1439,19 @@ FMetasoundFrontendVersion UMetaSoundBuilderBase::FindNodeClassVersion(const FMet
 	return FMetasoundFrontendVersion::GetInvalid();
 }
 
-TInstancedStruct<FMetaSoundFrontendDocumentConfiguration> UMetaSoundBuilderBase::GetConfiguration(EMetaSoundBuilderResult& OutResult) const
+TInstancedStruct<FMetaSoundFrontendDocumentTemplate> UMetaSoundBuilderBase::GetDocumentTemplate(EMetaSoundBuilderResult& OutResult) const
 {
 	using namespace Metasound::Frontend;
 
 	OutResult = EMetaSoundBuilderResult::Failed;
-	TInstancedStruct<FMetaSoundFrontendDocumentConfiguration> NewInstance;
+	TInstancedStruct<FMetaSoundFrontendDocumentTemplate> NewInstance;
 
-	if (TConstStructView<FDocumentConfiguration> DocConfig = Builder.GetConstDocumentConfiguration(); DocConfig.IsValid())
+	if (TConstStructView<FMetaSoundTemplate> DocTemplate = Builder.GetConstDocumentTemplate(); DocTemplate.IsValid())
 	{
-		const UScriptStruct* ScriptStruct = DocConfig.GetScriptStruct();
+		const UScriptStruct* ScriptStruct = DocTemplate.GetScriptStruct();
 		check(ScriptStruct);
 
-		NewInstance.InitializeAsScriptStruct(ScriptStruct, DocConfig.GetMemory());
+		NewInstance.InitializeAsScriptStruct(ScriptStruct, DocTemplate.GetMemory());
 		OutResult = EMetaSoundBuilderResult::Succeeded;
 	}
 
diff --git a/Engine/Plugins/Runtime/Metasound/Source/MetasoundEngine/Private/MetasoundBuilderSubsystem.cpp b/Engine/Plugins/Runtime/Metasound/Source/MetasoundEngine/Private/MetasoundBuilderSubsystem.cpp
--- a/Engine/Plugins/Runtime/Metasound/Source/MetasoundEngine/Private/MetasoundBuilderSubsystem.cpp
+++ b/Engine/Plugins/Runtime/Metasound/Source/MetasoundEngine/Private/MetasoundBuilderSubsystem.cpp
@@ -6,7 +6,6 @@
 #include "AudioDevice.h"
 #include "Components/AudioComponent.h"
 #include "CoreGlobals.h"
-#include "DocumentConfigurations/MetasoundFrontendPresetDocumentConfiguration.h"
 #include "Engine/Engine.h"
 #include "HAL/IConsoleManager.h"
 #include "Interfaces/MetasoundOutputFormatInterfaces.h"
@@ -37,6 +36,7 @@
 #include "MetasoundUObjectRegistry.h"
 #include "MetasoundVertex.h"
 #include "NodeTemplates/MetasoundFrontendNodeTemplateInput.h"
+#include "DocumentTemplates/MetasoundFrontendPresetTemplate.h"
 #include "UObject/PerPlatformProperties.h"
 
 #include UE_INLINE_GENERATED_CPP_BY_NAME(MetasoundBuilderSubsystem)
@@ -354,7 +354,7 @@ void UMetaSoundSourceBuilder::InitDelegates(Metasound::Frontend::FDocumentModify
 
 	if (!IsRunningCookCommandlet())
 	{
-		OutDocumentDelegates.OnDocumentConfigurationChanged.AddUObject(this, &UMetaSoundSourceBuilder::OnDocumentConfigurationChanged);
+		OutDocumentDelegates.OnDocumentTemplateChanged.AddUObject(this, &UMetaSoundSourceBuilder::OnDocumentTemplateChanged);
 
 		OutDocumentDelegates.PageDelegates.OnPageAdded.AddUObject(this, &UMetaSoundSourceBuilder::OnPageAdded);
 		OutDocumentDelegates.PageDelegates.OnRemovingPage.AddUObject(this, &UMetaSoundSourceBuilder::OnRemovingPage);
@@ -438,7 +438,7 @@ void UMetaSoundSourceBuilder::OnAssetReferenceAdded(TScriptInterface<IMetaSoundD
 	Source.ReferencedAssetClassKeys.Add(RegistryKey.ToString());
 }
 
-void UMetaSoundSourceBuilder::OnDocumentConfigurationChanged(const Metasound::Frontend::FDocumentConfigChangedArgs&) const
+void UMetaSoundSourceBuilder::OnDocumentTemplateChanged(const Metasound::Frontend::FDocumentTemplateChangedArgs&) const
 {
 	using namespace Metasound::Frontend;
 
@@ -1008,8 +1008,8 @@ void UMetaSoundSourceBuilder::SetFormat(EMetaSoundOutputAudioFormat OutputFormat
 	using namespace Metasound::Engine;
 	using namespace Metasound::Frontend;
 
-	// Clear document configuration since interface data is being altered
-	Builder.SetDocumentConfiguration({ });
+	// Clear template since interface data is being altered
+	Builder.SetTemplate({ });
 	Builder.ConfigureDocument();
 
 	const FOutputAudioFormatInfoMap& FormatMap = GetOutputAudioFormatInfo();
@@ -1375,14 +1375,14 @@ UMetaSoundBuilderBase* UMetaSoundBuilderSubsystem::FindParentBuilderOfPreset(con
 		return nullptr;
 	}
 
-	const FMetaSoundFrontendPresetDocumentConfiguration* PresetConfig = Builder->GetConstBuilder().GetConstDocumentConfigurationAs<FMetaSoundFrontendPresetDocumentConfiguration>();
-	if (!PresetConfig)
+	const FMetaSoundFrontendPresetTemplate* PresetTemplate = Builder->GetConstBuilder().GetConstTemplateAs<FMetaSoundFrontendPresetTemplate>();
+	if (!PresetTemplate)
 	{
 		UE_LOG(LogMetaSound, Error, TEXT("MetaSound '%s' is not a preset."), *GetNameSafe(InMetaSoundPreset.GetObject()));
 		return nullptr;
 	}
 
-	TScriptInterface<IMetaSoundDocumentInterface> ParentAsset = PresetConfig->Parent;
+	TScriptInterface<IMetaSoundDocumentInterface> ParentAsset = PresetTemplate->Parent;
 	if (!ParentAsset)
 	{
 		UE_LOG(LogMetaSound, Error, TEXT("MetaSound preset '%s' contains null reference."), *GetNameSafe(InMetaSoundPreset.GetObject()));
@@ -1408,13 +1408,13 @@ UMetaSoundBuilderBase* UMetaSoundBuilderSubsystem::FindParentBuilderOfPreset(con
 			return nullptr;
 		}
 
-		PresetConfig = ParentBuilder->GetConstBuilder().GetConstDocumentConfigurationAs<FMetaSoundFrontendPresetDocumentConfiguration>();
-		if (!PresetConfig)
+		PresetTemplate = ParentBuilder->GetConstBuilder().GetConstTemplateAs<FMetaSoundFrontendPresetTemplate>();
+		if (!PresetTemplate)
 		{
 			break;
 		}
 
-		ParentAsset = PresetConfig->Parent;
+		ParentAsset = PresetTemplate->Parent;
 		if (!ParentAsset)
 		{
 			UE_LOG(LogMetaSound, Warning, TEXT("MetaSound preset '%s' contains null reference."), *GetNameSafe(ParentAsset.GetObject()));
diff --git a/Engine/Plugins/Runtime/Metasound/Source/MetasoundEngine/Private/MetasoundEngineAsset.h b/Engine/Plugins/Runtime/Metasound/Source/MetasoundEngine/Private/MetasoundEngineAsset.h
--- a/Engine/Plugins/Runtime/Metasound/Source/MetasoundEngine/Private/MetasoundEngineAsset.h
+++ b/Engine/Plugins/Runtime/Metasound/Source/MetasoundEngine/Private/MetasoundEngineAsset.h
@@ -2,7 +2,7 @@
 
 #pragma once
 
-#include "DocumentConfigurations/MetasoundFrontendDocumentConfiguration.h"
+#include "DocumentTemplates/MetasoundFrontendDocumentTemplate.h"
 #include "Metasound.h"
 #include "MetasoundAssetManager.h"
 #include "MetasoundAssetSubsystem.h"
@@ -87,7 +87,7 @@ namespace Metasound::Engine
 			FMetaSoundFrontendDocumentBuilder& Builder = BuilderRegistry.FindOrBeginBuilding(InDocInterface);
 
 #if WITH_EDITORONLY_DATA
-			Builder.GetDocumentDelegates().OnDocumentConfigurationChanged.Broadcast({ EDocumentConfigChangeType::UndoRedo });
+			Builder.GetDocumentDelegates().OnDocumentTemplateChanged.Broadcast({ EDocumentTemplateChangeType::UndoRedo });
 			Builder.ConfigureDocument();
 #endif // WITH_EDITORONLY_DATA
 
@@ -205,13 +205,13 @@ namespace Metasound::Engine
 				});
 			}
 
-			if (const FMetaSoundFrontendDocumentConfiguration* DocConfig = Document.Configuration.GetPtr())
+			if (const FMetaSoundFrontendDocumentTemplate* DocTemplate = Document.Template.GetPtr())
 			{
 				FDocumentBuilderRegistry& BuilderRegistry = FDocumentBuilderRegistry::GetChecked();
 				UObject* MetaSoundObject = FAssetData(&MetaSound).GetAsset();
 				check(MetaSoundObject);
 				const UMetaSoundBuilderBase& Builder = BuilderRegistry.FindOrBeginBuilding(*MetaSoundObject);
-				CombineDataValidationResults(Result, DocConfig->IsDataValid(Builder.GetConstBuilder(), InOutContext));
+				CombineDataValidationResults(Result, DocTemplate->IsDataValid(Builder.GetConstBuilder(), InOutContext));
 			}
 			return Result;
 		}
diff --git a/Engine/Plugins/Runtime/Metasound/Source/MetasoundEngine/Private/MetasoundSource.cpp b/Engine/Plugins/Runtime/Metasound/Source/MetasoundEngine/Private/MetasoundSource.cpp
--- a/Engine/Plugins/Runtime/Metasound/Source/MetasoundEngine/Private/MetasoundSource.cpp
+++ b/Engine/Plugins/Runtime/Metasound/Source/MetasoundEngine/Private/MetasoundSource.cpp
@@ -8,7 +8,6 @@
 #include "AssetRegistry/AssetRegistryModule.h"
 #include "AudioDeviceManager.h"
 #include "Containers/Ticker.h"
-#include "DocumentConfigurations/MetasoundFrontendPresetDocumentConfiguration.h"
 #include "IAudioParameterInterfaceRegistry.h"
 #include "Interfaces/ITargetPlatform.h"
 #include "Interfaces/MetasoundOutputFormatInterfaces.h"
@@ -45,6 +44,7 @@
 #include "MetasoundTrace.h"
 #include "MetasoundTrigger.h"
 #include "MetasoundUObjectRegistry.h"
+#include "DocumentTemplates/MetasoundFrontendPresetTemplate.h"
 #include "UObject/ObjectSaveContext.h"
 #include "UObject/ScriptInterface.h"
 
@@ -397,8 +397,8 @@ public:
 		using namespace Metasound;
 		using namespace Metasound::Frontend;
 
-		const FPresetDocumentConfiguration* PresetConfig = Document.Configuration.GetPtr<FPresetDocumentConfiguration>();
-		if (!PresetConfig)
+		const FPresetTemplate* PresetTemplate = Document.Template.GetPtr<FPresetTemplate>();
+		if (!PresetTemplate)
 		{
 			return;
 		}
@@ -425,7 +425,7 @@ public:
 				{
 					// Avoid adding parameters which are NOT overriding a default.
 					const FGuid& NodeID = ClassInputs[ClassIndex]->NodeID;
-					if (const FPresetVertexMetadata* VertexMetadata = PresetConfig->FindConstVertexMetadata<FPresetVertexMetadata>(NodeID))
+					if (const FPresetVertexMetadata* VertexMetadata = PresetTemplate->FindConstVertexMetadata<FPresetVertexMetadata>(NodeID))
 					{
 						if (VertexMetadata->bOverrideInheritedDefault)
 						{
@@ -1949,13 +1949,13 @@ void UMetaSoundSource::CacheRuntimeInputData(TArrayView<const FGuid> InPageOrder
 	// cached.
 
 	bIsPresetGraphInflationSupported = false;
-	if (const FPresetDocumentConfiguration* PresetConfig = RootMetasoundDocument.Configuration.GetPtr<FPresetDocumentConfiguration>())
+	if (const FPresetTemplate* PresetTemplate = RootMetasoundDocument.Template.GetPtr<FPresetTemplate>())
 	{
-		auto IsOverriddenConstructorInput = [this, &PresetConfig](const FMetasoundFrontendClassInput& Input)
+		auto IsOverriddenConstructorInput = [this, &PresetTemplate](const FMetasoundFrontendClassInput& Input)
 		{
 			if (const FRuntimeInput* RuntimeInput = RuntimeInputData.InputMap.Find(Input.Name))
 			{
-				if (const FPresetVertexMetadata* VertexMetadata = PresetConfig->FindConstVertexMetadata<FPresetVertexMetadata>(Input.NodeID))
+				if (const FPresetVertexMetadata* VertexMetadata = PresetTemplate->FindConstVertexMetadata<FPresetVertexMetadata>(Input.NodeID))
 				{
 					return VertexMetadata->bOverrideInheritedDefault && RuntimeInput->AccessType == EMetasoundFrontendVertexAccessType::Value;
 				}
diff --git a/Engine/Plugins/Runtime/Metasound/Source/MetasoundEngine/Public/DocumentConfigurations/MetasoundDocumentConfigurator.h b/Engine/Plugins/Runtime/Metasound/Source/MetasoundEngine/Public/DocumentConfigurations/MetasoundDocumentConfigurator.h
deleted file mode 100644
diff --git a/Engine/Plugins/Runtime/Metasound/Source/MetasoundEngine/Public/DocumentConfigurations/MetasoundRandomizerDocumentConfiguration.h b/Engine/Plugins/Runtime/Metasound/Source/MetasoundEngine/Public/DocumentConfigurations/MetasoundRandomizerDocumentConfiguration.h
deleted file mode 100644
diff --git a/Engine/Plugins/Runtime/Metasound/Source/MetasoundEngine/Public/DocumentTemplates/MetasoundDocumentConfigurator.h b/Engine/Plugins/Runtime/Metasound/Source/MetasoundEngine/Public/DocumentTemplates/MetasoundDocumentConfigurator.h
new file mode 100644
--- /dev/null
+++ b/Engine/Plugins/Runtime/Metasound/Source/MetasoundEngine/Public/DocumentTemplates/MetasoundDocumentConfigurator.h
@@ -0,0 +1,185 @@
+// Copyright Epic Games, Inc. All Rights Reserved.
+#pragma once
+
+#include "Algo/Transform.h"
+#include "Containers/ContainersFwd.h"
+#include "MetasoundBuilderBase.h"
+#include "MetasoundBuilderSubsystem.h"
+#include "MetasoundFrontendDocument.h"
+#include "MetasoundFrontendGraphLayer.h"
+#include "MetasoundFrontendLiteral.h"
+#include "UObject/ObjectPtr.h"
+
+
+#define UE_API METASOUNDENGINE_API
+
+
+namespace Metasound::Engine
+{
+	struct FDocumentConfigurator
+	{
+		struct FArgs
+		{
+			FMetaSoundFrontendDocumentBuilder* Builder = nullptr;
+
+#if WITH_EDITORONLY_DATA
+			// Offset to use when positioning new nodes via the configuration
+			FVector2D Offset = FVector2D::Zero();
+#endif // WITH_EDITORONLY_DATA
+
+			bool bResetDoc = true;
+		};
+
+		UE_API FDocumentConfigurator(const FArgs& Args);
+		UE_API ~FDocumentConfigurator();
+
+#if WITH_EDITORONLY_DATA
+		UE_EXPERIMENTAL(5.8, "Configurator code reflection is in development and may change implementation or be removed in the future")
+		UE_API static void Reflect(FMetasoundFrontendGraphLayer InLayer, FString& OutCode);
+#endif // WITH_EDITORONLY_DATA
+
+	public:
+		struct FInterface
+		{
+			FMetasoundFrontendVersion Version;
+		};
+		UE_EXPERIMENTAL(5.8, "Configurator is in active development and subject to change")
+		UE_API FDocumentConfigurator& Add(FInterface Interface);
+
+		struct FInput
+		{
+			FName Name = "FloatInput";
+			FName DataType = GetMetasoundDataTypeName<float>();
+			FMetasoundFrontendLiteral DefaultValue = Frontend::MakeFrontendLiteral(0.f);
+			FVector2D Location = FVector2D::Zero();
+			bool bIsConstructorInput = false;
+		};
+		UE_EXPERIMENTAL(5.8, "Configurator is in active development and subject to change")
+		UE_API FDocumentConfigurator& Add(FInput Input);
+
+#if WITH_EDITORONLY_DATA
+		struct FInputNode
+		{
+			FName NodeName = "FloatInputNode";
+			FName InputName = "FloatInput";
+			FVector2D Location = FVector2D::Zero();
+		};
+		UE_EXPERIMENTAL(5.8, "Configurator is in active development and subject to change")
+		UE_API FDocumentConfigurator& Add(FInputNode InputNode);
+#endif // WITH_EDITORONLY_DATA
+
+		struct FOutput
+		{
+			FName Name = "FloatOutput";
+			FName DataType = GetMetasoundDataTypeName<float>();
+			FMetasoundFrontendLiteral DefaultValue = Frontend::MakeFrontendLiteral(0.f);
+			FVector2D Location = FVector2D::Zero();
+			bool bIsConstructorOutput = false;
+		};
+		UE_EXPERIMENTAL(5.8, "Configurator is in active development and subject to change")
+		UE_API FDocumentConfigurator& Add(FOutput Output);
+
+		struct FNode
+		{
+			FName Name;
+			FMetasoundFrontendClassName ClassName;
+			FTopLevelAssetPath AssetPath;
+			FVector2D Location = FVector2D::Zero();
+		};
+		UE_EXPERIMENTAL(5.8, "Configurator is in active development and subject to change")
+		UE_API FDocumentConfigurator& Add(FNode Node);
+
+		struct FEdge
+		{
+			FName OutputNode;
+			FName Output;
+			FName InputNode;
+			FName Input;
+		};
+		UE_EXPERIMENTAL(5.8, "Configurator is in active development and subject to change")
+		UE_API FDocumentConfigurator& Add(FEdge Edge);
+
+		struct FEdgeFromInput
+		{
+			FName GraphInput;
+			FName Node;
+			FName NodeInput;
+		};
+		UE_EXPERIMENTAL(5.8, "Configurator is in active development and subject to change")
+		UE_API FDocumentConfigurator& Add(FEdgeFromInput Edge);
+
+		struct FEdgeToOutput
+		{
+			FName Node;
+			FName NodeOutput;
+			FName GraphOutput;
+		};
+		UE_EXPERIMENTAL(5.8, "Configurator is in active development and subject to change")
+		UE_API FDocumentConfigurator& Add(FEdgeToOutput Edge);
+
+		template <typename TemplateType>
+		FDocumentConfigurator& Add(TArray<TemplateType> Args)
+		{
+			for (TemplateType& TemplateArgs : Args)
+			{
+				Add(MoveTemp(TemplateArgs));
+			}
+
+			return *this;
+		}
+
+		template <>
+		FDocumentConfigurator& Add<FInterface>(TArray<FInterface> Interfaces)
+		{
+			return AddInternal(MoveTemp(Interfaces));
+		}
+
+		UE_EXPERIMENTAL(5.8, "Configurator is in active development and subject to change")
+		UE_API UMetaSoundBuilderBase& GetBuilder();
+
+		UE_EXPERIMENTAL(5.8, "Configurator is in active development and subject to change")
+		UE_API FDocumentConfigurator& ResetDocument();
+
+		struct FInputLocation
+		{
+			FName Name;
+			FVector2D Location = FVector2D::Zero();
+		};
+		UE_EXPERIMENTAL(5.8, "Configurator is in active development and subject to change")
+		UE_API FDocumentConfigurator& Set(FInputLocation Location);
+
+		struct FNodeInputDefault
+		{
+			FName Node;
+			FName Input;
+			FMetasoundFrontendLiteral Value;
+		};
+		UE_EXPERIMENTAL(5.8, "Configurator is in active development and subject to change")
+		UE_API FDocumentConfigurator& Set(FNodeInputDefault Input);
+
+		template <typename TemplateType>
+		FDocumentConfigurator& Set(TArray<TemplateType> Args)
+		{
+			for (TemplateType& TemplateArgs : Args)
+			{
+				Set(MoveTemp(TemplateArgs));
+			}
+
+			return *this;
+		}
+
+		UE_EXPERIMENTAL(5.8, "Configurator is in active development and subject to change")
+		UE_API bool Succeeded() const;
+
+	private:
+		FDocumentConfigurator& AddInternal(TArray<FInterface> Interfaces);
+
+		TStrongObjectPtr<UMetaSoundBuilderBase> Builder;
+		mutable EMetaSoundBuilderResult Result = EMetaSoundBuilderResult::Succeeded;
+
+		TMap<FName, FMetaSoundNodeHandle> NodeNameToHandle;
+
+		FVector2D Offset;
+	};
+} // namespace Metasound::Engine
+#undef UE_API
diff --git a/Engine/Plugins/Runtime/Metasound/Source/MetasoundEngine/Public/DocumentTemplates/MetasoundRandomizerTemplate.h b/Engine/Plugins/Runtime/Metasound/Source/MetasoundEngine/Public/DocumentTemplates/MetasoundRandomizerTemplate.h
new file mode 100644
--- /dev/null
+++ b/Engine/Plugins/Runtime/Metasound/Source/MetasoundEngine/Public/DocumentTemplates/MetasoundRandomizerTemplate.h
@@ -0,0 +1,54 @@
+// Copyright Epic Games, Inc. All Rights Reserved.
+#pragma once
+
+#include "DocumentTemplates/MetasoundFrontendDocumentTemplate.h"
+#include "MetasoundDocumentInterface.h"
+#include "Sound/SoundWave.h"
+#include "UObject/ObjectPtr.h"
+
+#include "MetasoundRandomizerTemplate.generated.h"
+
+#define UE_API METASOUNDENGINE_API
+
+// Forward Declarations
+class UMetaSoundFrontendMemberMetadata;
+enum class EMetasoundFrontendClassType : uint8;
+
+
+USTRUCT(BlueprintType, meta = (Hidden, DisplayName = "Randomizer"))
+struct FMetaSoundRandomizerTemplate : public FMetaSoundFrontendDocumentTemplate
+{
+	GENERATED_BODY()
+
+	virtual ~FMetaSoundRandomizerTemplate() = default;
+
+#if WITH_EDITORONLY_DATA
+	UE_INTERNAL UE_API virtual bool ConfigureDocument(FMetaSoundFrontendDocumentBuilder& OutBuilder) override;
+#endif // WITH_EDITORONLY_DATA
+
+#if WITH_EDITOR
+	UE_EXPERIMENTAL(5.8, "MetaSound Templates are in active development and subject to change")
+	UE_API virtual const FMetaSoundFrontendDocumentTemplate::FEditorOptions& GetEditorOptions() const override;
+
+	UE_EXPERIMENTAL(5.8, "MetaSound Templates are in active development and subject to change")
+	UE_API virtual void OnAssetInitialized(TArray<UObject*> SelectedObjects, FMetaSoundFrontendDocumentBuilder& OutBuilder) override;
+
+	UE_EXPERIMENTAL(5.8, "MetaSound Templates are in active development and subject to change")
+	UE_API virtual void OnPropertyChanged(const FPropertyChangedEvent& InEvent, FMetaSoundFrontendDocumentBuilder& OutBuilder) override;
+
+	// Helper function to remove sounds that cannot be referenced by randomizer (i.e. if the sounds are not exactly of type USoundWave
+	// or, optionally, if the underlying pointer is null). Returns true if sounds were removed, false if not.
+	UE_API static bool RemoveInvalidSounds(const FMetaSoundFrontendDocumentBuilder& ParentBuilder, bool bRemoveNullEntries, TArray<TObjectPtr<USoundWave>>& OutSounds);
+#endif // WITH_EDITOR
+
+public:
+	UPROPERTY(EditAnywhere, BlueprintReadOnly, meta = (Category = "Randomizer Options"))
+	bool bIsOneShot = false;
+
+	UPROPERTY(EditAnywhere, BlueprintReadOnly, meta = (Category = "Randomizer Options", DisallowedClasses = "/Script/MetasoundEngine.MetaSoundSource, /Script/Engine.SoundSourceBus"))
+	TArray<TObjectPtr<USoundWave>> Sounds;
+
+	UPROPERTY(EditAnywhere, BlueprintReadOnly, meta = (Category = "Randomizer Options", ClampMin = "-36.0", ClampMax = "36.0", UIMin = "-12.0", UIMax = "12.0"))
+	FVector2f Pitch = { -3.f, 3.f };
+};
+#undef UE_API
\ No newline at end of file
diff --git a/Engine/Plugins/Runtime/Metasound/Source/MetasoundEngine/Public/MetasoundBuilderBase.h b/Engine/Plugins/Runtime/Metasound/Source/MetasoundEngine/Public/MetasoundBuilderBase.h
--- a/Engine/Plugins/Runtime/Metasound/Source/MetasoundEngine/Public/MetasoundBuilderBase.h
+++ b/Engine/Plugins/Runtime/Metasound/Source/MetasoundEngine/Public/MetasoundBuilderBase.h
@@ -325,9 +325,12 @@ public:
 	UFUNCTION(BlueprintCallable, Category = "Audio|MetaSound|Builder", meta = (ExpandEnumAsExecs = "OutResult"))
 	UE_API UPARAM(DisplayName = "Node Class Version") FMetasoundFrontendVersion FindNodeClassVersion(const FMetaSoundNodeHandle& NodeHandle, EMetaSoundBuilderResult& OutResult);
 
-	// Returns copy of the builder's configuration. Returns success if the builder has a configuration set, failed if not.
+	// Returns copy of the builder's template. Returns success if the builder has a template set, fails if not.
 	UFUNCTION(BlueprintCallable, Category = "Audio|MetaSound|Builder", meta = (ExpandEnumAsExecs = "OutResult"))
-	UE_API TInstancedStruct<FMetaSoundFrontendDocumentConfiguration> GetConfiguration(EMetaSoundBuilderResult& OutResult) const;
+	UE_API TInstancedStruct<FMetaSoundFrontendDocumentTemplate> GetDocumentTemplate(EMetaSoundBuilderResult& OutResult) const;
+
+	UFUNCTION(BlueprintCallable, Category = "Audio|MetaSound|Builder", meta = (DeprecatedFunction, DeprecationMessage = "5.8 - Use GetDocumentTemplate instead", ExpandEnumAsExecs = "OutResult"))
+	TInstancedStruct<FMetaSoundFrontendDocumentTemplate> GetConfiguration(EMetaSoundBuilderResult& OutResult) const { return GetDocumentTemplate(OutResult); }
 
 	// Gets names of all graph inputs. 
 	UFUNCTION(BlueprintCallable, Category = "Audio|MetaSound|Builder", meta = (ExpandEnumAsExecs = "OutResult"))
@@ -478,17 +481,23 @@ public:
 	UE_API void SetAuthor(const FString& InAuthor);
 #endif // WITH_EDITOR
 
-	// Sets the given builder's configuration and immediately configures the MetaSound.
+	// Sets the given builder's template and immediately configures the MetaSound.
 	UFUNCTION(BlueprintCallable, Category = "Audio|MetaSound|Builder")
-	UE_API void SetConfiguration(TInstancedStruct<FMetaSoundFrontendDocumentConfiguration> Configuration);
+	UE_API void SetTemplate(TInstancedStruct<FMetaSoundFrontendDocumentTemplate> Configuration);
+
+	UFUNCTION(BlueprintCallable, Category = "Audio|MetaSound|Builder", meta = (DeprecatedFunction, DeprecationMessage = "5.8 - Use SetTemplate instead"))
+	void SetConfiguration(TInstancedStruct<FMetaSoundFrontendDocumentTemplate> Configuration) { SetTemplate(MoveTemp(Configuration)); }
+
+	// Passes view of struct along for mutating properties, firing OnDocumentTemplateChanged on completion and immediately reconfigures the MetaSound.
+	UE_API void SetTemplateProperties(TFunctionRef<void(TStructView<FMetaSoundFrontendDocumentTemplate>)> SetPropertiesFunction);
 
-	// Passes view of struct along for mutating properties, firing OnDocumentConfigurationChanged on completion and immediately reconfigures the MetaSound.
-	UE_API void SetConfigurationProperties(TFunctionRef<void(TStructView<FMetaSoundFrontendDocumentConfiguration>)> SetPropertiesFunction);
+	UE_DEPRECATED(5.8, "Use SetTemplateProperties() instead")
+	void SetConfigurationProperties(TFunctionRef<void(TStructView<FMetaSoundFrontendDocumentTemplate>)> SetPropertiesFunction) { SetTemplateProperties(SetPropertiesFunction); }
 
 	template <typename TConfig>
-	bool SetConfigurationProperties(TFunctionRef<void(TConfig&)> SetPropertyFunction)
+	bool SetTemplateProperties(TFunctionRef<void(TConfig&)> SetPropertyFunction)
 	{
-		const bool bConfigUpdated = GetBuilder().SetDocumentConfigurationProperties(SetPropertyFunction);
+		const bool bConfigUpdated = GetBuilder().SetTemplateProperties(SetPropertyFunction);
 		if (bConfigUpdated)
 		{
 			PRAGMA_DISABLE_EXPERIMENTAL_WARNINGS
diff --git a/Engine/Plugins/Runtime/Metasound/Source/MetasoundEngine/Public/MetasoundBuilderSubsystem.h b/Engine/Plugins/Runtime/Metasound/Source/MetasoundEngine/Public/MetasoundBuilderSubsystem.h
--- a/Engine/Plugins/Runtime/Metasound/Source/MetasoundEngine/Public/MetasoundBuilderSubsystem.h
+++ b/Engine/Plugins/Runtime/Metasound/Source/MetasoundEngine/Public/MetasoundBuilderSubsystem.h
@@ -122,7 +122,7 @@ private:
 
 	UE_API void InitTargetPageDelegates(Metasound::Frontend::FDocumentModifyDelegates& OutDocumentDelegates);
 
-	UE_API void OnDocumentConfigurationChanged(const Metasound::Frontend::FDocumentConfigChangedArgs& Args) const;
+	UE_API void OnDocumentTemplateChanged(const Metasound::Frontend::FDocumentTemplateChangedArgs& Args) const;
 	UE_API void OnEdgeAdded(int32 EdgeIndex) const;
 	UE_API void OnInputAdded(int32 InputIndex);
 	UE_API void OnInputDataTypeChanged(int32 InputIndex);
diff --git a/Engine/Plugins/Runtime/Metasound/Source/MetasoundFrontend/Private/DocumentConfigurations/MetasoundFrontendDocumentConfiguration.cpp b/Engine/Plugins/Runtime/Metasound/Source/MetasoundFrontend/Private/DocumentConfigurations/MetasoundFrontendDocumentConfiguration.cpp
deleted file mode 100644
diff --git a/Engine/Plugins/Runtime/Metasound/Source/MetasoundFrontend/Private/DocumentConfigurations/MetasoundFrontendPresetDocumentConfiguration.cpp b/Engine/Plugins/Runtime/Metasound/Source/MetasoundFrontend/Private/DocumentConfigurations/MetasoundFrontendPresetDocumentConfiguration.cpp
deleted file mode 100644
diff --git a/Engine/Plugins/Runtime/Metasound/Source/MetasoundFrontend/Private/DocumentConfigurations/MetasoundFrontendVertexDocumentConfiguration.cpp b/Engine/Plugins/Runtime/Metasound/Source/MetasoundFrontend/Private/DocumentConfigurations/MetasoundFrontendVertexDocumentConfiguration.cpp
deleted file mode 100644
diff --git a/Engine/Plugins/Runtime/Metasound/Source/MetasoundFrontend/Private/DocumentTemplates/MetasoundFrontendDocumentTemplate.cpp b/Engine/Plugins/Runtime/Metasound/Source/MetasoundFrontend/Private/DocumentTemplates/MetasoundFrontendDocumentTemplate.cpp
new file mode 100644
--- /dev/null
+++ b/Engine/Plugins/Runtime/Metasound/Source/MetasoundFrontend/Private/DocumentTemplates/MetasoundFrontendDocumentTemplate.cpp
@@ -0,0 +1,67 @@
+// Copyright Epic Games, Inc. All Rights Reserved.
+
+#include "DocumentTemplates/MetasoundFrontendDocumentTemplate.h"
+
+#include "MetasoundDocumentInterface.h"
+#include "MetasoundFrontendDocument.h"
+#include "MetasoundFrontendDocumentBuilder.h"
+#include "MetasoundFrontendDocumentBuilderRegistry.h"
+#include "MetasoundFrontendSearchEngine.h"
+
+
+#define LOCTEXT_NAMESPACE "MetaSound"
+
+bool FMetaSoundFrontendDocumentTemplate::ConfigureDocument(FMetaSoundFrontendDocumentBuilder& OutBuilder)
+{
+	using namespace Metasound::Frontend;
+
+	TSharedRef<FDocumentModifyDelegates> ModifyDelegates = MakeShared<FDocumentModifyDelegates>(OutBuilder.GetDocumentDelegates());
+	OutBuilder.ClearDocument(ModifyDelegates);
+
+#if WITH_EDITORONLY_DATA
+	OutBuilder.SetIsGraphEditable(false);
+#endif // WITH_EDITORONLY_DATA
+
+	// Add required interfaces, but disable editing to ensure document maintains valid state
+	// (ex. MetaSound sources need  the source interface to properly initialize a null-rendering
+	// MetaSound generator.)
+	TArray<FMetasoundFrontendVersion> InitVersions = ISearchEngine::Get().FindUClassDefaultInterfaceVersions(OutBuilder.GetBuilderClassPath());
+	FModifyInterfaceOptions Options({ }, MoveTemp(InitVersions));
+	const bool bInterfacesModified = OutBuilder.ModifyInterfaces(MoveTemp(Options));
+	ensureMsgf(bInterfacesModified, TEXT("Failed to apply default interfaces when configuring default MetaSound template"));
+
+	return true;
+}
+
+#if WITH_EDITOR
+Metasound::Frontend::FOnConfigPropertyChangedDelegate& FMetaSoundFrontendDocumentTemplate::GetPropertyChangedDelegate() const
+{
+	return PropertyChangedDelegate;
+}
+
+const FMetaSoundFrontendDocumentTemplate::FEditorOptions& FMetaSoundFrontendDocumentTemplate::FEditorOptions::GetDefaultOptions()
+{
+	static const FEditorOptions Default;
+	return Default;
+}
+
+const FMetaSoundFrontendDocumentTemplate::FEditorOptions& FMetaSoundFrontendDocumentTemplate::GetEditorOptions() const
+{
+	return FEditorOptions::GetDefaultOptions();
+}
+
+EDataValidationResult FMetaSoundFrontendDocumentTemplate::IsDataValid(const FMetaSoundFrontendDocumentBuilder& Builder, FDataValidationContext& InOutContext) const
+{
+	return EDataValidationResult::Valid;
+}
+
+void FMetaSoundFrontendDocumentTemplate::OnAssetInitialized(TArray<UObject*> SelectedObjects, FMetaSoundFrontendDocumentBuilder& OutBuilder)
+{
+}
+
+void FMetaSoundFrontendDocumentTemplate::OnPropertyChanged(const FPropertyChangedEvent& InEvent, FMetaSoundFrontendDocumentBuilder& OutBuilder)
+{
+	PropertyChangedDelegate.Broadcast(InEvent);
+}
+#endif // WITH_EDITOR
+#undef LOCTEXT_NAMESPACE
diff --git a/Engine/Plugins/Runtime/Metasound/Source/MetasoundFrontend/Private/DocumentTemplates/MetasoundFrontendPresetTemplate.cpp b/Engine/Plugins/Runtime/Metasound/Source/MetasoundFrontend/Private/DocumentTemplates/MetasoundFrontendPresetTemplate.cpp
new file mode 100644
--- /dev/null
+++ b/Engine/Plugins/Runtime/Metasound/Source/MetasoundFrontend/Private/DocumentTemplates/MetasoundFrontendPresetTemplate.cpp
@@ -0,0 +1,208 @@
+// Copyright Epic Games, Inc. All Rights Reserved.
+#include "DocumentTemplates/MetasoundFrontendPresetTemplate.h"
+
+#include "Algo/AnyOf.h"
+
+#if WITH_EDITOR
+#include "MetasoundAssetManager.h"
+#endif // WITH_EDITOR
+
+#include "MetasoundDataTypeGetTraits.h"
+#include "MetasoundDocumentInterface.h"
+#include "MetasoundFrontendDocument.h"
+#include "MetasoundFrontendDocumentBuilder.h"
+#include "MetasoundFrontendDocumentBuilderRegistry.h"
+#include "MetasoundTrigger.h"
+
+#define LOCTEXT_NAMESPACE "MetaSound"
+
+namespace Metasound::Frontend::PresetDocTemplatePrivate
+{
+	// This likely should be deprecated in favor of using template data directly
+	void BroadcastInputInheritanceUpdated(const FMetaSoundFrontendDocumentBuilder& Builder, const FGuid& InNodeID)
+	{
+		auto IsNode = [&InNodeID](const FMetasoundFrontendClassInput& Input) { return Input.NodeID == InNodeID; };
+		const int32 Index = Builder.GetConstDocumentChecked().RootGraph.GetDefaultInterface().Inputs.IndexOfByPredicate(IsNode);
+		const_cast<FMetaSoundFrontendDocumentBuilder&>(Builder).GetDocumentDelegates()
+			.InterfaceDelegates.OnInputInheritsDefaultChanged.Broadcast(Index);
+	}
+
+	bool IsInheritableInputNode(const FMetaSoundFrontendDocumentBuilder& Builder, const FGuid& NodeID)
+	{
+		const FMetasoundFrontendNode* Node = Builder.FindNode(NodeID);
+		if (!Node)
+		{
+			return false;
+		}
+
+		const FMetasoundFrontendClass* Class = Builder.FindDependency(Node->ClassID);
+		if (!Class)
+		{
+			return false;
+		}
+
+		if (Class->Metadata.GetType() != EMetasoundFrontendClassType::Input)
+		{
+			return false;
+		}
+
+		const FMetasoundFrontendClassInput* Input = Builder.FindGraphInput(Node->Name);
+		if (!Input)
+		{
+			return false;
+		}
+
+		auto IsDefaultConstructed = [](const FMetasoundFrontendClassInputDefault& Default) { return Default.Literal.GetType() == EMetasoundFrontendLiteralType::None; };
+		if (Algo::AnyOf(Input->GetDefaults(), IsDefaultConstructed))
+		{
+			return false;
+		}
+
+		if (Input->TypeName == GetMetasoundDataTypeName<FTrigger>())
+		{
+			return false;
+		}
+
+		return true;
+	}
+} // namespace Metasound::Frontend::PresetDocTemplatePrivate
+
+#if WITH_EDITORONLY_DATA
+bool FMetaSoundFrontendPresetTemplate::ConfigureDocument(FMetaSoundFrontendDocumentBuilder& OutBuilder)
+{
+	using namespace Metasound::Frontend;
+
+	if (UObject* RefObject = Parent.GetObject())
+	{
+		if (RefObject->IsA(&OutBuilder.GetConstDocumentInterfaceChecked().GetBaseMetaSoundUClass()))
+		{
+			const FMetaSoundFrontendDocumentBuilder& ParentBuilder = IDocumentBuilderRegistry::GetChecked().FindOrBeginBuilding(Parent);
+
+			OutBuilder.SetIsGraphEditable(false);
+
+			// Apply root graph transform
+			FRebuildPresetRootGraph RebuildPresetRootGraph(ParentBuilder);
+			return RebuildPresetRootGraph.Transform(OutBuilder);
+		}
+	}
+
+	// Need to configure using parent template call when reference is invalid as by default, no referenced parent is set and
+	// document must be valid when default constructed for testing purposes.  Data Validation provides place to surface error
+	// (i.e. when saving or reserializing).
+	return Super::ConfigureDocument(OutBuilder);
+}
+
+EDataValidationResult FMetaSoundFrontendPresetTemplate::IsDataValid(const FMetaSoundFrontendDocumentBuilder& Builder, FDataValidationContext& InOutContext) const
+{
+	using namespace Metasound::Frontend;
+
+	EDataValidationResult Result = EDataValidationResult::Valid;
+	if (UObject* RefObject = Parent.GetObject())
+	{
+		if (!RefObject->IsA(&Builder.GetConstDocumentInterfaceChecked().GetBaseMetaSoundUClass()))
+		{
+			Result = EDataValidationResult::Invalid;
+			InOutContext.AddError(FText::Format(LOCTEXT("MetaSoundPreset_InvalidReferenceTypeFormat", "The referenced node type must match the base MetaSound class when converting "
+				"builder '{0}' to a preset (ex. source preset must reference another source)"),
+				FText::FromString(Builder.GetDebugName())
+			));
+		}
+	}
+	else
+	{
+		Result = EDataValidationResult::Invalid;
+		InOutContext.AddError(FText::Format(LOCTEXT("MetaSoundPreset_ReferenceUnsetFormat", "Preset '%s' generation failed: Referenced class not set."), FText::FromString(Builder.GetDebugName())));
+	}
+
+	return Result;
+}
+#endif // WITH_EDITORONLY_DATA
+
+TInstancedStruct<FMetaSoundFrontendDocumentTemplateVertexMetadata> FMetaSoundFrontendPresetTemplate::ConstructVertexMetadata(const FMetaSoundFrontendDocumentBuilder& Builder, const FMetasoundFrontendClassVertex& InVertex) const
+{
+	using namespace Metasound::Frontend;
+
+	// Presets only need ConfigVertexMetadata for inputs, so don't construct otherwise
+	if (PresetDocTemplatePrivate::IsInheritableInputNode(Builder, InVertex.NodeID))
+	{
+		return TInstancedStruct<FPresetVertexMetadata>::Make();
+	}
+
+	return { };
+}
+
+#if WITH_EDITOR
+const FMetaSoundFrontendDocumentTemplate::FEditorOptions& FMetaSoundFrontendPresetTemplate::GetEditorOptions() const
+{
+	static const FEditorOptions PresetOptions
+	{
+		.bMemberEditingEnabled = true,
+		.bGraphEditorVisible = true
+	};
+	return PresetOptions;
+}
+
+void FMetaSoundFrontendPresetTemplate::OnAssetInitialized(TArray<UObject*> SelectedObjects, FMetaSoundFrontendDocumentBuilder& OutBuilder)
+{
+	for (UObject* SelectedObject : SelectedObjects)
+	{
+		if (&OutBuilder.GetConstDocumentInterfaceChecked().GetBaseMetaSoundUClass() == SelectedObject->GetClass())
+		{
+			Parent = SelectedObject;
+			return;
+		}
+	}
+}
+
+void FMetaSoundFrontendPresetTemplate::OnPropertyChanged(const FPropertyChangedEvent& InEvent, FMetaSoundFrontendDocumentBuilder& OutBuilder)
+{
+	using namespace Metasound::Frontend;
+
+	if (InEvent.GetPropertyName() == GET_MEMBER_NAME_CHECKED(FMetaSoundFrontendPresetTemplate, Parent))
+	{
+		if (UObject* ParentObject = Parent.GetObject())
+		{
+			if (FMetasoundAssetBase* ParentAsset = IMetaSoundAssetManager::GetChecked().GetAsAsset(*ParentObject))
+			{
+				const bool bCreatesLoop = OutBuilder.GetMetasoundAsset().AddingReferenceCausesLoop(*ParentAsset);
+				if (bCreatesLoop)
+				{
+					UE_LOG(LogMetaSound, Error, TEXT("Parent could not be set to '%s': Causes reference loop."), *ParentObject->GetFullName());
+					Parent = nullptr;
+				}
+				else
+				{
+					ParentAsset->UpdateAndRegisterForExecution(IMetaSoundAssetManager::GetChecked().GetRegistrationOptions());
+				}
+			}
+		}
+	}
+}
+
+bool FMetaSoundFrontendPresetVertexMetadata::IsDefaultEditable() const
+{
+	return bOverrideInheritedDefault;
+}
+#endif // WITH_EDITOR
+
+void FMetaSoundFrontendPresetVertexMetadata::OnDefaultReset(const FMetaSoundFrontendDocumentBuilder& Builder, const FMetasoundFrontendClassVertex& InVertex)
+{
+	using namespace Metasound::Frontend;
+
+	bOverrideInheritedDefault = false;
+	PresetDocTemplatePrivate::BroadcastInputInheritanceUpdated(const_cast<FMetaSoundFrontendDocumentBuilder&>(Builder), InVertex.NodeID);
+}
+
+void FMetaSoundFrontendPresetVertexMetadata::OnDefaultUpdated(const FMetaSoundFrontendDocumentBuilder& Builder, const FMetasoundFrontendClassVertex& InVertex)
+{
+	using namespace Metasound::Frontend;
+
+	// When actively configuring the document, ignore calls to defaults being
+	// updated as this means the request isn't being applied by a user/client code.
+	if (!Builder.IsConfiguringDocument())
+	{
+		bOverrideInheritedDefault = true;
+		PresetDocTemplatePrivate::BroadcastInputInheritanceUpdated(const_cast<FMetaSoundFrontendDocumentBuilder&>(Builder), InVertex.NodeID);
+	}
+}
+#undef LOCTEXT_NAMESPACE
diff --git a/Engine/Plugins/Runtime/Metasound/Source/MetasoundFrontend/Private/DocumentTemplates/MetasoundFrontendVertexTemplate.cpp b/Engine/Plugins/Runtime/Metasound/Source/MetasoundFrontend/Private/DocumentTemplates/MetasoundFrontendVertexTemplate.cpp
new file mode 100644
--- /dev/null
+++ b/Engine/Plugins/Runtime/Metasound/Source/MetasoundFrontend/Private/DocumentTemplates/MetasoundFrontendVertexTemplate.cpp
@@ -0,0 +1,48 @@
+// Copyright Epic Games, Inc. All Rights Reserved.
+
+#include "DocumentTemplates/MetasoundFrontendDocumentVertexTemplate.h"
+
+#include "MetasoundDocumentInterface.h"
+#include "MetasoundFrontendDocument.h"
+#include "MetasoundFrontendDocumentBuilder.h"
+#include "MetasoundFrontendDocumentBuilderRegistry.h"
+
+
+#define LOCTEXT_NAMESPACE "MetaSound"
+
+FMetaSoundFrontendDocumentTemplateVertexMetadata* FMetaSoundFrontendDocumentVertexTemplate::AddVertexMetadata(const FMetaSoundFrontendDocumentBuilder& Builder, const FMetasoundFrontendClassVertex& InVertex)
+{
+	if (TInstancedStruct<FMetaSoundFrontendDocumentTemplateVertexMetadata>* Entry = VertexMetadata.Find(InVertex.NodeID))
+	{
+		return Entry->GetMutablePtr();
+	}
+
+PRAGMA_DISABLE_EXPERIMENTAL_WARNINGS
+	TInstancedStruct<FMetaSoundFrontendDocumentTemplateVertexMetadata> NewMetadata = ConstructVertexMetadata(Builder, InVertex);
+PRAGMA_ENABLE_EXPERIMENTAL_WARNINGS
+
+	if (NewMetadata.IsValid())
+	{
+		return VertexMetadata.Add(InVertex.NodeID, MoveTemp(NewMetadata)).GetMutablePtr();
+	}
+
+	return nullptr;
+}
+
+TInstancedStruct<FMetaSoundFrontendDocumentTemplateVertexMetadata> FMetaSoundFrontendDocumentVertexTemplate::ConstructVertexMetadata(
+	const FMetaSoundFrontendDocumentBuilder& Builder,
+	const FMetasoundFrontendClassVertex& InVertex) const
+{
+	return { };
+}
+
+bool FMetaSoundFrontendDocumentVertexTemplate::ContainsVertexMetadata(const FGuid& NodeID) const
+{
+	return VertexMetadata.Contains(NodeID);
+}
+
+bool FMetaSoundFrontendDocumentVertexTemplate::RemoveVertexMetadata(const FGuid& NodeID)
+{
+	return VertexMetadata.Remove(NodeID) > 0;
+}
+#undef LOCTEXT_NAMESPACE
diff --git a/Engine/Plugins/Runtime/Metasound/Source/MetasoundFrontend/Private/MetasoundAssetBase.cpp b/Engine/Plugins/Runtime/Metasound/Source/MetasoundFrontend/Private/MetasoundAssetBase.cpp
--- a/Engine/Plugins/Runtime/Metasound/Source/MetasoundFrontend/Private/MetasoundAssetBase.cpp
+++ b/Engine/Plugins/Runtime/Metasound/Source/MetasoundFrontend/Private/MetasoundAssetBase.cpp
@@ -7,7 +7,6 @@
 #include "Algo/ForEach.h"
 #include "Algo/Transform.h"
 #include "Containers/Set.h"
-#include "DocumentConfigurations/MetasoundFrontendPresetDocumentConfiguration.h"
 #include "HAL/FileManager.h"
 #include "HAL/IConsoleManager.h"
 #include "IAudioParameterTransmitter.h"
@@ -40,6 +39,7 @@
 #include "MetasoundTrace.h"
 #include "MetasoundVertex.h"
 #include "StructSerializer.h"
+#include "DocumentTemplates/MetasoundFrontendPresetTemplate.h"
 #include "Templates/SharedPointer.h"
 #include "UObject/MetaData.h"
 
@@ -636,12 +636,12 @@ FText FMetasoundAssetBase::GetDisplayName(FString&& InTypeName) const
 	using namespace Metasound::Frontend;
 
 	TScriptInterface<const IMetaSoundDocumentInterface> DocInterface(GetOwningAsset());
-	TConstStructView<FMetaSoundFrontendDocumentConfiguration> Config = DocInterface->GetConstDocument().Configuration;
-	if (Config.IsValid())
+	TConstStructView<FMetaSoundFrontendDocumentTemplate> Template = DocInterface->GetConstDocument().Template;
+	if (Template.IsValid())
 	{
-		const UScriptStruct* Struct = Config.GetScriptStruct();
+		const UScriptStruct* Struct = Template.GetScriptStruct();
 		check(Struct);
-		return FText::Format(LOCTEXT("DocConfigDisplayNameFormat", "{0} ({1})"), FText::FromString(MoveTemp(InTypeName)), Struct->GetDisplayNameText());
+		return FText::Format(LOCTEXT("DocTemplateDisplayNameFormat", "{0} ({1})"), FText::FromString(MoveTemp(InTypeName)), Struct->GetDisplayNameText());
 	}
 
 	return FText::FromString(MoveTemp(InTypeName));
diff --git a/Engine/Plugins/Runtime/Metasound/Source/MetasoundFrontend/Private/MetasoundAssetManager.cpp b/Engine/Plugins/Runtime/Metasound/Source/MetasoundFrontend/Private/MetasoundAssetManager.cpp
--- a/Engine/Plugins/Runtime/Metasound/Source/MetasoundFrontend/Private/MetasoundAssetManager.cpp
+++ b/Engine/Plugins/Runtime/Metasound/Source/MetasoundFrontend/Private/MetasoundAssetManager.cpp
@@ -4,12 +4,12 @@
 
 #include "Algo/Copy.h"
 #include "AssetRegistry/AssetData.h"
-#include "DocumentConfigurations/MetasoundFrontendPresetDocumentConfiguration.h"
 #include "JsonObjectConverter.h"
 #include "MetasoundAssetManagerTransactionPrivate.h"
 #include "MetasoundAssetTagCollections.h"
 #include "MetasoundFrontendDocumentVersioning.h"
 #include "MetasoundFrontendRegistryKey.h"
+#include "DocumentTemplates/MetasoundFrontendPresetTemplate.h"
 
 #include UE_INLINE_GENERATED_CPP_BY_NAME(MetasoundAssetManager)
 
@@ -61,7 +61,7 @@ FMetaSoundDocumentInfo::FMetaSoundDocumentInfo(const IMetaSoundDocumentInterface
 
 	if (!IsRunningCookCommandlet())
 	{
-		bIsPreset = (uint8)(Document.Configuration.GetPtr<FMetaSoundFrontendPresetDocumentConfiguration>() != nullptr);
+		bIsPreset = (uint8)(Document.Template.GetPtr<FMetaSoundFrontendPresetTemplate>() != nullptr);
 		DocumentVersion = Document.Metadata.Version.Number;
 
 		{
diff --git a/Engine/Plugins/Runtime/Metasound/Source/MetasoundFrontend/Private/MetasoundFrontendDocumentBuilder.cpp b/Engine/Plugins/Runtime/Metasound/Source/MetasoundFrontend/Private/MetasoundFrontendDocumentBuilder.cpp
--- a/Engine/Plugins/Runtime/Metasound/Source/MetasoundFrontend/Private/MetasoundFrontendDocumentBuilder.cpp
+++ b/Engine/Plugins/Runtime/Metasound/Source/MetasoundFrontend/Private/MetasoundFrontendDocumentBuilder.cpp
@@ -3,37 +3,30 @@
 
 #include "Algo/AllOf.h"
 #include "Algo/AnyOf.h"
-#include "Algo/Find.h"
 #include "Algo/ForEach.h"
-#include "Algo/NoneOf.h"
 #include "Algo/Sort.h"
 #include "Algo/Transform.h"
 #include "AudioParameter.h"
-#include "DocumentConfigurations/MetasoundFrontendDocumentConfiguration.h"
-#include "DocumentConfigurations/MetasoundFrontendPresetDocumentConfiguration.h"
+#include "DocumentTemplates/MetasoundFrontendPresetTemplate.h"
+#include "DocumentTemplates/MetasoundFrontendDocumentTemplate.h"
 #include "Interfaces/MetasoundFrontendInterfaceBindingRegistry.h"
 #include "Interfaces/MetasoundFrontendInterfaceRegistry.h"
 #include "MetasoundAssetBase.h"
 #include "MetasoundAssetManager.h"
 #include "MetasoundDocumentInterface.h"
-#include "MetasoundFrontendController.h"
 #include "MetasoundFrontendDataTypeRegistry.h"
 #include "MetasoundFrontendDocumentCache.h"
 #include "MetasoundFrontendDocument.h"
 #include "MetasoundFrontendDocumentIdGenerator.h"
 #include "MetasoundFrontendDocumentModifyDelegates.h"
 #include "MetasoundFrontendDocumentVersioning.h"
-#include "MetasoundFrontendGraphLayer.h"
 #include "MetasoundFrontendNodeTemplateRegistry.h"
 #include "MetasoundFrontendNodeUpdateTransform.h"
-#include "MetasoundFrontendRegistries.h"
 #include "MetasoundFrontendRegistryKey.h"
 #include "MetasoundFrontendSearchEngine.h"
-#include "MetasoundFrontendTransform.h"
 #include "MetasoundTrace.h"
 #include "MetasoundVariableNodes.h"
 #include "NodeTemplates/MetasoundFrontendNodeTemplateInput.h"
-#include "UObject/UObjectIterator.h"
 
 #include UE_INLINE_GENERATED_CPP_BY_NAME(MetasoundFrontendDocumentBuilder)
 
@@ -812,17 +805,17 @@ const FMetasoundFrontendDocument& UMetaSoundBuilderDocument::GetConstDocument()
 }
 
 #if WITH_EDITOR
-TStructView<FMetaSoundFrontendDocumentConfiguration> FMetaSoundFrontendDocumentBuilder::GetDocumentConfiguration() const
+TStructView<FMetaSoundFrontendDocumentTemplate> FMetaSoundFrontendDocumentBuilder::GetDocumentTemplate() const
 {
-	TInstancedStruct<FMetaSoundFrontendDocumentConfiguration>& Config = GetDocumentChecked().Configuration;
+	TInstancedStruct<FMetaSoundFrontendDocumentTemplate>& Config = GetDocumentChecked().Template;
 	return Config;
 }
 #endif // WITH_EDITOR
 
-TConstStructView<FMetaSoundFrontendDocumentConfiguration> FMetaSoundFrontendDocumentBuilder::GetConstDocumentConfiguration() const
+TConstStructView<FMetaSoundFrontendDocumentTemplate> FMetaSoundFrontendDocumentBuilder::GetConstDocumentTemplate() const
 {
-	const TInstancedStruct<FMetaSoundFrontendDocumentConfiguration>& Config = GetConstDocumentChecked().Configuration;
-	return Config;
+	const TInstancedStruct<FMetaSoundFrontendDocumentTemplate>& Template = GetConstDocumentChecked().Template;
+	return Template;
 }
 
 EMetasoundFrontendClassAccessFlags UMetaSoundBuilderDocument::GetDefaultAccessFlags() const
@@ -1288,10 +1281,10 @@ const FMetasoundFrontendNode* FMetaSoundFrontendDocumentBuilder::AddGraphInput(F
 
 		DocumentDelegates->InterfaceDelegates.OnInputAdded.Broadcast(NewIndex);
 #if WITH_EDITORONLY_DATA
-		if (FDocumentVertexConfiguration* DocConfig = Document.Configuration.GetMutablePtr<FDocumentVertexConfiguration>())
+		if (FDocumentVertexTemplate* DocTemplate = Document.Template.GetMutablePtr<FDocumentVertexTemplate>())
 		{
 PRAGMA_DISABLE_EXPERIMENTAL_WARNINGS
-			DocConfig->AddVertexMetadata(*this, NewInput);
+			DocTemplate->AddVertexMetadata(*this, NewInput);
 PRAGMA_ENABLE_EXPERIMENTAL_WARNINGS
 		}
 
@@ -1408,10 +1401,10 @@ const FMetasoundFrontendNode* FMetaSoundFrontendDocumentBuilder::AddGraphOutput(
 		DocumentDelegates->InterfaceDelegates.OnOutputAdded.Broadcast(NewIndex);
 
 #if WITH_EDITORONLY_DATA
-		if (FDocumentVertexConfiguration* DocConfig = Document.Configuration.GetMutablePtr<FDocumentVertexConfiguration>())
+		if (FDocumentVertexTemplate* DocTemplate = Document.Template.GetMutablePtr<FDocumentVertexTemplate>())
 		{
 PRAGMA_DISABLE_EXPERIMENTAL_WARNINGS
-			DocConfig->AddVertexMetadata(*this, NewOutput);
+			DocTemplate->AddVertexMetadata(*this, NewOutput);
 PRAGMA_ENABLE_EXPERIMENTAL_WARNINGS
 		}
 
@@ -2524,18 +2517,18 @@ bool FMetaSoundFrontendDocumentBuilder::ConfigureDocument()
 {
 	using namespace Metasound::Frontend;
 
-	if (bInDocConfigScope)
+	if (bInTemplateScope)
 	{
 		UE_LOG(LogMetaSound, Error, TEXT("Cannot call Configure Document: Configure already called higher in stack"));
 		return false;
 	}
 
-	bInDocConfigScope = true;
+	bInTemplateScope = true;
 
 	FMetasoundFrontendDocument& Document = GetDocumentChecked();
-	if (FMetaSoundFrontendDocumentConfiguration* DocConfig = Document.Configuration.GetMutablePtr())
+	if (FMetaSoundFrontendDocumentTemplate* DocTemplate = Document.Template.GetMutablePtr())
 	{
-		if (DocConfig->ConfigureDocument(*this))
+		if (DocTemplate->ConfigureDocument(*this))
 		{
 			GetDocumentInterfaceChecked().ConformObjectToDocument();
 
@@ -2543,7 +2536,7 @@ bool FMetaSoundFrontendDocumentBuilder::ConfigureDocument()
 			FInputNodeTemplate::GetChecked().Inject(*this);
 #endif // WITH_EDITORONLY_DATA
 
-			bInDocConfigScope = false;
+			bInTemplateScope = false;
 			return true;
 		}
 
@@ -2568,7 +2561,7 @@ bool FMetaSoundFrontendDocumentBuilder::ConfigureDocument()
 	}
 #endif // WITH_EDITORONLY_DATA
 
-	bInDocConfigScope = false;
+	bInTemplateScope = false;
 	return false;
 }
 
@@ -2597,10 +2590,10 @@ bool FMetaSoundFrontendDocumentBuilder::ConvertFromPreset()
 
 	if (IsPreset())
 	{
-		const bool bConfigChanged = GetDocumentChecked().Configuration.IsValid();
+		const bool bConfigChanged = GetDocumentChecked().Template.IsValid();
 
-		// Don't process the null configuration in order to "freeze" the current document state and not clear it.
-		GetDocumentChecked().Configuration.Reset();
+		// Don't configure after nulling the template in order to "freeze" the current document/MetaSound UObject state (i.e. not clear it).
+		GetDocumentChecked().Template.Reset();
 
 #if WITH_EDITORONLY_DATA
 		SetIsGraphEditable(true);
@@ -2608,7 +2601,7 @@ bool FMetaSoundFrontendDocumentBuilder::ConvertFromPreset()
 
 		if (bConfigChanged)
 		{
-			DocumentDelegates->OnDocumentConfigurationChanged.Broadcast(FDocumentConfigChangedArgs{ .Type = EDocumentConfigChangeType::Struct });
+			DocumentDelegates->OnDocumentTemplateChanged.Broadcast(FDocumentTemplateChangedArgs { .Type = EDocumentTemplateChangeType::Struct });
 		}
 
 		return true;
@@ -2627,10 +2620,10 @@ bool FMetaSoundFrontendDocumentBuilder::ConvertToPreset(const FMetaSoundFrontend
 	using namespace Metasound;
 	using namespace Metasound::Frontend;
 
-	using FPresetConfigStruct = TInstancedStruct<FMetaSoundFrontendPresetDocumentConfiguration>;
-	FPresetConfigStruct NewConfig = FPresetConfigStruct::Make<FMetaSoundFrontendPresetDocumentConfiguration>();
+	using FPresetConfigStruct = TInstancedStruct<FMetaSoundFrontendPresetTemplate>;
+	FPresetConfigStruct NewConfig = FPresetConfigStruct::Make<FMetaSoundFrontendPresetTemplate>();
 	NewConfig.GetMutable().Parent = &InReferencedDocumentBuilder.CastDocumentObjectChecked<UObject>();
-	SetDocumentConfiguration(MoveTemp(NewConfig));
+	SetTemplate(MoveTemp(NewConfig));
 	ConfigureDocument();
 
 	return true;
@@ -3773,9 +3766,9 @@ FMetasoundAssetBase* FMetaSoundFrontendDocumentBuilder::GetReferencedPresetAsset
 {
 	using namespace Metasound::Frontend;
 
-	if (const FMetaSoundFrontendPresetDocumentConfiguration* PresetConfig = GetConstDocumentConfigurationAs<FMetaSoundFrontendPresetDocumentConfiguration>())
+	if (const FMetaSoundFrontendPresetTemplate* PresetTemplate = GetConstTemplateAs<FMetaSoundFrontendPresetTemplate>())
 	{
-		if (UObject* Parent = PresetConfig->Parent.GetObject())
+		if (UObject* Parent = PresetTemplate->Parent.GetObject())
 		{
 			return IMetaSoundAssetManager::GetChecked().GetAsAsset(*Parent);
 		}
@@ -4330,7 +4323,7 @@ void FMetaSoundFrontendDocumentBuilder::InitNodeLocations()
 
 bool FMetaSoundFrontendDocumentBuilder::IsConfiguringDocument() const
 {
-	return bInDocConfigScope;
+	return bInTemplateScope;
 }
 
 bool FMetaSoundFrontendDocumentBuilder::IsDependencyReferenced(const FGuid& InClassID) const
@@ -4479,7 +4472,7 @@ bool FMetaSoundFrontendDocumentBuilder::IsPreset() const
 {
 	using namespace Metasound::Frontend;
 
-	return IsDocumentConfiguredBy<FPresetDocumentConfiguration>();
+	return IsDocumentConfiguredBy<FPresetTemplate>();
 }
 
 Metasound::Frontend::EInvalidEdgeReason FMetaSoundFrontendDocumentBuilder::IsValidEdge(const FMetasoundFrontendEdge& InEdge, const FGuid* InPageID) const
@@ -5243,9 +5236,9 @@ bool FMetaSoundFrontendDocumentBuilder::RemoveGraphInput(FName InputName, bool b
 		if (bNodesRemoved)
 		{
 #if WITH_EDITORONLY_DATA
-			if (FDocumentVertexConfiguration* DocConfig = Document.Configuration.GetMutablePtr<FDocumentVertexConfiguration>())
+			if (FDocumentVertexTemplate* DocTemplate = Document.Template.GetMutablePtr<FDocumentVertexTemplate>())
 			{
-				DocConfig->RemoveVertexMetadata(NodeID);
+				DocTemplate->RemoveVertexMetadata(NodeID);
 			}
 #endif // WITH_EDITORONLY_DATA
 
@@ -5308,9 +5301,9 @@ bool FMetaSoundFrontendDocumentBuilder::RemoveGraphOutput(FName OutputName)
 	if (bNodesRemoved)
 	{
 #if WITH_EDITORONLY_DATA
-		if (FDocumentVertexConfiguration* DocConfig = Document.Configuration.GetMutablePtr<FDocumentVertexConfiguration>())
+		if (FDocumentVertexTemplate* DocTemplate = Document.Template.GetMutablePtr<FDocumentVertexTemplate>())
 		{
-			DocConfig->RemoveVertexMetadata(NodeID);
+			DocTemplate->RemoveVertexMetadata(NodeID);
 		}
 #endif // WITH_EDITORONLY_DATA
 
@@ -5849,9 +5842,9 @@ bool FMetaSoundFrontendDocumentBuilder::ResetGraphInputDefault(FName InputName)
 			});
 		}
 		
-		if (FDocumentVertexConfiguration* Config = Document.Configuration.GetMutablePtr<FDocumentVertexConfiguration>())
+		if (FDocumentVertexTemplate* Config = Document.Template.GetMutablePtr<FDocumentVertexTemplate>())
 		{
-			if (FMetaSoundFrontendDocumentConfigurationVertexMetadata* VertexMetadata = Config->FindVertexMetadata(Input.NodeID))
+			if (FMetaSoundFrontendDocumentTemplateVertexMetadata* VertexMetadata = Config->FindVertexMetadata(Input.NodeID))
 			{
 				VertexMetadata->OnDefaultReset(*this, Input);
 			}
@@ -5933,19 +5926,19 @@ bool FMetaSoundFrontendDocumentBuilder::SetBuildPageID(const FGuid& InBuildPageI
 }
 #endif // WITH_EDITORONLY_DATA
 
-void FMetaSoundFrontendDocumentBuilder::SetDocumentConfiguration(TInstancedStruct<FMetaSoundFrontendDocumentConfiguration> InConfig)
+void FMetaSoundFrontendDocumentBuilder::SetTemplate(TInstancedStruct<FMetaSoundFrontendDocumentTemplate> InConfig)
 {
 	using namespace Metasound::Frontend;
 
-	GetDocumentChecked().Configuration = MoveTemp(InConfig);
-	DocumentDelegates->OnDocumentConfigurationChanged.Broadcast({ EDocumentConfigChangeType::Struct });
+	GetDocumentChecked().Template = MoveTemp(InConfig);
+	DocumentDelegates->OnDocumentTemplateChanged.Broadcast({ EDocumentTemplateChangeType::Struct });
 }
 
-void FMetaSoundFrontendDocumentBuilder::SetDocumentConfigurationProperties(TFunctionRef<void(TStructView<FMetaSoundFrontendDocumentConfiguration>)> SetPropertyFunction)
+void FMetaSoundFrontendDocumentBuilder::SetTemplateProperties(TFunctionRef<void(TStructView<FMetaSoundFrontendDocumentTemplate>)> SetPropertyFunction)
 {
 	using namespace Metasound::Frontend;
-	SetPropertyFunction(GetDocumentChecked().Configuration);
-	DocumentDelegates->OnDocumentConfigurationChanged.Broadcast({ EDocumentConfigChangeType::Property });
+	SetPropertyFunction(GetDocumentChecked().Template);
+	DocumentDelegates->OnDocumentTemplateChanged.Broadcast({ EDocumentTemplateChangeType::Property });
 }
 
 #if WITH_EDITORONLY_DATA
@@ -6213,9 +6206,9 @@ bool FMetaSoundFrontendDocumentBuilder::SetGraphInputDefault(FName InputName, FM
 					.ChangeType = InChangeType
 				});
 
-			if (const FDocumentConfigurationVertexMetadata* ConfigVertexMetadata = FindConstVertexMetadata<FDocumentConfigurationVertexMetadata>(Input.NodeID))
+			if (const FTemplateVertexMetadata* ConfigVertexMetadata = FindConstVertexMetadata<FTemplateVertexMetadata>(Input.NodeID))
 			{
-				const_cast<FMetaSoundFrontendDocumentConfigurationVertexMetadata*>(ConfigVertexMetadata)->OnDefaultUpdated(*this, Input);
+				const_cast<FMetaSoundFrontendDocumentTemplateVertexMetadata*>(ConfigVertexMetadata)->OnDefaultUpdated(*this, Input);
 			}
 
 			return true;
@@ -6293,9 +6286,9 @@ bool FMetaSoundFrontendDocumentBuilder::SetGraphInputDefaults(FName InputName, T
 				});
 			}
 
-			if (FDocumentVertexConfiguration* DocConfig = Document.Configuration.GetMutablePtr<FDocumentVertexConfiguration>())
+			if (FDocumentVertexTemplate* DocTemplate = Document.Template.GetMutablePtr<FDocumentVertexTemplate>())
 			{
-				if (FDocumentConfigurationVertexMetadata* VertexMetadata = DocConfig->FindVertexMetadata(Input.NodeID))
+				if (FTemplateVertexMetadata* VertexMetadata = DocTemplate->FindVertexMetadata(Input.NodeID))
 				{
 					VertexMetadata->OnDefaultUpdated(*this, Input);
 				}
diff --git a/Engine/Plugins/Runtime/Metasound/Source/MetasoundFrontend/Private/MetasoundFrontendDocumentVersioning.cpp b/Engine/Plugins/Runtime/Metasound/Source/MetasoundFrontend/Private/MetasoundFrontendDocumentVersioning.cpp
--- a/Engine/Plugins/Runtime/Metasound/Source/MetasoundFrontend/Private/MetasoundFrontendDocumentVersioning.cpp
+++ b/Engine/Plugins/Runtime/Metasound/Source/MetasoundFrontend/Private/MetasoundFrontendDocumentVersioning.cpp
@@ -7,7 +7,7 @@
 #include "Algo/Transform.h"
 #include "Algo/Unique.h"
 #include "CoreGlobals.h"
-#include "DocumentConfigurations/MetasoundFrontendPresetDocumentConfiguration.h"
+#include "DocumentTemplates/MetasoundFrontendPresetTemplate.h"
 #include "HAL/IConsoleManager.h"
 #include "Interfaces/MetasoundFrontendInterface.h"
 #include "Interfaces/MetasoundFrontendInterfaceRegistry.h"
@@ -1130,14 +1130,14 @@ PRAGMA_ENABLE_DEPRECATION_WARNINGS
 				{
 					bIsPreset = false;
 
-					using FPresetConfigStruct = TInstancedStruct<FMetaSoundFrontendPresetDocumentConfiguration>;
-					FPresetConfigStruct PresetConfig = FPresetConfigStruct::Make<FMetaSoundFrontendPresetDocumentConfiguration>();
-					PresetConfig.GetMutable().Parent = FindReferencedMetaSound_Legacy(OutBuilder);
+					using FPresetConfigStruct = TInstancedStruct<FMetaSoundFrontendPresetTemplate>;
+					FPresetConfigStruct PresetTemplate = FPresetConfigStruct::Make<FMetaSoundFrontendPresetTemplate>();
+					PresetTemplate.GetMutable().Parent = FindReferencedMetaSound_Legacy(OutBuilder);
 
 					for (const FMetasoundFrontendClassInput& Input : Document.RootGraph.GetDefaultInterface().Inputs)
 					{
 						// Some inputs don't support inheriting a default (ex. triggers), so ignore if no metadata is added.
-						if (FPresetVertexMetadata* VertexMetadata = PresetConfig.GetMutable().AddVertexMetadata<FPresetVertexMetadata>(OutBuilder, Input))
+						if (FPresetVertexMetadata* VertexMetadata = PresetTemplate.GetMutable().AddVertexMetadata<FPresetVertexMetadata>(OutBuilder, Input))
 						{
 PRAGMA_DISABLE_DEPRECATION_WARNINGS
 							const bool bOverrideInheritedDefault = !PresetOptions.InputsInheritingDefault.Contains(Input.Name);
@@ -1147,9 +1147,9 @@ PRAGMA_ENABLE_DEPRECATION_WARNINGS
 					}
 
 					{
-						using FConfigStruct = TInstancedStruct<FMetaSoundFrontendDocumentConfiguration>;
-						FConfigStruct& ConfigRef = const_cast<FConfigStruct&>(Document.Configuration);
-						ConfigRef = MoveTemp(PresetConfig);
+						using FTemplateStruct = TInstancedStruct<FMetaSoundFrontendDocumentTemplate>;
+						FTemplateStruct& TemplateRef = const_cast<FTemplateStruct&>(Document.Template);
+						TemplateRef = MoveTemp(PresetTemplate);
 					}
 				}
 
diff --git a/Engine/Plugins/Runtime/Metasound/Source/MetasoundFrontend/Private/MetasoundFrontendInvalidController.h b/Engine/Plugins/Runtime/Metasound/Source/MetasoundFrontend/Private/MetasoundFrontendInvalidController.h
--- a/Engine/Plugins/Runtime/Metasound/Source/MetasoundFrontend/Private/MetasoundFrontendInvalidController.h
+++ b/Engine/Plugins/Runtime/Metasound/Source/MetasoundFrontend/Private/MetasoundFrontendInvalidController.h
@@ -39,7 +39,7 @@ namespace Metasound
 			const FMetasoundFrontendNodeInterface& GetInvalidNodeInterface();
 			const FMetasoundFrontendClassMetadata& GetInvalidClassMetadata();
 
-			UE_DEPRECATED(5.8, "PresetOptions have been integrated into Document Configuration")
+			UE_DEPRECATED(5.8, "Preset options are now serialized in the Preset MetaSound Template")
 			const FMetasoundFrontendGraphClassPresetOptions& GetInvalidGraphClassPresetOptions();
 
 			const FMetasoundFrontendGraphClass& GetInvalidGraphClass();
diff --git a/Engine/Plugins/Runtime/Metasound/Source/MetasoundFrontend/Private/MetasoundFrontendTransform.cpp b/Engine/Plugins/Runtime/Metasound/Source/MetasoundFrontend/Private/MetasoundFrontendTransform.cpp
--- a/Engine/Plugins/Runtime/Metasound/Source/MetasoundFrontend/Private/MetasoundFrontendTransform.cpp
+++ b/Engine/Plugins/Runtime/Metasound/Source/MetasoundFrontend/Private/MetasoundFrontendTransform.cpp
@@ -2,7 +2,7 @@
 #include "MetasoundFrontendTransform.h"
 
 #include "Algo/Transform.h"
-#include "DocumentConfigurations/MetasoundFrontendPresetDocumentConfiguration.h"
+#include "DocumentTemplates/MetasoundFrontendPresetTemplate.h"
 #include "Interfaces/MetasoundFrontendInterface.h"
 #include "Interfaces/MetasoundFrontendInterfaceRegistry.h"
 #include "NodeTemplates/MetasoundFrontendNodeTemplateInput.h"
@@ -507,7 +507,7 @@ namespace Metasound
 			bool bDidEdit = false;
 			
 			// If document config provided, run configure
-			if (TConstStructView<FMetaSoundFrontendDocumentConfiguration> DocConfig = InOutBuilderToTransform.GetConstDocumentConfiguration(); DocConfig.IsValid())
+			if (TConstStructView<FMetaSoundFrontendDocumentTemplate> DocTemplate = InOutBuilderToTransform.GetConstDocumentTemplate(); DocTemplate.IsValid())
 			{
 				InOutBuilderToTransform.ConfigureDocument();
 				
@@ -519,7 +519,7 @@ namespace Metasound
 			// Update external node dependencies until there are no more updates
 			// (A node could need a minor version update, then an update transform, 
 			// and then updated to another minor version which would require multiple passes).
-			// This should be run post document config as well as a config's implementing code may
+			// This should be run post document configuration as well as a template's implementing code may
 			// or may not be updated to generate newest dependencies.
 			bool bUpdated = true;
 			while (bUpdated)
@@ -657,9 +657,9 @@ namespace Metasound
 		{
 			METASOUND_TRACE_CPUPROFILER_EVENT_SCOPE(Metasound::Frontend::FRebuildPresetRootGraph::Transform);
 
-			// This function is to be moved to the preset doc configuration and should no longer be directly called outside that context.
-			const FMetaSoundFrontendPresetDocumentConfiguration* PresetConfig = InOutBuilderToTransform.GetConstDocumentConfigurationAs<FMetaSoundFrontendPresetDocumentConfiguration>();
-			if (!PresetConfig || !ParentBuilder)
+			// This function is to be moved (see FMetaSoundFrontendPresetTemplate) and should no longer be directly called outside that context.
+			const FMetaSoundFrontendPresetTemplate* PresetTemplate = InOutBuilderToTransform.GetConstTemplateAs<FMetaSoundFrontendPresetTemplate>();
+			if (!PresetTemplate || !ParentBuilder)
 			{
 				return false;
 			}
@@ -671,7 +671,7 @@ namespace Metasound
 
 #if WITH_EDITORONLY_DATA
 			// Cache off member metadata so it be can be readded if necessary after the graph is cleared 
-			FMemberIDToMetadataMap CachedMemberMetadata = CacheMemberMetadata(*PresetConfig, InOutBuilderToTransform);
+			FMemberIDToMetadataMap CachedMemberMetadata = CacheMemberMetadata(*PresetTemplate, InOutBuilderToTransform);
 #endif // WITH_EDITORONLY_DATA
 
 			FGuid PresetNodeID;
@@ -717,7 +717,7 @@ namespace Metasound
 #endif // WITH_EDITOR
 
 			// Connect parent graph to referenced graph
-			AddAndConnectInputs(*PresetConfig, ClassInputs, InOutBuilderToTransform, PresetNodeID);
+			AddAndConnectInputs(*PresetTemplate, ClassInputs, InOutBuilderToTransform, PresetNodeID);
 			AddAndConnectOutputs(ClassOutputs, InOutBuilderToTransform, PresetNodeID);
 
 #if WITH_EDITORONLY_DATA
@@ -743,7 +743,7 @@ namespace Metasound
 			}
 		}
 		
-		FRebuildPresetRootGraph::FMemberIDToMetadataMap FRebuildPresetRootGraph::CacheMemberMetadata(const FMetaSoundFrontendPresetDocumentConfiguration& PresetConfig, FMetaSoundFrontendDocumentBuilder& InOutBuilderToTransform) const
+		FRebuildPresetRootGraph::FMemberIDToMetadataMap FRebuildPresetRootGraph::CacheMemberMetadata(const FMetaSoundFrontendPresetTemplate& PresetTemplate, FMetaSoundFrontendDocumentBuilder& InOutBuilderToTransform) const
 		{
 			FMemberIDToMetadataMap CachedMemberMetadata;
 
@@ -773,7 +773,7 @@ namespace Metasound
 					// If the input vertex already exists in the parent graph,
 					// check if parent should be used or not from set of managed
 					// input names.
-					if (const FPresetVertexMetadata* VertexMetadata = PresetConfig.FindConstVertexMetadata<FPresetVertexMetadata>(MemberID);
+					if (const FPresetVertexMetadata* VertexMetadata = PresetTemplate.FindConstVertexMetadata<FPresetVertexMetadata>(MemberID);
 						VertexMetadata && !VertexMetadata->bOverrideInheritedDefault)
 					{
 						UMetaSoundFrontendMemberMetadata* ParentMetadata = ParentBuilder->FindMemberMetadata(ParentClassInput.NodeID);
@@ -802,7 +802,7 @@ namespace Metasound
 		}
 #endif // WITH_EDITORONLY_DATA
 
-		void FRebuildPresetRootGraph::AddAndConnectInputs(const FMetaSoundFrontendPresetDocumentConfiguration& PresetConfig, const TArray<FMetasoundFrontendClassInput>& InClassInputs, FMetaSoundFrontendDocumentBuilder& InOutBuilderToTransform, const FGuid& InReferencedNodeID) const
+		void FRebuildPresetRootGraph::AddAndConnectInputs(const FMetaSoundFrontendPresetTemplate& PresetTemplate, const TArray<FMetasoundFrontendClassInput>& InClassInputs, FMetaSoundFrontendDocumentBuilder& InOutBuilderToTransform, const FGuid& InReferencedNodeID) const
 		{
 			TArray<const FMetasoundFrontendVertex*> ReferencedNodeInputVertices;
 
@@ -980,8 +980,8 @@ namespace Metasound
 			const FMetasoundFrontendDocument& ParentDocument = ParentBuilder->GetConstDocumentChecked();
 			const FMetasoundFrontendDocument& DocumentToTransform = InDocumentToTransformBuilder.GetConstDocumentChecked();
 
-			const FMetaSoundFrontendPresetDocumentConfiguration* PresetConfig = InDocumentToTransformBuilder.GetConstDocumentConfigurationAs<FMetaSoundFrontendPresetDocumentConfiguration>();
-			check(PresetConfig);
+			const FMetaSoundFrontendPresetTemplate* PresetTemplate = InDocumentToTransformBuilder.GetConstTemplateAs<FMetaSoundFrontendPresetTemplate>();
+			check(PresetTemplate);
 
 			// Iterate through all input nodes of referenced graph
 			for (const FMetasoundFrontendClassInput& ParentClassInput : ParentDocument.RootGraph.GetDefaultInterface().Inputs)
@@ -1019,7 +1019,7 @@ namespace Metasound
 					// If the input vertex already exists in the parent graph,
 					// check if parent should be used or not from set of managed
 					// input names.
-					if (const FPresetVertexMetadata* VertexMetadata = PresetConfig->FindConstVertexMetadata<FPresetVertexMetadata>(GraphInput->NodeID))
+					if (const FPresetVertexMetadata* VertexMetadata = PresetTemplate->FindConstVertexMetadata<FPresetVertexMetadata>(GraphInput->NodeID))
 					{
 						if (!VertexMetadata->bOverrideInheritedDefault)
 						{
diff --git a/Engine/Plugins/Runtime/Metasound/Source/MetasoundFrontend/Public/DocumentConfigurations/MetasoundFrontendDocumentConfiguration.h b/Engine/Plugins/Runtime/Metasound/Source/MetasoundFrontend/Public/DocumentConfigurations/MetasoundFrontendDocumentConfiguration.h
deleted file mode 100644
diff --git a/Engine/Plugins/Runtime/Metasound/Source/MetasoundFrontend/Public/DocumentConfigurations/MetasoundFrontendDocumentVertexConfiguration.h b/Engine/Plugins/Runtime/Metasound/Source/MetasoundFrontend/Public/DocumentConfigurations/MetasoundFrontendDocumentVertexConfiguration.h
deleted file mode 100644
diff --git a/Engine/Plugins/Runtime/Metasound/Source/MetasoundFrontend/Public/DocumentConfigurations/MetasoundFrontendPresetDocumentConfiguration.h b/Engine/Plugins/Runtime/Metasound/Source/MetasoundFrontend/Public/DocumentConfigurations/MetasoundFrontendPresetDocumentConfiguration.h
deleted file mode 100644
diff --git a/Engine/Plugins/Runtime/Metasound/Source/MetasoundFrontend/Public/DocumentTemplates/MetasoundFrontendDocumentTemplate.h b/Engine/Plugins/Runtime/Metasound/Source/MetasoundFrontend/Public/DocumentTemplates/MetasoundFrontendDocumentTemplate.h
new file mode 100644
--- /dev/null
+++ b/Engine/Plugins/Runtime/Metasound/Source/MetasoundFrontend/Public/DocumentTemplates/MetasoundFrontendDocumentTemplate.h
@@ -0,0 +1,117 @@
+// Copyright Epic Games, Inc. All Rights Reserved.
+#pragma once
+
+#if WITH_EDITOR
+#include "Delegates/DelegateCombinations.h"
+#include "Misc/DataValidation.h"
+#endif // WITH_EDITOR
+
+#include "Internationalization/Text.h"
+
+#include "MetasoundFrontendDocumentTemplate.generated.h"
+
+#define UE_API METASOUNDFRONTEND_API
+
+// Forward Declarations
+struct FMetaSoundFrontendDocumentBuilder;
+class UMetaSoundFrontendMemberMetadata;
+
+enum class EMetasoundFrontendClassType : uint8;
+
+#if WITH_EDITOR
+namespace Metasound::Frontend
+{
+	DECLARE_MULTICAST_DELEGATE_OneParam(FOnConfigPropertyChangedDelegate, const FPropertyChangedEvent& /* InEvent */);
+} // namespace Metasound::Frontend
+#endif // WITH_EDITOR
+
+// Struct housed on a MetaSoundDocument and used to define procedural behavior applied when making
+// modifications to a MetaSound prior to serialization. Templates enable MetaSounds to be
+// procedurally generated at edit and cook-time to adhere to a specific design or topology anytime
+// a given configuration or its respective properties are changed, enabling faster and more accurate
+// iteration where fully-customized MetaSound edit behavior is not desired.
+USTRUCT(BlueprintType, meta = (Hidden))
+struct FMetaSoundFrontendDocumentTemplate
+{
+	GENERATED_BODY()
+
+	virtual ~FMetaSoundFrontendDocumentTemplate() = default;
+
+	// Executes procedural configuration of document template. Returns true if template applied
+	// to document, false if not (resulting in a follow-up call to reset the document).
+	UE_INTERNAL UE_API virtual bool ConfigureDocument(FMetaSoundFrontendDocumentBuilder& OutBuilder);
+
+#if WITH_EDITOR
+	struct FEditorOptions
+	{
+		// Asset class types this template supports as a selection (passed on initial creation) when user
+		// creates a configured MetaSound from an action (eg. right-click action in the Content Browser right-click
+		// menu). Defaults to all MetaSound asset types (eg. Source & Patch) if unset.
+		UE_EXPERIMENTAL(5.8, "MetaSound Templates are in active development and subject to change")
+		TArray<TSoftClassPtr<UObject>> AssetActionClasses;
+
+		// MetaSound asset types this template can be set to when user creates a configured MetaSound
+		// from an action (eg. right-click action in the Content Browser right-click menu). Defaults to
+		// all MetaSound asset types (eg. Source & Patch) if unset.
+		UE_EXPERIMENTAL(5.8, "MetaSound Templates are in active development and subject to change")
+		TArray<TSoftClassPtr<UObject>> MetaSoundClasses;
+
+		// ToolTip shown when selecting document template from list of available configurations
+		UE_EXPERIMENTAL(5.8, "MetaSound Templates are in active development and subject to change")
+		FText ToolTip;
+
+		// Whether or not template supports custom tabs (i.e. custom template widgets)
+		UE_EXPERIMENTAL(5.8, "MetaSound Templates are in active development and subject to change")
+		bool bCustomTabEnabled = true;
+
+		// Whether or not interface editing is enabled in MetaSound asset editors
+		UE_EXPERIMENTAL(5.8, "MetaSound Templates are in active development and subject to change")
+		bool bInterfaceEditingEnabled = false;
+
+		// Whether or not member editing is enabled in MetaSound asset editors
+		UE_EXPERIMENTAL(5.8, "MetaSound Templates are in active development and subject to change")
+		bool bMemberEditingEnabled = false;
+
+		// Whether or not graph page editing is enabled in MetaSound asset editors
+		UE_EXPERIMENTAL(5.8, "MetaSound Templates are in active development and subject to change")
+		bool bPageGraphEditingEnabled = false;
+
+		// Whether or not graph editor is visible by default (user can overwrite if desired)
+		UE_EXPERIMENTAL(5.8, "MetaSound Templates are in active development and subject to change")
+		bool bGraphEditorVisible = false;
+
+		// Whether or not to display template properties in the details panel.
+		UE_EXPERIMENTAL(5.8, "MetaSound Templates are in active development and subject to change")
+		bool bTemplatePropertiesVisible = true;
+
+		// Default options if none are specified
+		UE_API static const FEditorOptions& GetDefaultOptions();
+	};
+
+	UE_EXPERIMENTAL(5.8, "MetaSound Templates are in active development and subject to change")
+	UE_API virtual const FEditorOptions& GetEditorOptions() const;
+
+#if WITH_EDITOR
+	UE_API virtual EDataValidationResult IsDataValid(const FMetaSoundFrontendDocumentBuilder& Builder, FDataValidationContext& InOutContext) const;
+#endif // WITH_EDITOR
+
+	UE_EXPERIMENTAL(5.8, "MetaSound Templates are in active development and subject to change")
+	UE_API virtual void OnAssetInitialized(TArray<UObject*> SelectedObjects, FMetaSoundFrontendDocumentBuilder& OutBuilder);
+
+	UE_EXPERIMENTAL(5.8, "MetaSound Templates are in active development and subject to change")
+	UE_API virtual void OnPropertyChanged(const FPropertyChangedEvent& InEvent, FMetaSoundFrontendDocumentBuilder& OutBuilder);
+
+	UE_API Metasound::Frontend::FOnConfigPropertyChangedDelegate& GetPropertyChangedDelegate() const;
+
+private:
+	mutable Metasound::Frontend::FOnConfigPropertyChangedDelegate PropertyChangedDelegate;
+#endif // WITH_EDITOR
+};
+
+// Convenience, less verbose aliases
+namespace Metasound::Frontend
+{
+	using FMetaSoundTemplate = FMetaSoundFrontendDocumentTemplate;
+} // namespace Metasound::Frontend;
+
+#undef UE_API
diff --git a/Engine/Plugins/Runtime/Metasound/Source/MetasoundFrontend/Public/DocumentTemplates/MetasoundFrontendDocumentVertexTemplate.h b/Engine/Plugins/Runtime/Metasound/Source/MetasoundFrontend/Public/DocumentTemplates/MetasoundFrontendDocumentVertexTemplate.h
new file mode 100644
--- /dev/null
+++ b/Engine/Plugins/Runtime/Metasound/Source/MetasoundFrontend/Public/DocumentTemplates/MetasoundFrontendDocumentVertexTemplate.h
@@ -0,0 +1,136 @@
+// Copyright Epic Games, Inc. All Rights Reserved.
+#pragma once
+
+#include "MetasoundFrontendDocument.h"
+#include "MetasoundFrontendDocumentTemplate.h"
+
+#include "MetasoundFrontendDocumentVertexTemplate.generated.h"
+
+#define UE_API METASOUNDFRONTEND_API
+
+// Forward Declarations
+struct FMetaSoundFrontendDocumentBuilder;
+
+
+// Vertex metadata displayed on MetaSound inputs & outputs pertaining to a given configuration.
+USTRUCT(BlueprintType)
+struct FMetaSoundFrontendDocumentTemplateVertexMetadata
+{
+	GENERATED_BODY()
+
+public:
+	FMetaSoundFrontendDocumentTemplateVertexMetadata() = default;
+
+	virtual ~FMetaSoundFrontendDocumentTemplateVertexMetadata() = default;
+
+#if WITH_EDITOR
+	// If true, default values are directly editable if the parent configuration lists member editing as enabled.
+	UE_EXPERIMENTAL(5.8, "MetaSound Templates are in active development and subject to change")
+	virtual bool IsDefaultEditable() const { return true; }
+#endif // WITH_EDITOR
+
+	UE_EXPERIMENTAL(5.8, "MetaSound Templates are in active development and subject to change")
+	virtual void OnDefaultReset(const FMetaSoundFrontendDocumentBuilder& Builder, const FMetasoundFrontendClassVertex& InVertex) { }
+
+	UE_EXPERIMENTAL(5.8, "MetaSound Templates are in active development and subject to change")
+	virtual void OnDefaultUpdated(const FMetaSoundFrontendDocumentBuilder& Builder, const FMetasoundFrontendClassVertex& InVertex) { }
+
+#if WITH_EDITOR
+	UE_EXPERIMENTAL(5.8, "MetaSound Templates are in active development and subject to change")
+	virtual void OnPropertyChanged(const FPropertyChangedEvent& InEvent, FMetaSoundFrontendDocumentBuilder& OutBuilder, const FMetasoundFrontendClassVertex& InVertex) { }
+#endif // WITH_EDITOR
+};
+
+// Template of a MetaSound that stores vertex (input or output) Metadata to display alongside standard properties in MetaSound vertex editors.
+USTRUCT(BlueprintType, meta = (Hidden))
+struct FMetaSoundFrontendDocumentVertexTemplate : public FMetaSoundFrontendDocumentTemplate
+{
+	GENERATED_BODY()
+
+	virtual ~FMetaSoundFrontendDocumentVertexTemplate() = default;
+
+	// Creates and stores new VertexMetadata. Overwrites if metadata already exists with
+	// default constructed VertexMetadata if template supports Metadata.
+	UE_INTERNAL FMetaSoundFrontendDocumentTemplateVertexMetadata* AddVertexMetadata(const FMetaSoundFrontendDocumentBuilder& Builder, const FMetasoundFrontendClassVertex& InVertex);
+
+	// Creates and stores new VertexMetadata with the provided type and name, asserting that the type created by the overloaded template
+	// matches the template type. Overwrites if metadata already exists with default constructed VertexMetadata. Returns null if template
+	// does not support Metadata.
+	template <typename TTemplateVertexMetadata>
+	UE_INTERNAL TTemplateVertexMetadata* AddVertexMetadata(const FMetaSoundFrontendDocumentBuilder& Builder, const FMetasoundFrontendClassVertex& InVertex)
+	{
+		if (TInstancedStruct<FMetaSoundFrontendDocumentTemplateVertexMetadata>* Entry = VertexMetadata.Find(InVertex.NodeID))
+		{
+			return &Entry->GetMutable<TTemplateVertexMetadata>();
+		}
+
+PRAGMA_DISABLE_EXPERIMENTAL_WARNINGS
+		TInstancedStruct<FMetaSoundFrontendDocumentTemplateVertexMetadata> NewMetadata = ConstructVertexMetadata(Builder, InVertex);
+PRAGMA_ENABLE_EXPERIMENTAL_WARNINGS
+
+		if (NewMetadata.IsValid())
+		{
+			return &VertexMetadata.Add(InVertex.NodeID, MoveTemp(NewMetadata)).GetMutable<TTemplateVertexMetadata>();
+		}
+
+		return nullptr;
+	}
+
+	// Returns whether or not there is existing Metadata with the given NodeID.
+	UE_API bool ContainsVertexMetadata(const FGuid& NodeID) const;
+
+	// Helper function to returns VertexMetadata with the provided class type and name and cast to underlying config type.
+	template <typename TTemplateVertexMetadata = FMetaSoundFrontendDocumentTemplateVertexMetadata>
+	UE_INTERNAL TTemplateVertexMetadata* FindVertexMetadata(const FGuid& NodeID)
+	{
+		if (TInstancedStruct<FMetaSoundFrontendDocumentTemplateVertexMetadata>* Entry = VertexMetadata.Find(NodeID))
+		{
+			return Entry->GetMutablePtr<TTemplateVertexMetadata>();
+		}
+
+		return nullptr;
+	}
+
+	// Helper function to returns VertexMetadata with the provided class type and name and cast to underlying config type.
+	template <typename TTemplateVertexMetadata = FMetaSoundFrontendDocumentTemplateVertexMetadata>
+	const TTemplateVertexMetadata* FindConstVertexMetadata(const FGuid& NodeID) const
+	{
+		if (const TInstancedStruct<FMetaSoundFrontendDocumentTemplateVertexMetadata>* Entry = VertexMetadata.Find(NodeID))
+		{
+			return Entry->GetPtr<TTemplateVertexMetadata>();
+		}
+
+		return nullptr;
+	}
+
+	const UScriptStruct* GetVertexMetadataStruct(const FGuid& NodeID) const
+	{
+		if (const TInstancedStruct<FMetaSoundFrontendDocumentTemplateVertexMetadata>* Entry = VertexMetadata.Find(NodeID))
+		{
+			return Entry->GetScriptStruct();
+		}
+
+		return nullptr;
+	}
+
+	// Removes VertexMetadata with the provided NodeID. Returns if Metadata was found and removed or not.
+	UE_EXPERIMENTAL(5.8, "MetaSound Templates are in active development and subject to change")
+	UE_API bool RemoveVertexMetadata(const FGuid& NodeID);
+
+protected:
+	UE_EXPERIMENTAL(5.8, "MetaSound Templates are in active development and subject to change")
+	UE_API virtual TInstancedStruct<FMetaSoundFrontendDocumentTemplateVertexMetadata> ConstructVertexMetadata(const FMetaSoundFrontendDocumentBuilder& Builder, const FMetasoundFrontendClassVertex& InVertex) const;
+
+private:
+	UPROPERTY()
+	TMap<FGuid, TInstancedStruct<FMetaSoundFrontendDocumentTemplateVertexMetadata>> VertexMetadata;
+};
+
+// Convenience, less verbose aliases
+namespace Metasound::Frontend
+{
+	using FDocumentVertexTemplate = FMetaSoundFrontendDocumentVertexTemplate;
+	using FTemplateVertexMetadata = FMetaSoundFrontendDocumentTemplateVertexMetadata;
+} // namespace Metasound::Frontend;
+
+#undef UE_API
diff --git a/Engine/Plugins/Runtime/Metasound/Source/MetasoundFrontend/Public/DocumentTemplates/MetasoundFrontendPresetTemplate.h b/Engine/Plugins/Runtime/Metasound/Source/MetasoundFrontend/Public/DocumentTemplates/MetasoundFrontendPresetTemplate.h
new file mode 100644
--- /dev/null
+++ b/Engine/Plugins/Runtime/Metasound/Source/MetasoundFrontend/Public/DocumentTemplates/MetasoundFrontendPresetTemplate.h
@@ -0,0 +1,76 @@
+// Copyright Epic Games, Inc. All Rights Reserved.
+#pragma once
+
+#include "MetasoundDocumentInterface.h"
+#include "MetasoundFrontendDocument.h"
+#include "MetasoundFrontendDocumentVertexTemplate.h"
+#include "UObject/NameTypes.h"
+
+#include "MetasoundFrontendPresetTemplate.generated.h"
+
+#define UE_API METASOUNDFRONTEND_API
+
+// Forward Declarations
+class UMetaSoundFrontendMemberMetadata;
+enum class EMetasoundFrontendClassType : uint8;
+
+
+USTRUCT(BlueprintType, DisplayName = "Preset")
+struct FMetaSoundFrontendPresetTemplate final : public FMetaSoundFrontendDocumentVertexTemplate
+{
+	GENERATED_BODY()
+
+	virtual ~FMetaSoundFrontendPresetTemplate() = default;
+
+#if WITH_EDITORONLY_DATA
+	UE_INTERNAL UE_API virtual bool ConfigureDocument(FMetaSoundFrontendDocumentBuilder& OutBuilder) override;
+#endif // WITH_EDITORONLY_DATA
+
+#if WITH_EDITOR
+	UE_EXPERIMENTAL(5.8, "MetaSound Templates are in active development and subject to change")
+	UE_API virtual const FMetaSoundFrontendDocumentTemplate::FEditorOptions& GetEditorOptions() const override;
+
+	UE_EXPERIMENTAL(5.8, "MetaSound Templates are in active development and subject to change")
+	UE_API EDataValidationResult IsDataValid(const FMetaSoundFrontendDocumentBuilder& Builder, FDataValidationContext& InOutContext) const override;
+
+	UE_EXPERIMENTAL(5.8, "MetaSound Templates are in active development and subject to change")
+	UE_API virtual void OnAssetInitialized(TArray<UObject*> SelectedObjects, FMetaSoundFrontendDocumentBuilder& OutBuilder) override;
+
+	UE_EXPERIMENTAL(5.8, "MetaSound Templates are in active development and subject to change")
+	UE_API virtual void OnPropertyChanged(const FPropertyChangedEvent& InEvent, FMetaSoundFrontendDocumentBuilder& OutBuilder) override;
+#endif // WITH_EDITOR
+
+protected:
+	UE_INTERNAL UE_API virtual TInstancedStruct<FMetaSoundFrontendDocumentTemplateVertexMetadata> ConstructVertexMetadata(const FMetaSoundFrontendDocumentBuilder& Builder, const FMetasoundFrontendClassVertex& InVertex) const override;
+
+public:
+	UPROPERTY(EditAnywhere, BlueprintReadOnly, Category = "Preset Options")
+	TScriptInterface<IMetaSoundDocumentInterface> Parent;
+};
+
+USTRUCT(BlueprintType)
+struct FMetaSoundFrontendPresetVertexMetadata : public FMetaSoundFrontendDocumentTemplateVertexMetadata
+{
+	GENERATED_BODY()
+
+	virtual ~FMetaSoundFrontendPresetVertexMetadata() = default;
+
+#if WITH_EDITOR
+	virtual bool IsDefaultEditable() const override;
+#endif // WITH_EDITOR
+
+	virtual void OnDefaultReset(const FMetaSoundFrontendDocumentBuilder& Builder, const FMetasoundFrontendClassVertex& InVertex) override;
+	virtual void OnDefaultUpdated(const FMetaSoundFrontendDocumentBuilder& Builder, const FMetasoundFrontendClassVertex& InVertex) override;
+
+	// Enables overriding the input's inherited default value otherwise provided by the parent graph.
+	// Setting to true disables the configuration utilizing the input's default value if updated on the parent asset.
+	UPROPERTY(EditAnywhere, BlueprintReadOnly, Category = PresetOptions)
+	bool bOverrideInheritedDefault = false;
+};
+
+namespace Metasound::Frontend
+{
+	using FPresetTemplate = FMetaSoundFrontendPresetTemplate;
+	using FPresetVertexMetadata = FMetaSoundFrontendPresetVertexMetadata;
+} // namespace Metasound::Frontend;
+#undef UE_API
\ No newline at end of file
diff --git a/Engine/Plugins/Runtime/Metasound/Source/MetasoundFrontend/Public/MetasoundFrontendController.h b/Engine/Plugins/Runtime/Metasound/Source/MetasoundFrontend/Public/MetasoundFrontendController.h
--- a/Engine/Plugins/Runtime/Metasound/Source/MetasoundFrontend/Public/MetasoundFrontendController.h
+++ b/Engine/Plugins/Runtime/Metasound/Source/MetasoundFrontend/Public/MetasoundFrontendController.h
@@ -2,14 +2,10 @@
 
 #pragma once
 
-#include "MetasoundAccessPtr.h"
-#include "MetasoundBuilderInterface.h"
 #include "MetasoundFrontend.h"
 #include "MetasoundFrontendDocument.h"
 #include "MetasoundFrontendDocumentAccessPtr.h"
 #include "MetasoundFrontendNodeTemplateRegistry.h"
-#include "MetasoundFrontendRegistries.h"
-#include "MetasoundGraph.h"
 #include "MetasoundVertex.h"
 #include "StructUtils/InstancedStruct.h"
 #include "StructUtils/StructView.h"
@@ -594,10 +590,10 @@ namespace Metasound
 			/** Returns the ClassID associated with this graph. */
 			virtual FGuid GetClassID() const = 0;
 
-			UE_DEPRECATED(5.8, "PresetOptions have been integrated into Document Configuration")
+			UE_DEPRECATED(5.8, "Preset options are now serialized in the Preset MetaSound Template")
 			virtual const FMetasoundFrontendGraphClassPresetOptions& GetGraphPresetOptions() const;
 
-			UE_DEPRECATED(5.8, "PresetOptions have been integrated into Document Configuration")
+			UE_DEPRECATED(5.8, "Preset options are now serialized in the Preset MetaSound Template")
 			virtual void SetGraphPresetOptions(const FMetasoundFrontendGraphClassPresetOptions& InPresetOptions);
 
 			/** Return the metadata for the current graph. */
@@ -653,13 +649,13 @@ namespace Metasound
 			/** Returns all input nodes in the graph. */
 			virtual TArray<FConstNodeHandle> GetConstInputNodes() const = 0;
 
-			UE_DEPRECATED(5.8, "PresetOptions have been integrated into Document Configuration")
+			UE_DEPRECATED(5.8, "Preset options are now serialized in the Preset MetaSound Template")
 			virtual const TSet<FName>& GetInputsInheritingDefault() const;
 
-			UE_DEPRECATED(5.8, "PresetOptions have been integrated into Document Configuration")
+			UE_DEPRECATED(5.8, "Preset options are now serialized in the Preset MetaSound Template")
 			virtual bool SetInputInheritsDefault(FName InName, bool bDefaultIsInherited);
 
-			UE_DEPRECATED(5.8, "PresetOptions have been integrated into Document Configuration")
+			UE_DEPRECATED(5.8, "Preset options are now serialized in the Preset MetaSound Template")
 			virtual void SetInputsInheritingDefault(TSet<FName>&& InNames);
 
 			/** Adds a new variable to the graph */
diff --git a/Engine/Plugins/Runtime/Metasound/Source/MetasoundFrontend/Public/MetasoundFrontendDocument.h b/Engine/Plugins/Runtime/Metasound/Source/MetasoundFrontend/Public/MetasoundFrontendDocument.h
--- a/Engine/Plugins/Runtime/Metasound/Source/MetasoundFrontend/Public/MetasoundFrontendDocument.h
+++ b/Engine/Plugins/Runtime/Metasound/Source/MetasoundFrontend/Public/MetasoundFrontendDocument.h
@@ -1,11 +1,9 @@
 // Copyright Epic Games, Inc. All Rights Reserved.
 #pragma once
 
-#include "Algo/Transform.h"
 #include "Containers/Array.h"
 #include "Containers/Map.h"
-#include "DocumentConfigurations/MetasoundFrontendDocumentConfiguration.h"
-#include "HAL/IConsoleManager.h"
+#include "DocumentTemplates/MetasoundFrontendDocumentTemplate.h"
 #include "IAudioParameterInterfaceRegistry.h"
 #include "Internationalization/Text.h"
 #include "MetasoundAccessPtr.h"
@@ -16,12 +14,10 @@
 #include "Misc/EnumRange.h"
 #include "Misc/Guid.h"
 #include "StructUtils/InstancedStruct.h"
-#include "StructUtils/SharedStruct.h"
 #include "StructUtils/StructView.h"
 #include "Templates/Function.h"
 #include "Templates/Invoke.h"
 #include "Templates/TypeHash.h"
-#include "UObject/NoExportTypes.h"
 
 #if WITH_EDITORONLY_DATA
 #include "Types/SlateVector2.h"
@@ -1902,10 +1898,10 @@ struct FMetasoundFrontendGraphClassPresetOptions
 	GENERATED_BODY()
 
 #if WITH_EDITORONLY_DATA
-	UPROPERTY(meta = (DeprecatedProperty, DeprecationMessage = "5.8 - Preset options now serialized as Document Configuration. See 'FMetaSoundFrontendPresetDocumentConfiguration'"))
+	UPROPERTY(meta = (DeprecatedProperty, DeprecationMessage = "5.8 - Preset options are now serialized in the Preset MetaSound Template. See 'FMetaSoundFrontendPresetTemplate'"))
 	bool bIsPreset = false;
 
-	UPROPERTY(meta = (DeprecatedProperty, DeprecationMessage = "5.8 - Preset options now serialized as Document Configuration. See 'FMetaSoundFrontendPresetDocumentConfiguration'"))
+	UPROPERTY(meta = (DeprecatedProperty, DeprecationMessage = "5.8 - Preset options are now serialized in the Preset MetaSound Template. See 'FMetaSoundFrontendPresetTemplate'"))
 	TSet<FName> InputsInheritingDefault;
 #endif // WITH_EDITORONLY_DATA
 };
@@ -2076,7 +2072,7 @@ public:
 	TArray<FMetasoundFrontendClass> Dependencies;
 
 	UPROPERTY(EditAnywhere, Category = CustomView)
-	TInstancedStruct<FMetaSoundFrontendDocumentConfiguration> Configuration;
+	TInstancedStruct<FMetaSoundFrontendDocumentTemplate> Template;
 
 	uint32 GetNextIdCounter() const
 	{
diff --git a/Engine/Plugins/Runtime/Metasound/Source/MetasoundFrontend/Public/MetasoundFrontendDocumentBuilder.h b/Engine/Plugins/Runtime/Metasound/Source/MetasoundFrontend/Public/MetasoundFrontendDocumentBuilder.h
--- a/Engine/Plugins/Runtime/Metasound/Source/MetasoundFrontend/Public/MetasoundFrontendDocumentBuilder.h
+++ b/Engine/Plugins/Runtime/Metasound/Source/MetasoundFrontend/Public/MetasoundFrontendDocumentBuilder.h
@@ -2,7 +2,7 @@
 #pragma once
 
 #include "Containers/Map.h"
-#include "DocumentConfigurations/MetasoundFrontendDocumentVertexConfiguration.h"
+#include "DocumentTemplates/MetasoundFrontendDocumentVertexTemplate.h"
 #include "Interfaces/MetasoundFrontendInterfaceRegistry.h"
 #include "MetasoundFrontendDocumentCacheInterface.h"
 #include "MetasoundDataReference.h"
@@ -27,7 +27,7 @@
 
 
 // Forward Declarations
-struct FMetaSoundFrontendDocumentConfiguration;
+struct FMetaSoundFrontendDocumentTemplate;
 class FMetasoundAssetBase;
 
 
@@ -158,11 +158,15 @@ private:
 	// Convenience template function for attempting to get the underlying pointer of the set configuration.
 	// If the provided configuration type is not that which the current document is set to or no configuration
 	// is set, returns null.
-	template<typename TConfigType>
-	TConfigType* GetDocumentConfigurationAs()
+	template<typename TTemplateType>
+	TTemplateType* GetTemplateAs()
 	{
-		return GetDocument().Configuration.GetMutablePtr<TConfigType>();
+		return GetDocument().Template.GetMutablePtr<TTemplateType>();
 	}
+
+	template<typename TTemplateType>
+	UE_DEPRECATED(5.8, "Use GetTemplateAs() instead")
+	TTemplateType* GetDocumentConfigurationAs() { return GetTemplateAs<TTemplateType>(); }
 #endif // WITH_EDITORONLY_DATA
 
 	UE_API virtual void OnBeginActiveBuilder() override;
@@ -276,8 +280,8 @@ PRAGMA_ENABLE_EXPERIMENTAL_WARNINGS
 
 	// Clears document completely of all graph page data (nodes, edges, & member metadata), dependencies,
 	// interfaces, member metadata, preset state, etc. Leaves ClassMetadata intact. Reloads the builder state,
-	// so external delegates must be relinked if desired. NOTE: Configuration is not cleared and must be managed
-	// by separate API.
+	// so external delegates must be relinked if desired. NOTE: Template struct is not cleared and must be managed
+	// by separate API call (see SetTemplate').
 	UE_API void ClearDocument(TSharedRef<Metasound::Frontend::FDocumentModifyDelegates> ModifyDelegates);
 
 #if WITH_EDITORONLY_DATA
@@ -293,13 +297,13 @@ PRAGMA_ENABLE_EXPERIMENTAL_WARNINGS
 	UE_API bool ContainsEdge(const FMetasoundFrontendEdge& InEdge, const FGuid* InPageID = nullptr) const;
 	UE_API bool ContainsNode(const FGuid& InNodeID, const FGuid* InPageID = nullptr) const;
 
-	UE_DEPRECATED(5.8, "Use SetDocumentConfiguration instead")
+	UE_DEPRECATED(5.8, "Use SetTemplate instead")
 	UE_API bool ConvertFromPreset();
 
-	UE_DEPRECATED(5.8, "Use SetDocumentConfiguration instead")
+	UE_DEPRECATED(5.8, "Use SetTemplate instead")
 	UE_API bool ConvertToPreset(const FMetasoundFrontendDocument& InReferencedDocument, TSharedPtr<Metasound::Frontend::FDocumentModifyDelegates> ModifyDelegates = { });
 
-	UE_DEPRECATED(5.8, "Use SetDocumentConfiguration instead")
+	UE_DEPRECATED(5.8, "Use SetTemplate instead")
 	UE_API bool ConvertToPreset(const FMetaSoundFrontendDocumentBuilder& InReferencedDocumentBuilder, TSharedPtr<Metasound::Frontend::FDocumentModifyDelegates> ModifyDelegates = { });
 	
 
@@ -397,59 +401,69 @@ PRAGMA_ENABLE_EXPERIMENTAL_WARNINGS
 
 	UE_API const FMetasoundFrontendDocument& GetConstDocumentChecked() const;
 
-	// Convenience template function for attempting to get the underlying mutable
-	// vertex configuration metadata. If the provided vertex configuration type
-	// is not that which the current document is set to or no configuration is
+	// Convenience function for attempting to get the underlying mutable
+	// document template vertex metadata. If the provided template vertex type
+	// is not that which the current document is set to or no document template is
 	// set, returns null.
 	template<typename TVertexMetadataType>
 	UE_INTERNAL TVertexMetadataType* FindVertexMetadata(const FGuid& InNodeID) const
 	{
 		using namespace Metasound::Frontend;
 
-		if (FDocumentVertexConfiguration* Config = GetDocumentChecked().Configuration.GetMutablePtr<FDocumentVertexConfiguration>())
+		if (FDocumentVertexTemplate* Template = GetDocumentChecked().Template.GetMutablePtr<FDocumentVertexTemplate>())
 		{
-			return Config->FindVertexMetadata<TVertexMetadataType>(InNodeID);
+			return Template->FindVertexMetadata<TVertexMetadataType>(InNodeID);
 		}
 
 		return nullptr;
 	}
 
-	// Convenience template function for attempting to get the underlying const
-	// vertex configuration metadata. If the provided vertex configuration type
-	// is not that which the current document is set to or no configuration is
+	// Convenience function for attempting to get the underlying mutable
+	// document template vertex metadata. If the provided template vertex type
+	// is not that which the current document is set to or no document template is
 	// set, returns null.
 	template<typename TVertexMetadataType>
 	UE_INTERNAL const TVertexMetadataType* FindConstVertexMetadata(const FGuid& InNodeID) const
 	{
 		using namespace Metasound::Frontend;
 
-		if (const FDocumentVertexConfiguration* Config = GetConstDocumentConfiguration().GetPtr<FDocumentVertexConfiguration>())
+		if (const FDocumentVertexTemplate* Template = GetConstDocumentTemplate().GetPtr<FDocumentVertexTemplate>())
 		{
-			return Config->FindConstVertexMetadata<TVertexMetadataType>(InNodeID);
+			return Template->FindConstVertexMetadata<TVertexMetadataType>(InNodeID);
 		}
 
 		return nullptr;
 	}
 
 #if WITH_EDITOR
-	// Returns struct view of document configuration if set. Internal use only to avoid
+	// Returns mutable struct view of document template if set. Internal use only to avoid
 	// writing against struct members in invalid or dangerous, non-editor contexts (use
-	// 'SetDocumentConfiguration' instead for most script, gameplay or runtime applications).
-	UE_INTERNAL UE_API TStructView<FMetaSoundFrontendDocumentConfiguration> GetDocumentConfiguration() const;
+	// 'SetTemplate' instead for most script or edit-time applications).
+	UE_INTERNAL UE_API TStructView<FMetaSoundFrontendDocumentTemplate> GetDocumentTemplate() const;
+
+	UE_DEPRECATED(5.8, "Use GetDocumentTemplate() instead")
+	UE_INTERNAL TStructView<FMetaSoundFrontendDocumentTemplate> GetDocumentConfiguration() const { return GetDocumentTemplate(); }
 #endif // WITH_EDITOR
 
-	// Returns struct view of document configuration if set.
-	UE_API TConstStructView<FMetaSoundFrontendDocumentConfiguration> GetConstDocumentConfiguration() const;
+	// Returns struct view of document template if set.
+	UE_API TConstStructView<FMetaSoundFrontendDocumentTemplate> GetConstDocumentTemplate() const;
+
+	UE_DEPRECATED(5.8, "Use GetConstDocumentTemplate() instead")
+	TConstStructView<FMetaSoundFrontendDocumentTemplate> GetConstDocumentConfiguration() const { return GetConstDocumentTemplate(); }
 
 	// Convenience template function for attempting to get the underlying const pointer of the set configuration.
 	// If the provided configuration type is not that which the current document is set to or no configuration
 	// is set, returns null.
-	template<typename TConfigType>
-	const TConfigType* GetConstDocumentConfigurationAs() const
+	template<typename TTemplateType>
+	const TTemplateType* GetConstTemplateAs() const
 	{
-		return GetConstDocumentConfiguration().GetPtr<TConfigType>();
+		return GetConstDocumentTemplate().GetPtr<TTemplateType>();
 	}
 
+	template<typename TTemplateType>
+	UE_DEPRECATED(5.8, "Use GetConstTemplateAs() instead")
+	const TTemplateType* GetConstDocumentConfigurationAs() const { return GetConstTemplateAs<TTemplateType>(); }
+
 	UE_API const IMetaSoundDocumentInterface& GetConstDocumentInterfaceChecked() const;
 	UE_API const FString GetDebugName() const;
 
@@ -495,14 +509,14 @@ PRAGMA_ENABLE_EXPERIMENTAL_WARNINGS
 
 	UE_API FMetasoundAssetBase& GetMetasoundAsset() const;
 
-	UE_DEPRECATED(5.8, "Use DocumentConfiguration API to access configuration as type FMetaSoundFrontendPresetDocumentConfiguration and use the 'Parent' property directly.")
+	UE_DEPRECATED(5.8, "Use Template API to access configuration as type FMetaSoundFrontendPresetTemplate and use the 'Parent' property directly.")
 	UE_API FMetasoundAssetBase* GetReferencedPresetAsset() const;
 
 	UE_API int32 GetTransactionCount() const;
 
 	UE_API TArray<const FMetasoundFrontendNode*> GetGraphInputTemplateNodes(FName InInputName, const FGuid* InPageID = nullptr);
 
-	UE_DEPRECATED(5.8, "Input inheritance data is now stored on MetaSound Preset document configuration metadata.")
+	UE_DEPRECATED(5.8, "Input inheritance data is now stored on MetaSound Preset template metadata.")
 	UE_API const TSet<FName>* GetGraphInputsInheritingDefault() const;
 
 	UE_API EMetasoundFrontendVertexAccessType GetNodeInputAccessType(const FGuid& InNodeID, const FGuid& InVertexID, const FGuid* InPageID = nullptr) const;
@@ -551,13 +565,13 @@ PRAGMA_ENABLE_EXPERIMENTAL_WARNINGS
 
 	UE_API bool IsDependencyReferenced(const FGuid& InClassID) const;
 
-	// Determines if the builder's document configuration is of the given type
-	template<typename TConfigType>
+	// Determines if the builder is configured the given template type
+	template<typename TTemplateType>
 	const bool IsDocumentConfiguredBy() const
 	{
-		if (const UScriptStruct* Struct = GetConstDocumentChecked().Configuration.GetScriptStruct())
+		if (const UScriptStruct* Struct = GetConstDocumentChecked().Template.GetScriptStruct())
 		{
-			return Struct->IsChildOf(TConfigType::StaticStruct());
+			return Struct->IsChildOf(TTemplateType::StaticStruct());
 		}
 
 		return false;
@@ -631,7 +645,7 @@ PRAGMA_ENABLE_EXPERIMENTAL_WARNINGS
 
 	UE_API void Reload(TSharedPtr<Metasound::Frontend::FDocumentModifyDelegates> Delegates = {}, bool bPrimeCache = false);
 
-	UE_DEPRECATED(5.8, "Input inheritance data is now cached in doc config.  No need to call directly here. Use overload without bool parameter.")
+	UE_DEPRECATED(5.8, "Input inheritance data is now cached in MetaSound Preset template vertex metadata.  No need to call directly here. Use overload without 'bClearInheritsDefault' parameter.")
 	UE_API bool RemoveGraphInputDefault(FName InputName, const FGuid& InPageID, bool bClearInheritsDefault);
 	UE_API bool RemoveGraphInputDefault(FName InputName, const FGuid& InPageID);
 
@@ -667,23 +681,29 @@ PRAGMA_ENABLE_EXPERIMENTAL_WARNINGS
 
 PRAGMA_DISABLE_EXPERIMENTAL_WARNINGS
 	// Sets the entire internal configuration with that provided.
-	UE_API void SetDocumentConfiguration(TInstancedStruct<FMetaSoundFrontendDocumentConfiguration> InConfig);
+	UE_API void SetTemplate(TInstancedStruct<FMetaSoundFrontendDocumentTemplate> InConfig);
 
-	// Update the given document configuration inline using the internal struct view passed to a given function.
-	UE_API void SetDocumentConfigurationProperties(TFunctionRef<void(TStructView<FMetaSoundFrontendDocumentConfiguration>)> SetPropertyFunction);
+	UE_DEPRECATED(5.8, "Use SetTemplate() instead")
+	void SetDocumentConfiguration(TInstancedStruct<FMetaSoundFrontendDocumentTemplate> InConfig) { SetTemplate(MoveTemp(InConfig)); }
+
+	// Update the given document template inline using the internal struct view passed to a given function.
+	UE_API void SetTemplateProperties(TFunctionRef<void(TStructView<FMetaSoundFrontendDocumentTemplate>)> SetPropertyFunction);
+
+	UE_DEPRECATED(5.8, "Use SetTemplateProperties() instead")
+	void SetDocumentConfigurationProperties(TFunctionRef<void(TStructView<FMetaSoundFrontendDocumentTemplate>)> SetPropertyFunction) { SetTemplateProperties(SetPropertyFunction); }
 
 	// If struct type of function coincides with that of the given builder's MetaSound configuration, sets
-	// the given document configuration properties inline using the internal struct view passed to a given
+	// the given document template properties inline using the internal struct view passed to a given
 	// function. Returns true if set function successfully executed, false if not.
 	template <typename TConfig>
-	bool SetDocumentConfigurationProperties(TFunctionRef<void(TConfig&)> SetPropertyFunction)
+	bool SetTemplateProperties(TFunctionRef<void(TConfig&)> SetPropertyFunction)
 	{
 		bool bSucceeded = false;
 
-		const UScriptStruct* ScriptStruct = GetConstDocumentConfiguration().GetScriptStruct();
+		const UScriptStruct* ScriptStruct = GetConstDocumentTemplate().GetScriptStruct();
 		if (ScriptStruct && ScriptStruct->IsChildOf(TConfig::StaticStruct()))
 		{
-			SetDocumentConfigurationProperties([&SetPropertyFunction, &bSucceeded](TStructView<FMetaSoundFrontendDocumentConfiguration> Config)
+			SetTemplateProperties([&SetPropertyFunction, &bSucceeded](TStructView<FMetaSoundFrontendDocumentTemplate> Config)
 			{
 				if (TConfig* ConfigPtr = Config.GetPtr<TConfig>())
 				{
@@ -695,6 +715,10 @@ PRAGMA_DISABLE_EXPERIMENTAL_WARNINGS
 
 		return bSucceeded;
 	}
+
+	template <typename TConfig>
+	UE_DEPRECATED(5.8, "Use SetTemplateProperties() instead")
+	bool SetDocumentConfigurationProperties(TFunctionRef<void(TConfig&)> SetPropertyFunction) { return SetTemplateProperties<TConfig>(SetPropertyFunction); }
 PRAGMA_ENABLE_EXPERIMENTAL_WARNINGS
 
 #if WITH_EDITORONLY_DATA
@@ -730,10 +754,10 @@ PRAGMA_ENABLE_EXPERIMENTAL_WARNINGS
 	UE_API bool SetGraphInputDisplayName(FName InputName, FText DisplayName);
 #endif // WITH_EDITORONLY_DATA
 
-	UE_DEPRECATED(5.8, "Input inheritance information moved to preset document configuration")
+	UE_DEPRECATED(5.8, "Input inheritance information moved to preset document template")
 	UE_API bool SetGraphInputInheritsDefault(FName InName, bool bInputInheritsDefault, bool bForceUpdate = false);
 
-	UE_DEPRECATED(5.8, "Input inheritance information moved to preset document configuration")
+	UE_DEPRECATED(5.8, "Input inheritance information moved to preset document template")
 	UE_API bool SetGraphInputsInheritingDefault(TSet<FName>&& InNames);
 
 	// Sets a given graph input's name to a new name. Succeeds if the graph output exists and the new name is set (or is the same as the old name).
@@ -862,7 +886,7 @@ PRAGMA_ENABLE_EXPERIMENTAL_WARNINGS
 	UE_API const FTopLevelAssetPath GetBuilderClassPath() const;
 
 private:
-	bool bInDocConfigScope = false;
+	bool bInTemplateScope = false;
 
 	using FFinalizeNodeFunctionRef = TFunctionRef<void(FMetasoundFrontendNode&, const Metasound::Frontend::FNodeRegistryKey&)>;
 
diff --git a/Engine/Plugins/Runtime/Metasound/Source/MetasoundFrontend/Public/MetasoundFrontendDocumentModifyDelegates.h b/Engine/Plugins/Runtime/Metasound/Source/MetasoundFrontend/Public/MetasoundFrontendDocumentModifyDelegates.h
--- a/Engine/Plugins/Runtime/Metasound/Source/MetasoundFrontend/Public/MetasoundFrontendDocumentModifyDelegates.h
+++ b/Engine/Plugins/Runtime/Metasound/Source/MetasoundFrontend/Public/MetasoundFrontendDocumentModifyDelegates.h
@@ -46,7 +46,7 @@ namespace Metasound::Frontend
 		FGuid PageID;
 	};
 
-	enum class EDocumentConfigChangeType : uint8
+	enum class EDocumentTemplateChangeType : uint8
 	{
 		// Initial, unset value.
 		Invalid = 0,
@@ -70,17 +70,26 @@ namespace Metasound::Frontend
 	};
 	DECLARE_MULTICAST_DELEGATE_OneParam(FOnMetaSoundFrontendDocumentMutateInput, const FDocumentArrayPagedInputArgs& /* Args */);
 
-	struct FDocumentConfigChangedArgs
+	struct FDocumentTemplateChangedArgs
 	{
-		EDocumentConfigChangeType Type = EDocumentConfigChangeType::Invalid;
+		EDocumentTemplateChangeType Type = EDocumentTemplateChangeType::Invalid;
 	};
 
-	DECLARE_MULTICAST_DELEGATE_OneParam(FOnDocumentConfigurationChanged, const FDocumentConfigChangedArgs& /* Args */);
+	DECLARE_MULTICAST_DELEGATE_OneParam(FOnDocumentTemplateChanged, const FDocumentTemplateChangedArgs& /* Args */);
 
-	// Soft deprecated in 5.8 in favor of FOnDocumentConfigurationChanged.  Once delegate is removed, it can be officially deprecated.
+	UE_DEPRECATED(5.8, "Use FOnDocumentTemplateChanged instead")
+	typedef FOnDocumentTemplateChanged FOnDocumentConfigurationChanged;
+
+	UE_DEPRECATED(5.8, "Use EDocumentTemplateChangeType instead")
+	typedef EDocumentTemplateChangeType EDocumentConfigChangeType;
+
+	UE_DEPRECATED(5.8, "Use FDocumentTemplateChangedArgs instead")
+	typedef FDocumentTemplateChangedArgs FDocumentConfigChangedArgs;
+
+	// Soft deprecated in 5.8 in favor of FOnDocumentTemplateChanged.  Once delegate is removed, it can be officially deprecated.
 	struct FDocumentPresetStateChangedArgs {};
 
-	// Soft deprecated in 5.8 as usage is no longer valid (migrated to OnDocumentConfigurationChanged)
+	// Soft deprecated in 5.8 as usage is no longer valid (migrated to OnDocumentTemplateChanged)
 	DECLARE_MULTICAST_DELEGATE_OneParam(FOnDocumentPresetStateChanged, const FDocumentPresetStateChangedArgs& /* Args */);
 
 	DECLARE_MULTICAST_DELEGATE_OneParam(FOnDocumentPageAdded, const FDocumentMutatePageArgs& /* Args */);
@@ -166,10 +175,10 @@ PRAGMA_ENABLE_DEPRECATION_WARNINGS
 		FOnMetaSoundFrontendDocumentRemoveSwappingArray OnRemoveSwappingDependency;
 		FOnMetaSoundFrontendDocumentRenameClass OnRenamingDependencyClass;
 
-		// Soft deprecated in 5.8: "Use OnDocumentConfigurationChanged instead"
+		// Soft deprecated in 5.8: "Use OnDocumentTemplateChanged instead"
 		FOnDocumentPresetStateChanged OnPresetStateChanged;
 
-		FOnDocumentConfigurationChanged OnDocumentConfigurationChanged;
+		FOnDocumentTemplateChanged OnDocumentTemplateChanged;
 
 		FPageModifyDelegates PageDelegates;
 		FInterfaceModifyDelegates InterfaceDelegates;
diff --git a/Engine/Plugins/Runtime/Metasound/Source/MetasoundFrontend/Public/MetasoundFrontendTemplateNodeConfiguration.h b/Engine/Plugins/Runtime/Metasound/Source/MetasoundFrontend/Public/MetasoundFrontendTemplateNodeConfiguration.h
--- a/Engine/Plugins/Runtime/Metasound/Source/MetasoundFrontend/Public/MetasoundFrontendTemplateNodeConfiguration.h
+++ b/Engine/Plugins/Runtime/Metasound/Source/MetasoundFrontend/Public/MetasoundFrontendTemplateNodeConfiguration.h
@@ -6,6 +6,10 @@
 
 #include "MetasoundFrontendTemplateNodeConfiguration.generated.h"
 
+// Default node configuration used by the Node Template system (see IMetaSoundNodeTemplateBase
+// & FNodeTemplateBase). Distinct from FMetaSoundFrontendDocumentTemplate, which configures the
+// MetaSound document as a whole. Also distinct from FMetaSoundFrontendNodeConfiguration, which
+// stores per-node interface swap data (node configuration, a separate feature).
 USTRUCT()
 struct FMetaSoundFrontendTemplateNodeConfiguration : public FMetaSoundFrontendNodeConfiguration
 {
diff --git a/Engine/Plugins/Runtime/Metasound/Source/MetasoundFrontend/Public/MetasoundFrontendTransform.h b/Engine/Plugins/Runtime/Metasound/Source/MetasoundFrontend/Public/MetasoundFrontendTransform.h
--- a/Engine/Plugins/Runtime/Metasound/Source/MetasoundFrontend/Public/MetasoundFrontendTransform.h
+++ b/Engine/Plugins/Runtime/Metasound/Source/MetasoundFrontend/Public/MetasoundFrontendTransform.h
@@ -15,7 +15,7 @@
 
 // Forward Declarations
 struct FMetaSoundFrontendDocumentBuilder;
-struct FMetaSoundFrontendPresetDocumentConfiguration;
+struct FMetaSoundFrontendPresetTemplate;
 
 
 namespace Metasound::Frontend
@@ -197,7 +197,7 @@ namespace Metasound::Frontend
 		TArray<FMetasoundFrontendClassOutput> GenerateRequiredClassOutputs(FMetaSoundFrontendDocumentBuilder& InDocumentToTransformBuilder) const;
 		
 		// Add inputs to parent graph and connect to wrapped graph node.
-		void AddAndConnectInputs(const FMetaSoundFrontendPresetDocumentConfiguration& PresetConfig, const TArray<FMetasoundFrontendClassInput>& InClassInputs, FMetaSoundFrontendDocumentBuilder& InOutBuilderToTransform, const FGuid& InReferencedNodeID) const;
+		void AddAndConnectInputs(const FMetaSoundFrontendPresetTemplate& PresetTemplate, const TArray<FMetasoundFrontendClassInput>& InClassInputs, FMetaSoundFrontendDocumentBuilder& InOutBuilderToTransform, const FGuid& InReferencedNodeID) const;
 
 		// Add outputs to parent graph and connect to wrapped graph node.
 		void AddAndConnectOutputs(const TArray<FMetasoundFrontendClassOutput>& InClassOutputs, FMetaSoundFrontendDocumentBuilder& InOutBuilderToTransform, const FGuid& InReferencedNodeID) const;
@@ -205,7 +205,7 @@ namespace Metasound::Frontend
 #if WITH_EDITORONLY_DATA
 		using FMemberIDToMetadataMap = TMap<FGuid, TObjectPtr<UMetaSoundFrontendMemberMetadata>>;
 		void AddMemberMetadata(const FMemberIDToMetadataMap& InCachedMemberMetadata, FMetaSoundFrontendDocumentBuilder& InOutBuilderToTransform) const;
-		FMemberIDToMetadataMap CacheMemberMetadata(const FMetaSoundFrontendPresetDocumentConfiguration& PresetConfig, FMetaSoundFrontendDocumentBuilder& InOutBuilderToTransform) const;
+		FMemberIDToMetadataMap CacheMemberMetadata(const FMetaSoundFrontendPresetTemplate& PresetTemplate, FMetaSoundFrontendDocumentBuilder& InOutBuilderToTransform) const;
 #endif // WITH_EDITORONLY_DATA
 
 		const FMetaSoundFrontendDocumentBuilder* ParentBuilder = nullptr;
