Stopgap undeprecation until callsites are fixed up
#rnx

diff --git a/Engine/Plugins/Runtime/Metasound/Source/MetasoundFrontend/Public/MetasoundFrontendDocumentBuilder.h b/Engine/Plugins/Runtime/Metasound/Source/MetasoundFrontend/Public/MetasoundFrontendDocumentBuilder.h
--- a/Engine/Plugins/Runtime/Metasound/Source/MetasoundFrontend/Public/MetasoundFrontendDocumentBuilder.h
+++ b/Engine/Plugins/Runtime/Metasound/Source/MetasoundFrontend/Public/MetasoundFrontendDocumentBuilder.h
@@ -441,7 +441,8 @@ PRAGMA_ENABLE_EXPERIMENTAL_WARNINGS
 	// 'SetTemplate' instead for most script or edit-time applications).
 	UE_INTERNAL UE_API TStructView<FMetaSoundFrontendDocumentTemplate> GetDocumentTemplate() const;
 
-	UE_DEPRECATED(5.8, "Use GetDocumentTemplate() instead")
+	// Soft deprecated. To be deprecated in subsequent change.
+// 	UE_DEPRECATED(5.8, "Use GetDocumentTemplate() instead")
 	UE_INTERNAL TStructView<FMetaSoundFrontendDocumentTemplate> GetDocumentConfiguration() const { return GetDocumentTemplate(); }
 #endif // WITH_EDITOR
 
