UE_API missing mark-up and remove UE_INTERNAL as stopgap
#rnx

diff --git a/Engine/Plugins/Runtime/Metasound/Source/MetasoundFrontend/Public/MetasoundFrontendDocumentBuilder.h b/Engine/Plugins/Runtime/Metasound/Source/MetasoundFrontend/Public/MetasoundFrontendDocumentBuilder.h
--- a/Engine/Plugins/Runtime/Metasound/Source/MetasoundFrontend/Public/MetasoundFrontendDocumentBuilder.h
+++ b/Engine/Plugins/Runtime/Metasound/Source/MetasoundFrontend/Public/MetasoundFrontendDocumentBuilder.h
@@ -442,8 +442,7 @@ PRAGMA_ENABLE_EXPERIMENTAL_WARNINGS
 	UE_INTERNAL UE_API TStructView<FMetaSoundFrontendDocumentTemplate> GetDocumentTemplate() const;
 
 	// Soft deprecated. To be deprecated in subsequent change.
-// 	UE_DEPRECATED(5.8, "Use GetDocumentTemplate() instead")
-	UE_INTERNAL TStructView<FMetaSoundFrontendDocumentTemplate> GetDocumentConfiguration() const { return GetDocumentTemplate(); }
+	UE_API TStructView<FMetaSoundFrontendDocumentTemplate> GetDocumentConfiguration() const { return GetDocumentTemplate(); }
 #endif // WITH_EDITOR
 
 	// Returns struct view of document template if set.
