One more time, con gusto
#rnx

diff --git a/Engine/Plugins/Runtime/Metasound/Source/MetasoundFrontend/Private/MetasoundFrontendDocumentBuilder.cpp b/Engine/Plugins/Runtime/Metasound/Source/MetasoundFrontend/Private/MetasoundFrontendDocumentBuilder.cpp
--- a/Engine/Plugins/Runtime/Metasound/Source/MetasoundFrontend/Private/MetasoundFrontendDocumentBuilder.cpp
+++ b/Engine/Plugins/Runtime/Metasound/Source/MetasoundFrontend/Private/MetasoundFrontendDocumentBuilder.cpp
@@ -810,6 +810,11 @@ TStructView<FMetaSoundFrontendDocumentTemplate> FMetaSoundFrontendDocumentBuilde
 	TInstancedStruct<FMetaSoundFrontendDocumentTemplate>& Config = GetDocumentChecked().Template;
 	return Config;
 }
+
+TStructView<FMetaSoundFrontendDocumentTemplate> FMetaSoundFrontendDocumentBuilder::GetDocumentConfiguration() const
+{
+	return GetDocumentTemplate();
+}
 #endif // WITH_EDITOR
 
 TConstStructView<FMetaSoundFrontendDocumentTemplate> FMetaSoundFrontendDocumentBuilder::GetConstDocumentTemplate() const
diff --git a/Engine/Plugins/Runtime/Metasound/Source/MetasoundFrontend/Public/MetasoundFrontendDocumentBuilder.h b/Engine/Plugins/Runtime/Metasound/Source/MetasoundFrontend/Public/MetasoundFrontendDocumentBuilder.h
--- a/Engine/Plugins/Runtime/Metasound/Source/MetasoundFrontend/Public/MetasoundFrontendDocumentBuilder.h
+++ b/Engine/Plugins/Runtime/Metasound/Source/MetasoundFrontend/Public/MetasoundFrontendDocumentBuilder.h
@@ -439,10 +439,10 @@ PRAGMA_ENABLE_EXPERIMENTAL_WARNINGS
 	// Returns mutable struct view of document template if set. Internal use only to avoid
 	// writing against struct members in invalid or dangerous, non-editor contexts (use
 	// 'SetTemplate' instead for most script or edit-time applications).
-	UE_INTERNAL UE_API TStructView<FMetaSoundFrontendDocumentTemplate> GetDocumentTemplate() const;
+	UE_API TStructView<FMetaSoundFrontendDocumentTemplate> GetDocumentTemplate() const;
 
 	// Soft deprecated. To be deprecated in subsequent change.
-	UE_API TStructView<FMetaSoundFrontendDocumentTemplate> GetDocumentConfiguration() const { return GetDocumentTemplate(); }
+	UE_API TStructView<FMetaSoundFrontendDocumentTemplate> GetDocumentConfiguration() const;
 #endif // WITH_EDITOR
 
 	// Returns struct view of document template if set.
