[AutoRTFM] Use ONCOMMIT to avoid a transactional hazard with TSpscQueue.

#Synced-CL 51778913

diff --git a/Engine/Plugins/Runtime/Metasound/Source/MetasoundGraphCore/Private/MetasoundDynamicOperatorTransactor.cpp b/Engine/Plugins/Runtime/Metasound/Source/MetasoundGraphCore/Private/MetasoundDynamicOperatorTransactor.cpp
--- a/Engine/Plugins/Runtime/Metasound/Source/MetasoundGraphCore/Private/MetasoundDynamicOperatorTransactor.cpp
+++ b/Engine/Plugins/Runtime/Metasound/Source/MetasoundGraphCore/Private/MetasoundDynamicOperatorTransactor.cpp
@@ -3,6 +3,7 @@
 #include "MetasoundDynamicOperatorTransactor.h"
 
 #include "Algo/Transform.h"
+#include "AutoRTFM.h"
 #include "Containers/SpscQueue.h"
 #include "Containers/UnrealString.h"
 #include "MetasoundAudioBuffer.h"
@@ -1417,7 +1418,11 @@ namespace Metasound
 					TUniquePtr<IDynamicOperatorTransform> Transform = InFunc(*OperatorInfoIterator);
 					if (Transform.IsValid())
 					{
-						OperatorQueue->Enqueue(MoveTemp(Transform));
+						IDynamicOperatorTransform* InnerTransform = Transform.Release();
+						UE_AUTORTFM_ONCOMMIT(OperatorQueue, InnerTransform)
+						{
+							OperatorQueue->Enqueue(TUniquePtr<IDynamicOperatorTransform>(InnerTransform));
+						};
 					}
 				}
 				else
