[MetasoundEditor] Node style refinements

- Asset nodes now use the icon and color of their asset class
- Native nodes now have a default color and icon that can be overrided per node in the node's class

#rnx
#Synced-CL 51794084

diff --git a/Engine/Plugins/Runtime/Metasound/Source/MetasoundEditor/Private/FindInMetasound.cpp b/Engine/Plugins/Runtime/Metasound/Source/MetasoundEditor/Private/FindInMetasound.cpp
--- a/Engine/Plugins/Runtime/Metasound/Source/MetasoundEditor/Private/FindInMetasound.cpp
+++ b/Engine/Plugins/Runtime/Metasound/Source/MetasoundEditor/Private/FindInMetasound.cpp
@@ -298,7 +298,7 @@ namespace Metasound::Editor
 				}
 				else
 				{
-					Brush = MetaSoundGraphNode->GetNodeTitleIcon().GetIcon();
+					Brush = MetaSoundGraphNode->GetNodeTitleIconBrush();
 					IconColor = GraphNode->GetNodeTitleColor();
 				}
 			}
@@ -309,13 +309,12 @@ namespace Metasound::Editor
 		}
 
 		static const FVector2D IconSize16 = FVector2D(16.0f, 16.0f);
-		TOptional<FVector2D> BrushSize = bIsPin ? IconSize16 : TOptional<FVector2D>();
 
 		return SNew(SImage)
 			.Image(Brush)
 			.ColorAndOpacity(IconColor)
 			.ToolTipText(GetCategory())
-			.DesiredSizeOverride(BrushSize);
+			.DesiredSizeOverride(IconSize16);
 	}
 
 	FString FFindInMetasoundResult::GetCommentText() const
diff --git a/Engine/Plugins/Runtime/Metasound/Source/MetasoundEditor/Private/MetasoundEditorGraphInputNode.cpp b/Engine/Plugins/Runtime/Metasound/Source/MetasoundEditor/Private/MetasoundEditorGraphInputNode.cpp
--- a/Engine/Plugins/Runtime/Metasound/Source/MetasoundEditor/Private/MetasoundEditorGraphInputNode.cpp
+++ b/Engine/Plugins/Runtime/Metasound/Source/MetasoundEditor/Private/MetasoundEditorGraphInputNode.cpp
@@ -129,7 +129,7 @@ FGuid UMetasoundEditorGraphInputNode::GetNodeID() const
 	return NodeID;
 }
 
-FLinearColor UMetasoundEditorGraphInputNode::GetNodeTitleColor() const
+FLinearColor UMetasoundEditorGraphInputNode::GetNodeTitleTextColor() const
 {
 	if (const UMetasoundEditorSettings* EditorSettings = GetDefault<UMetasoundEditorSettings>())
 	{
@@ -139,10 +139,10 @@ FLinearColor UMetasoundEditorGraphInputNode::GetNodeTitleColor() const
 	return Super::GetNodeTitleColor();
 }
 
-FSlateIcon UMetasoundEditorGraphInputNode::GetNodeTitleIcon() const
+const FSlateBrush* UMetasoundEditorGraphInputNode::GetNodeTitleIconBrush() const
 {
-	static const FName NativeIconName = "MetasoundEditor.Graph.Node.Class.Input";
-	return FSlateIcon("MetaSoundStyle", NativeIconName);
+	const FSlateIcon InputIcon = FSlateIcon("MetaSoundStyle", "MetasoundEditor.Graph.Node.Class.Input");
+	return InputIcon.GetIcon();
 }
 
 void UMetasoundEditorGraphInputNode::GetPinHoverText(const UEdGraphPin& Pin, FString& OutHoverText) const
diff --git a/Engine/Plugins/Runtime/Metasound/Source/MetasoundEditor/Private/MetasoundEditorGraphNode.cpp b/Engine/Plugins/Runtime/Metasound/Source/MetasoundEditor/Private/MetasoundEditorGraphNode.cpp
--- a/Engine/Plugins/Runtime/Metasound/Source/MetasoundEditor/Private/MetasoundEditorGraphNode.cpp
+++ b/Engine/Plugins/Runtime/Metasound/Source/MetasoundEditor/Private/MetasoundEditorGraphNode.cpp
@@ -2,6 +2,8 @@
 
 #include "MetasoundEditorGraphNode.h"
 
+#include "AssetDefinitionRegistry.h"
+#include "AssetToolsModule.h"
 #include "EdGraph/EdGraphPin.h"
 #include "Editor/EditorEngine.h"
 #include "Engine/Font.h"
@@ -280,6 +282,42 @@ UObject& UMetasoundEditorGraphNode::GetMetasoundChecked() const
 	return EdGraph->GetMetasoundChecked();
 }
 
+const FMetasoundFrontendClass* UMetasoundEditorGraphNode::GetNodeFrontendClass() const
+{
+	using namespace Metasound::Frontend;
+
+	UObject& MetaSound = GetMetasoundChecked();
+	const FMetaSoundFrontendDocumentBuilder& DocBuilder = IDocumentBuilderRegistry::GetChecked().FindOrBeginBuilding(&MetaSound);
+
+	const FMetasoundFrontendNode* Node = DocBuilder.FindNode(GetNodeID());
+	if (!Node)
+	{
+		return nullptr;
+	}
+
+	return DocBuilder.FindDependency(Node->ClassID);
+}
+
+const UClass* UMetasoundEditorGraphNode::GetNodeAssetClass() const
+{
+	using namespace Metasound::Frontend;
+
+	const FMetasoundFrontendClass* NodeFrontendClass = GetNodeFrontendClass();
+	if (!NodeFrontendClass || !IMetaSoundAssetManager::GetChecked().IsAssetClass(NodeFrontendClass->Metadata))
+	{
+		return nullptr;
+	}
+
+	FMetasoundAssetBase* Asset = IMetaSoundAssetManager::GetChecked().FindAsset(NodeFrontendClass->Metadata);
+	if (!Asset)
+	{
+		return nullptr;
+	}
+
+	const IMetaSoundDocumentInterface* DocumentInterface = Cast<IMetaSoundDocumentInterface>(Asset->GetOwningAsset());
+	return DocumentInterface ? &DocumentInterface->GetBaseMetaSoundUClass() : nullptr;
+}
+
 Metasound::Frontend::FConstGraphHandle UMetasoundEditorGraphNode::GetConstRootGraphHandle() const
 {
 	const FMetasoundAssetBase* ConstMetasoundAsset = Metasound::IMetasoundUObjectRegistry::Get().GetObjectAsAssetBase(&GetMetasoundChecked());
@@ -773,6 +811,32 @@ void UMetasoundEditorGraphNode::PostDuplicate(bool bDuplicateForPIE)
 	}
 }
 
+const FSlateBrush* UMetasoundEditorGraphNode::GetNodeTitleIconBrush() const
+{
+	using namespace Metasound::Frontend;
+
+	const FSlateIcon DefaultIcon = FSlateIcon("MetaSoundStyle", "MetasoundEditor.Metasound.Node.Icon");
+
+	const FMetasoundFrontendClass* NodeFrontendClass = GetNodeFrontendClass();
+	if (!NodeFrontendClass)
+	{
+		return DefaultIcon.GetIcon();
+	}
+
+	const FMetaSoundAssetKey AssetKey(NodeFrontendClass->Metadata);
+	if (const TScriptInterface<IMetaSoundDocumentInterface> MetaSoundDoc = IMetaSoundAssetManager::GetChecked().FindAssetAsDocumentInterface(AssetKey))
+	{
+		const UClass& Class = MetaSoundDoc->GetBaseMetaSoundUClass();
+		if (const UAssetDefinition* AssetDefinition = UAssetDefinitionRegistry::Get()->GetAssetDefinitionForClass(&Class))
+		{
+			const FAssetData AssetData(MetaSoundDoc.GetObject());
+			return AssetDefinition->GetIconBrush(AssetData, AssetData.AssetClassPath.GetAssetName());
+		}
+	}
+
+	return DefaultIcon.GetIcon();
+}
+
 void UMetasoundEditorGraphNode::GetNodeContextMenuActions(UToolMenu* Menu, UGraphNodeContextMenuContext* Context) const
 {
 	using namespace Metasound::Editor;
@@ -893,6 +957,16 @@ FString UMetasoundEditorGraphMemberNode::GetFindReferenceSearchString_Impl(EGetF
 	return FString();
 }
 
+FLinearColor UMetasoundEditorGraphMemberNode::GetNodeTitleColor() const
+{
+	return FLinearColor::Transparent;
+}
+
+FLinearColor UMetasoundEditorGraphMemberNode::GetNodeTitleIconColor() const
+{
+	return GetNodeTitleTextColor();
+}
+
 void UMetasoundEditorGraphOutputNode::PinDefaultValueChanged(UEdGraphPin* InPin)
 {
 	using namespace Metasound::Editor;
@@ -1111,24 +1185,25 @@ void UMetasoundEditorGraphOutputNode::SetNodeID(FGuid InNodeID)
 	}
 }
 
-FLinearColor UMetasoundEditorGraphOutputNode::GetNodeTitleColor() const
+FLinearColor UMetasoundEditorGraphOutputNode::GetNodeTitleTextColor() const
 {
 	if (const UMetasoundEditorSettings* EditorSettings = GetDefault<UMetasoundEditorSettings>())
 	{
 		return EditorSettings->OutputNodeTitleColor;
 	}
 
-	return Super::GetNodeTitleColor();
+	return Super::GetNodeTitleTextColor();
 }
 
-UMetasoundEditorGraphMember* UMetasoundEditorGraphOutputNode::GetMember() const
+const FSlateBrush* UMetasoundEditorGraphOutputNode::GetNodeTitleIconBrush() const
 {
-	return Output;
+	FSlateIcon OutputIcon = FSlateIcon("MetaSoundStyle", "MetasoundEditor.Graph.Node.Class.Output");
+	return OutputIcon.GetIcon();
 }
 
-FSlateIcon UMetasoundEditorGraphOutputNode::GetNodeTitleIcon() const
+UMetasoundEditorGraphMember* UMetasoundEditorGraphOutputNode::GetMember() const
 {
-	return FSlateIcon("MetaSoundStyle", "MetasoundEditor.Graph.Node.Class.Output");
+	return Output;
 }
 
 void UMetasoundEditorGraphExternalNode::ReconstructNode()
@@ -1574,27 +1649,74 @@ const FMetasoundEditorGraphNodeBreadcrumb& UMetasoundEditorGraphExternalNode::Ge
 
 FLinearColor UMetasoundEditorGraphExternalNode::GetNodeTitleColor() const
 {
-	if (const UMetasoundEditorSettings* EditorSettings = GetDefault<UMetasoundEditorSettings>())
+	using namespace Metasound::Frontend;
+
+	// Return the asset type color if the node is an asset type (e.g. Source, Patch)
+	if (const UClass* NodeAssetClass = GetNodeAssetClass())
 	{
-		if (bIsClassNative)
+		if (const UAssetDefinition* AssetDefinition = UAssetDefinitionRegistry::Get()->GetAssetDefinitionForClass(NodeAssetClass))
 		{
-			return EditorSettings->NativeNodeTitleColor;
+			return AssetDefinition->GetAssetColor();
+		}
+	}
+
+	// Use the style color of native nodes, if specified in the node class.
+	UObject& MetaSound = GetMetasoundChecked();
+	const FMetaSoundFrontendDocumentBuilder& DocBuilder = IDocumentBuilderRegistry::GetChecked().FindOrBeginBuilding(&MetaSound);
+	if (const FMetasoundFrontendNode* Node = DocBuilder.FindNode(GetNodeID()))
+	{
+		const FMetasoundFrontendClass* FrontendClass = DocBuilder.FindDependency(Node->ClassID);
+		const FMetasoundFrontendClassStyleDisplay& StyleDisplay = FrontendClass->Style.Display;
+		if (!StyleDisplay.StyleSet.IsNone() && !StyleDisplay.Color.IsNone())
+		{
+			if (const ISlateStyle* Style = FSlateStyleRegistry::FindSlateStyle(StyleDisplay.StyleSet))
+			{
+				return Style->GetSlateColor(StyleDisplay.Color, nullptr, FLinearColor::Transparent).GetSpecifiedColor();
+			}
 		}
+	}
 
-		return EditorSettings->AssetReferenceNodeTitleColor;
+	// Fallback color for outdated nodes, etc.
+	if (const UMetasoundEditorSettings* EditorSettings = GetDefault<UMetasoundEditorSettings>())
+	{
+		return EditorSettings->NativeNodeTitleColor;
 	}
 
-	return Super::GetNodeTitleColor();
+	return FLinearColor::Transparent;
 }
 
-FSlateIcon UMetasoundEditorGraphExternalNode::GetNodeTitleIcon() const
+const FSlateBrush* UMetasoundEditorGraphExternalNode::GetNodeTitleIconBrush() const
 {
+	if (const FMetasoundFrontendClass* FrontendClass = GetFrontendClass())
+	{
+		const FMetasoundFrontendClassStyleDisplay StyleDisplay = FrontendClass->Style.Display;
+		if (!StyleDisplay.StyleSet.IsNone() && !StyleDisplay.Icon.IsNone())
+		{
+			return FSlateIcon(StyleDisplay.StyleSet, StyleDisplay.Icon).GetIcon();
+		}
+	}
+
+	return Super::GetNodeTitleIconBrush();
+}
+
+FText UMetasoundEditorGraphExternalNode::GetNodeTitleIconTooltipText() const
+{
+	FText TooltipText;
 	if (!bIsClassNative)
 	{
-		return FSlateIcon("MetasoundStyle", "MetasoundEditor.Graph.Node.Class.Graph");
+		if (const FMetasoundFrontendClass* FrontendClass = GetFrontendClass())
+		{
+			const FMetasoundFrontendClassMetadata& Metadata = FrontendClass->Metadata;
+			const FString MetadataAuthor = Metadata.GetAuthor();
+
+			if (!MetadataAuthor.IsEmpty())
+			{
+				TooltipText = FText::Format(LOCTEXT("MetaSoundNodeAuthorFormat", "Author: {0}"), FText::FromString(MetadataAuthor));
+			}
+		}
 	}
 
-	return FSlateIcon();
+	return TooltipText;
 }
 
 bool UMetasoundEditorGraphExternalNode::ShouldDrawNodeAsControlPointOnly(int32& OutInputPinIndex, int32& OutOutputPinIndex) const
@@ -1788,9 +1910,9 @@ FLinearColor UMetasoundEditorGraphVariableNode::GetNodeTitleColor() const
 	return Super::GetNodeTitleColor();
 }
 
-FSlateIcon UMetasoundEditorGraphVariableNode::GetNodeTitleIcon() const
+const FSlateBrush* UMetasoundEditorGraphVariableNode::GetNodeTitleIconBrush() const
 {
-	return FSlateIcon();
+	return FSlateIcon().GetIcon();
 }
 
 void UMetasoundEditorGraphVariableNode::SetNodeID(FGuid InNodeID)
diff --git a/Engine/Plugins/Runtime/Metasound/Source/MetasoundEditor/Private/MetasoundEditorGraphSchema.cpp b/Engine/Plugins/Runtime/Metasound/Source/MetasoundEditor/Private/MetasoundEditorGraphSchema.cpp
--- a/Engine/Plugins/Runtime/Metasound/Source/MetasoundEditor/Private/MetasoundEditorGraphSchema.cpp
+++ b/Engine/Plugins/Runtime/Metasound/Source/MetasoundEditor/Private/MetasoundEditorGraphSchema.cpp
@@ -1,6 +1,7 @@
 // Copyright Epic Games, Inc. All Rights Reserved.
 #include "MetasoundEditorGraphSchema.h"
 
+#include "AssetDefinitionRegistry.h"
 #include "Algo/AnyOf.h"
 #include "EdGraph/EdGraph.h"
 #include "EdGraphNode_Comment.h"
@@ -963,10 +964,35 @@ const FSlateBrush* FMetasoundGraphSchemaAction_NewNode::GetIconBrush() const
 {
 	using namespace Metasound::Frontend;
 
-	const bool bIsClassNative = QueryResult.IsValid() && QueryResult->IsNative();
-	if (bIsClassNative)
+	if (QueryResult.IsValid() && QueryResult->IsNative())
 	{
-		return &Metasound::Editor::Style::GetSlateBrushSafe("MetasoundEditor.Graph.Node.Class.Native");
+		const FMetaSoundAssetKey AssetKey = QueryResult->GetAssetKey();
+
+		FMetasoundFrontendClass FrontendClass;
+		INodeClassRegistry::GetChecked().FindFrontendClassFromRegistered(AssetKey, FrontendClass);
+		const FMetasoundFrontendClassStyleDisplay& StyleDisplay = FrontendClass.Style.Display;
+
+		FSlateIcon NativeIcon = FSlateIcon("MetaSoundStyle", "MetasoundEditor.Metasound.Node.Icon");
+		if (!StyleDisplay.StyleSet.IsNone() && !StyleDisplay.Icon.IsNone())
+		{
+			NativeIcon = FSlateIcon(StyleDisplay.StyleSet, StyleDisplay.Icon);
+		}
+
+		return NativeIcon.GetIcon();
+	}
+
+	if (QueryResult.IsValid() && QueryResult.Get())
+	{
+		const FMetaSoundAssetKey AssetKey = QueryResult.Get()->GetAssetKey();
+		if (const TScriptInterface<IMetaSoundDocumentInterface> MetaSoundDoc = IMetaSoundAssetManager::GetChecked().FindAssetAsDocumentInterface(AssetKey))
+		{
+			const UClass& Class = MetaSoundDoc->GetBaseMetaSoundUClass();
+			if (const UAssetDefinition* AssetDefinition = UAssetDefinitionRegistry::Get()->GetAssetDefinitionForClass(&Class))
+			{
+				const FAssetData AssetData(MetaSoundDoc.GetObject());
+				return AssetDefinition->GetIconBrush(AssetData, AssetData.AssetClassPath.GetAssetName());
+			}
+		}
 	}
 
 	return &Metasound::Editor::Style::GetSlateBrushSafe("MetasoundEditor.Graph.Node.Class.Graph");
@@ -976,14 +1002,45 @@ const FLinearColor& FMetasoundGraphSchemaAction_NewNode::GetIconColor() const
 {
 	using namespace Metasound::Frontend;
 
+	if (QueryResult.IsValid() && QueryResult->IsNative())
+	{
+		const FMetaSoundAssetKey AssetKey = QueryResult->GetAssetKey();
+
+		FMetasoundFrontendClass FrontendClass;
+		INodeClassRegistry::GetChecked().FindFrontendClassFromRegistered(AssetKey, FrontendClass);
+		const FMetasoundFrontendClassStyleDisplay& StyleDisplay = FrontendClass.Style.Display;
+		if (!StyleDisplay.StyleSet.IsNone() && !StyleDisplay.Color.IsNone())
+		{
+			if (const ISlateStyle* Style = FSlateStyleRegistry::FindSlateStyle(StyleDisplay.StyleSet))
+			{
+				static FLinearColor IconColor;
+				IconColor = Style->GetSlateColor(StyleDisplay.Color, nullptr, FLinearColor::Transparent).GetSpecifiedColor();
+				return IconColor;
+			}
+		}
+	}
+
+	if (QueryResult.IsValid() && QueryResult.Get())
+	{
+		const FMetaSoundAssetKey AssetKey = QueryResult.Get()->GetAssetKey();
+		if (const TScriptInterface<IMetaSoundDocumentInterface> MetaSoundDoc = IMetaSoundAssetManager::GetChecked().FindAssetAsDocumentInterface(AssetKey))
+		{
+			if (const ISlateStyle* MetasoundStyle = FSlateStyleRegistry::FindSlateStyle("MetaSoundStyle"))
+			{
+				const FAssetData AssetData(MetaSoundDoc.GetObject());
+				FName AssetClassName = AssetData.AssetClassPath.GetAssetName();
+				FName ColorName(*(AssetClassName.ToString() + TEXT(".Color")));
+				return MetasoundStyle->GetColor(ColorName);
+			}
+		}
+	}
+
 	if (const UMetasoundEditorSettings* EditorSettings = GetDefault<UMetasoundEditorSettings>())
 	{
-		if (QueryResult.IsValid() && QueryResult->IsNative())
+		if (QueryResult.IsValid())
 		{
 			return EditorSettings->NativeNodeTitleColor;
 		}
-
-		return EditorSettings->AssetReferenceNodeTitleColor;
 	}
 
 	return Super::GetIconColor();
diff --git a/Engine/Plugins/Runtime/Metasound/Source/MetasoundEditor/Private/MetasoundEditorModule.cpp b/Engine/Plugins/Runtime/Metasound/Source/MetasoundEditor/Private/MetasoundEditorModule.cpp
--- a/Engine/Plugins/Runtime/Metasound/Source/MetasoundEditor/Private/MetasoundEditorModule.cpp
+++ b/Engine/Plugins/Runtime/Metasound/Source/MetasoundEditor/Private/MetasoundEditorModule.cpp
@@ -496,6 +496,7 @@ namespace Metasound::Editor
 				SetClassIcon(TEXT("MetasoundPatch"));
 				SetClassIcon(TEXT("MetasoundSource"));
 
+				Set("MetasoundEditor.Metasound.Node.Icon", new IMAGE_BRUSH_SVG(TEXT("Icons/metasound"), Icon20x20));
 				Set("MetasoundEditor.MetasoundPatch.Icon", new IMAGE_BRUSH_SVG(TEXT("Icons/metasoundpatch_icon"), Icon20x20));
 				Set("MetasoundEditor.MetasoundPatch.Preset.Icon", new IMAGE_BRUSH_SVG(TEXT("Icons/metasoundpatchpreset_icon"), Icon20x20));
 				Set("MetasoundEditor.MetasoundSource.Icon", new IMAGE_BRUSH_SVG(TEXT("Icons/metasoundsource_icon"), Icon20x20));
diff --git a/Engine/Plugins/Runtime/Metasound/Source/MetasoundEditor/Private/MetasoundEditorSettings.cpp b/Engine/Plugins/Runtime/Metasound/Source/MetasoundEditor/Private/MetasoundEditorSettings.cpp
--- a/Engine/Plugins/Runtime/Metasound/Source/MetasoundEditor/Private/MetasoundEditorSettings.cpp
+++ b/Engine/Plugins/Runtime/Metasound/Source/MetasoundEditor/Private/MetasoundEditorSettings.cpp
@@ -35,7 +35,6 @@ UMetasoundEditorSettings::UMetasoundEditorSettings(const FObjectInitializer& Obj
 	WaveTablePinTypeColor = FLinearColor(0.580392f, 0.0f, 0.827450f, 1.0f);		// purple
 
 	NativeNodeTitleColor = FLinearColor(0.4f, 0.85f, 0.35f, 1.0f);				// pale green
-	AssetReferenceNodeTitleColor = FLinearColor(0.047f, 0.686f, 0.988f);		// sky blue
 	InputNodeTitleColor = FLinearColor(0.168f, 1.0f, 0.7294f);					// sea foam
 	OutputNodeTitleColor = FLinearColor(1.0f, 0.878f, 0.1686f);					// yellow
 	VariableNodeTitleColor = FLinearColor(0.211f, 0.513f, 0.035f);				// copper
diff --git a/Engine/Plugins/Runtime/Metasound/Source/MetasoundEditor/Private/MetasoundEditorSettings.h b/Engine/Plugins/Runtime/Metasound/Source/MetasoundEditor/Private/MetasoundEditorSettings.h
--- a/Engine/Plugins/Runtime/Metasound/Source/MetasoundEditor/Private/MetasoundEditorSettings.h
+++ b/Engine/Plugins/Runtime/Metasound/Source/MetasoundEditor/Private/MetasoundEditorSettings.h
@@ -245,7 +245,7 @@ public:
 	FLinearColor NativeNodeTitleColor;
 
 	/** Title color for references to MetaSound assets */
-	UPROPERTY(EditAnywhere, config, Category = NodeTitleColors)
+	UE_DEPRECATED(5.8, "Asset nodes now use their corresponding asset type color.")
 	FLinearColor AssetReferenceNodeTitleColor;
 
 	/** Input node title color */
diff --git a/Engine/Plugins/Runtime/Metasound/Source/MetasoundEditor/Private/SMetasoundActionMenu.cpp b/Engine/Plugins/Runtime/Metasound/Source/MetasoundEditor/Private/SMetasoundActionMenu.cpp
--- a/Engine/Plugins/Runtime/Metasound/Source/MetasoundEditor/Private/SMetasoundActionMenu.cpp
+++ b/Engine/Plugins/Runtime/Metasound/Source/MetasoundEditor/Private/SMetasoundActionMenu.cpp
@@ -350,6 +350,7 @@ namespace Metasound
 						SNew(SImage)
 						.ColorAndOpacity(IconColor)
 						.Image(IconBrush)
+						.DesiredSizeOverride(FVector2D(16.f, 16.f))
 					];
 			}
 
diff --git a/Engine/Plugins/Runtime/Metasound/Source/MetasoundEditor/Private/SMetasoundGraphNode.cpp b/Engine/Plugins/Runtime/Metasound/Source/MetasoundEditor/Private/SMetasoundGraphNode.cpp
--- a/Engine/Plugins/Runtime/Metasound/Source/MetasoundEditor/Private/SMetasoundGraphNode.cpp
+++ b/Engine/Plugins/Runtime/Metasound/Source/MetasoundEditor/Private/SMetasoundGraphNode.cpp
@@ -460,39 +460,14 @@ namespace Metasound
 				return SNullWidget::NullWidget;
 			}
 
-			TSharedPtr<SHorizontalBox> TitleBoxWidget = SNew(SHorizontalBox);
-
-			const FSlateIcon NodeIcon = Node.GetNodeTitleIcon();
-			if (const FSlateBrush* IconBrush = NodeIcon.GetIcon())
-			{
-				if (IconBrush != FStyleDefaults::GetNoBrush() && Node.GetBreadcrumb().bIsClassNative)
-				{
-					TSharedPtr<SImage> Image;
-					TitleBoxWidget->AddSlot()
-					.AutoWidth()
-					.VAlign(VAlign_Center)
-					.HAlign(HAlign_Right)
-					[
-						SNew(SHorizontalBox)
-						+ SHorizontalBox::Slot()
-						.Padding(0.0f, 0.0f, 4.0f, 0.0f)
-						.VAlign(VAlign_Center)
-						.HAlign(HAlign_Right)
-						[
-							SAssignNew(Image, SImage)
-						]
-					];
-					Image->SetColorAndOpacity(TAttribute<FSlateColor>::CreateLambda([this]() { return FSlateColor(GetNodeTitleColorOverride()); }));
-					Image->SetImage(IconBrush);
-				}
-			}
-
-			TitleBoxWidget->AddSlot()
-			.AutoWidth()
-			.VAlign(VAlign_Center)
-			[
-				SGraphNode::CreateTitleWidget(NodeTitle)
-			];
+			const TSharedPtr<SHorizontalBox> TitleBoxWidget =
+				SNew(SHorizontalBox)
+				+ SHorizontalBox::Slot()
+				.AutoWidth()
+				.VAlign(VAlign_Center)
+				[
+					SGraphNode::CreateTitleWidget(NodeTitle)
+				];
 
 			InlineEditableText->SetColorAndOpacity(TAttribute<FLinearColor>::Create(TAttribute<FLinearColor>::FGetter::CreateSP(this, &SMetaSoundGraphNode::GetNodeTitleColorOverride)));
 
@@ -573,41 +548,10 @@ namespace Metasound
 		FLinearColor SMetaSoundGraphNode::GetNodeTitleColorOverride() const
 		{
 			const UMetasoundEditorGraphNode& Node = GetMetaSoundNode();
-
-			FMetasoundFrontendClassStyleDisplay StyleDisplay;
-			if (const FMetasoundFrontendClass* FrontendClass = Node.GetFrontendClass())
-			{
-				StyleDisplay = FrontendClass->Style.Display;
-			}
-
-			FLinearColor TitleColor = GraphNode->IsDeprecated() ? FLinearColor::Red : GetNodeObj()->GetNodeTitleColor();
-			const bool bDimTitle = !GraphNode->IsNodeEnabled() || GraphNode->IsDisplayAsDisabledForced() || GraphNode->IsNodeUnrelated();
-
-			if (bDimTitle)
-			{
-				TitleColor *= FLinearColor(0.5f, 0.5f, 0.5f, 0.4f);
-				return TitleColor;
-			}
-
+			FLinearColor TitleColor = GraphNode->IsDeprecated() ? FLinearColor::Red : Node.GetNodeTitleTextColor();
 			if (ClassType != EMetasoundFrontendClassType::External)
 			{
 				TitleColor.A = FadeCurve.GetLerp();
-				return TitleColor;
-			}
-
-			if (!Node.GetBreadcrumb().bIsClassNative)
-			{
-				TitleColor = FMath::Lerp(TitleColor, FLinearColor::White, 0.5f);
-				return TitleColor;
-			}
-
-			if (FSlateStyleRegistry::FindSlateStyle(StyleDisplay.StyleSet))
-			{
-				// Use white text if we're applying a style color to the node title's background.
-				if (!StyleDisplay.Color.IsNone())
-				{
-					return FLinearColor::White;
-				}
 			}
 
 			return TitleColor;
@@ -632,38 +576,40 @@ namespace Metasound
 			}
 
 			DefaultTitleAreaWidget->ClearChildren();
+
+			const bool bShowTitleIcon = Node.ShouldShowTitleIcon();
 			const TSharedPtr<SNodeTitle> NodeTitle = SNew(SNodeTitle, GraphNode);
 
-			const TSharedRef<SWidget> TitleContent =
+			const TSharedRef<SWidget> TitleTextBlock =
+				SNew(SVerticalBox)
+				+ SVerticalBox::Slot()
+				.AutoHeight()
+				.HAlign(bShowTitleIcon ? HAlign_Left : HAlign_Center)
+				[
+					CreateTitleWidget(NodeTitle)
+				]
+				+ SVerticalBox::Slot()
+				.AutoHeight()
+				.HAlign(bShowTitleIcon ? HAlign_Left : HAlign_Center)
+				[
+					NodeTitle.ToSharedRef()
+				];
+
+			TSharedRef<SWidget> TitleContent =
 				SNew(SHorizontalBox)
 				+ SHorizontalBox::Slot()
 				.AutoWidth()
-				.HAlign(HAlign_Left)
 				.VAlign(VAlign_Center)
-				.Padding(FMargin(0.f, 0.f, 6.f, 0.f))
+				.Padding(0, 0, 6, 0)
 				[
 					CreateTitleIcon()
 				]
 				+ SHorizontalBox::Slot()
 				.FillWidth(1.f)
-				.HAlign(HAlign_Center)
+				.HAlign(HAlign_Left)
 				.VAlign(VAlign_Center)
 				[
-					SNew(SVerticalBox)
-					+ SVerticalBox::Slot()
-					.AutoHeight()
-					.HAlign(HAlign_Center)
-					.VAlign(VAlign_Center)
-					[
-						CreateTitleWidget(NodeTitle)
-					]
-					+ SVerticalBox::Slot()
-					.AutoHeight()
-					.HAlign(HAlign_Center)
-					.VAlign(VAlign_Center)
-					[
-						NodeTitle.ToSharedRef()
-					]
+					TitleTextBlock
 				]
 				+ SHorizontalBox::Slot()
 				.AutoWidth()
@@ -673,13 +619,11 @@ namespace Metasound
 					CreateTitleRightWidget()
 				];
 
-			const FLinearColor TitleBorderColor = Node.GetBreadcrumb().bIsClassNative ? GetTitleBorderColor() : GetNodeObj()->GetNodeTitleColor();
-
 			const TSharedRef<SWidget> MainTitleArea =
 				SNew(SBorder)
 				.BorderImage(FAppStyle::GetBrush("Graph.Node.ColorSpill"))
 				.Padding(FMargin(10.f, 5.f, 10.f, 3.f))
-				.BorderBackgroundColor(TitleBorderColor)
+				.BorderBackgroundColor(Node.GetNodeTitleColor())
 				.HAlign(HAlign_Fill)
 				.VAlign(VAlign_Center)
 				[
@@ -709,44 +653,19 @@ namespace Metasound
 
 		TSharedRef<SWidget> SMetaSoundGraphNode::CreateTitleIcon() const
 		{
-			if (ClassType == EMetasoundFrontendClassType::Input || ClassType == EMetasoundFrontendClassType::Output)
-			{
-				return SNullWidget::NullWidget;
-			}
-
 			const UMetasoundEditorGraphNode& Node = GetMetaSoundNode();
-
-			const bool bIsAsset = !Node.GetBreadcrumb().bIsClassNative;
-
-			FText AuthorText = FText::GetEmpty();
-			bool bShowAuthor = false;
-
-			if (bIsAsset)
-			{
-				if (const FMetasoundFrontendClass* FrontendClass = Node.GetFrontendClass())
-				{
-					const FMetasoundFrontendClassMetadata& Metadata = FrontendClass->Metadata;
-					const FString MetadataAuthor = Metadata.GetAuthor();
-
-					if (!MetadataAuthor.IsEmpty())
-					{
-						AuthorText = FText::Format(LOCTEXT("MetaSoundNodeAuthorFormat", "Author: {0}"), FText::FromString(MetadataAuthor));
-						bShowAuthor = true;
-					}
-				}
-			}
-
-			const FSlateBrush* IconBrush = bIsAsset ? Node.GetNodeTitleIcon().GetIcon() : GetIconBrush();
-
-			if (!IconBrush)
+			if (!Node.ShouldShowTitleIcon())
 			{
 				return SNullWidget::NullWidget;
 			}
 
-			TSharedPtr<SToolTip> AuthorToolTip;
-			if (bShowAuthor)
+			const FText& TooltipText = Node.GetNodeTitleIconTooltipText();
+			const bool bShowTooltip = !TooltipText.IsEmpty();
+
+			TSharedPtr<SToolTip> ToolTip;
+			if (bShowTooltip)
 			{
-				AuthorToolTip =
+				ToolTip =
 					SNew(SToolTip)
 					[
 						SNew(SVerticalBox)
@@ -755,7 +674,7 @@ namespace Metasound
 						.Padding(0, 0, 0, 2)
 						[
 							SNew(STextBlock)
-							.Text(AuthorText)
+							.Text(TooltipText)
 						]
 					];
 			}
@@ -763,7 +682,7 @@ namespace Metasound
 			return SNew(SBox)
 				.WidthOverride(20.f)
 				.HeightOverride(20.f)
-				.ToolTip(AuthorToolTip)
+				.ToolTip(ToolTip)
 				[
 					SNew(SBorder)
 					.HAlign(HAlign_Fill)
@@ -781,47 +700,13 @@ namespace Metasound
 						.VAlign(VAlign_Center)
 						[
 							SNew(SImage)
-							.Image(IconBrush)
-							.ColorAndOpacity(bIsAsset ? GetNodeObj()->GetNodeTitleColor() : FLinearColor::White)
+							.Image(Node.GetNodeTitleIconBrush())
+							.ColorAndOpacity(FLinearColor::White)
 						]
 					]
 				];
 		}
 
-		const FSlateBrush* SMetaSoundGraphNode::GetIconBrush() const
-		{
-			const UMetasoundEditorGraphNode& Node = GetMetaSoundNode();
-
-			FMetasoundFrontendClassStyleDisplay StyleDisplay;
-			if (const FMetasoundFrontendClass* FrontendClass = Node.GetFrontendClass())
-			{
-				StyleDisplay = FrontendClass->Style.Display;
-			}
-
-			const ISlateStyle* Style = FSlateStyleRegistry::FindSlateStyle(StyleDisplay.StyleSet);
-			if (!Style || StyleDisplay.Icon.IsNone())
-			{
-				return nullptr;
-			}
-
-			return Style->GetBrush(StyleDisplay.Icon);
-		}
-
-		FLinearColor SMetaSoundGraphNode::GetTitleBorderColor() const
-		{
-			const UMetasoundEditorGraphNode& Node = GetMetaSoundNode();
-			if (const FMetasoundFrontendClass* FrontendClass = Node.GetFrontendClass())
-			{
-				const FMetasoundFrontendClassStyleDisplay& StyleDisplay = FrontendClass->Style.Display;
-				if (const ISlateStyle* Style = FSlateStyleRegistry::FindSlateStyle(StyleDisplay.StyleSet))
-				{
-					return Style->GetSlateColor(StyleDisplay.Color, nullptr, FLinearColor::Transparent).GetSpecifiedColor();
-				}
-			}
-
-			return FLinearColor::Transparent;
-		}
-
 		void SMetaSoundGraphNode::MoveTo(const FVector2f& NewPosition, FNodeSet& NodeFilter, bool bMarkDirty)
 		{
 			SGraphNode::MoveTo(NewPosition, NodeFilter, bMarkDirty);
diff --git a/Engine/Plugins/Runtime/Metasound/Source/MetasoundEditor/Private/SMetasoundGraphNode.h b/Engine/Plugins/Runtime/Metasound/Source/MetasoundEditor/Private/SMetasoundGraphNode.h
--- a/Engine/Plugins/Runtime/Metasound/Source/MetasoundEditor/Private/SMetasoundGraphNode.h
+++ b/Engine/Plugins/Runtime/Metasound/Source/MetasoundEditor/Private/SMetasoundGraphNode.h
@@ -76,8 +76,6 @@ namespace Metasound::Editor
 		void FinalizeValueTransaction(TWeakObjectPtr<UMetasoundEditorGraphMember> GraphMemberPtr, TFunctionRef<void(const FGuid& /*BuildPageID*/, UMetasoundEditorGraphMember& /*GraphMember*/, bool /*bPostTransaction*/)> SetValue);
 
 		TSharedRef<SWidget> CreateTitleIcon() const;
-		const FSlateBrush* GetIconBrush() const;
-		FLinearColor GetTitleBorderColor() const;
 
 		// Returns attribute that returns whether or not the input widget should be enabled.
 		TAttribute<bool> GetInputWidgetEnabled() const;
diff --git a/Engine/Plugins/Runtime/Metasound/Source/MetasoundEditor/Public/MetasoundEditorGraphInputNode.h b/Engine/Plugins/Runtime/Metasound/Source/MetasoundEditor/Public/MetasoundEditorGraphInputNode.h
--- a/Engine/Plugins/Runtime/Metasound/Source/MetasoundEditor/Public/MetasoundEditorGraphInputNode.h
+++ b/Engine/Plugins/Runtime/Metasound/Source/MetasoundEditor/Public/MetasoundEditorGraphInputNode.h
@@ -45,10 +45,10 @@ public:
 	UE_API virtual void UpdatePreviewInstance(const Metasound::FVertexName& InParameterName, TScriptInterface<IAudioParameterControllerInterface>& InParameterInterface) const;
 	UE_API virtual FMetasoundFrontendClassName GetClassName() const;
 
-	UE_API FText GetDisplayName() const override;
+	UE_API virtual FText GetDisplayName() const override;
 	UE_API virtual FGuid GetNodeID() const override;
-	UE_API virtual FLinearColor GetNodeTitleColor() const override;
-	UE_API virtual FSlateIcon GetNodeTitleIcon() const override;
+	UE_API virtual FLinearColor GetNodeTitleTextColor() const override;
+	UE_API virtual const FSlateBrush* GetNodeTitleIconBrush() const override;
 	UE_API virtual void GetPinHoverText(const UEdGraphPin& Pin, FString& OutHoverText) const override;
 	UE_API virtual void Validate(const FMetaSoundFrontendDocumentBuilder& InBuilder, Metasound::Editor::FGraphNodeValidationResult& OutResult) const override;
 	UE_API virtual bool EnableInteractWidgets() const override;
diff --git a/Engine/Plugins/Runtime/Metasound/Source/MetasoundEditor/Public/MetasoundEditorGraphNode.h b/Engine/Plugins/Runtime/Metasound/Source/MetasoundEditor/Public/MetasoundEditorGraphNode.h
--- a/Engine/Plugins/Runtime/Metasound/Source/MetasoundEditor/Public/MetasoundEditorGraphNode.h
+++ b/Engine/Plugins/Runtime/Metasound/Source/MetasoundEditor/Public/MetasoundEditorGraphNode.h
@@ -136,7 +136,10 @@ public:
 	virtual void PostDuplicate(bool bDuplicateForPIE) override;
 	// End of UObject interface
 
-	virtual FSlateIcon GetNodeTitleIcon() const { return FSlateIcon(); }
+	virtual bool ShouldShowTitleIcon() const { return true; }
+	virtual const FSlateBrush* GetNodeTitleIconBrush() const;
+	virtual FLinearColor GetNodeTitleIconColor() const { return FLinearColor::White; }
+	virtual FText GetNodeTitleIconTooltipText() const { return FText(); }
 	virtual FName GetCornerIcon() const { return FName(); }
 	virtual bool CanAddInputPin() const { return false; }
 
@@ -163,6 +166,9 @@ public:
 	UObject* GetMetasound() const;
 	UObject& GetMetasoundChecked() const;
 
+	const FMetasoundFrontendClass* GetNodeFrontendClass() const;
+	const UClass* GetNodeAssetClass() const;
+
 	virtual bool RemoveFromDocument() const;
 
 	UE_DEPRECATED(5.4, "Use UpdateFrontendNodeLocation and/or SyncLocationFromFrontendNode")
@@ -272,6 +278,8 @@ public:
 protected:
 	virtual FString GetFindReferenceSearchString_Impl(EGetFindReferenceSearchStringFlags InFlags) const override;
 
+	virtual FLinearColor GetNodeTitleColor() const override;
+	virtual FLinearColor GetNodeTitleIconColor() const override;
 };
 
 /** Node that represents a graph output */
@@ -312,8 +320,8 @@ protected:
 	UPROPERTY()
 	FMetasoundEditorGraphVertexNodeBreadcrumb Breadcrumb;
 
-	virtual FLinearColor GetNodeTitleColor() const override;
-	virtual FSlateIcon GetNodeTitleIcon() const override;
+	virtual FLinearColor GetNodeTitleTextColor() const override;
+	virtual const FSlateBrush* GetNodeTitleIconBrush() const override;
 	virtual void SetNodeID(FGuid InNodeID) override;
 
 	// Friended to enable mutation of Frontend NodeID & direct breadcrumb access
@@ -340,10 +348,11 @@ protected:
 	bool bIsClassNative = true;
 
 public:
-	virtual const FMetasoundEditorGraphNodeBreadcrumb& GetBreadcrumb() const;
+	virtual const FMetasoundEditorGraphNodeBreadcrumb& GetBreadcrumb() const override;
 	virtual FGuid GetNodeID() const override { return NodeID; }
 	virtual FLinearColor GetNodeTitleColor() const override;
-	virtual FSlateIcon GetNodeTitleIcon() const override;
+	virtual const FSlateBrush* GetNodeTitleIconBrush() const override;
+	virtual FText GetNodeTitleIconTooltipText() const override;
 	virtual bool ShouldDrawNodeAsControlPointOnly(int32& OutInputPinIndex, int32& OutOutputPinIndex) const override;
 
 	virtual void ReconstructNode() override;
@@ -412,7 +421,7 @@ protected:
 	FMetasoundEditorGraphMemberNodeBreadcrumb Breadcrumb;
 
 	virtual FLinearColor GetNodeTitleColor() const override;
-	virtual FSlateIcon GetNodeTitleIcon() const override;
+	virtual const FSlateBrush* GetNodeTitleIconBrush() const override;
 	virtual void SetNodeID(FGuid InNodeID) override;
 
 	// Friended to enable mutation of Frontend NodeID & direct breadcrumb access
