[MetasoundEditor] Fix for assert on FMetasoundAssetKey being constructed from incorrect metadata type

- Fixed issue with input/output icon colors in FindInMetasound

#rnx
#Synced-CL 51794084

diff --git a/Engine/Plugins/Runtime/Metasound/Source/MetasoundEditor/Private/FindInMetasound.cpp b/Engine/Plugins/Runtime/Metasound/Source/MetasoundEditor/Private/FindInMetasound.cpp
--- a/Engine/Plugins/Runtime/Metasound/Source/MetasoundEditor/Private/FindInMetasound.cpp
+++ b/Engine/Plugins/Runtime/Metasound/Source/MetasoundEditor/Private/FindInMetasound.cpp
@@ -296,6 +296,11 @@ namespace Metasound::Editor
 						IconColor = Schema->GetPinTypeColor(FirstPin->PinType);
 					}
 				}
+				else if (UMetasoundEditorGraphMemberNode* MetaSoundGraphMemberNode = Cast<UMetasoundEditorGraphMemberNode>(MetaSoundGraphNode))
+				{
+					Brush = MetaSoundGraphNode->GetNodeTitleIconBrush();
+					IconColor = GraphNode->GetNodeTitleTextColor();
+				}
 				else
 				{
 					Brush = MetaSoundGraphNode->GetNodeTitleIconBrush();
diff --git a/Engine/Plugins/Runtime/Metasound/Source/MetasoundEditor/Private/MetasoundEditorGraphNode.cpp b/Engine/Plugins/Runtime/Metasound/Source/MetasoundEditor/Private/MetasoundEditorGraphNode.cpp
--- a/Engine/Plugins/Runtime/Metasound/Source/MetasoundEditor/Private/MetasoundEditorGraphNode.cpp
+++ b/Engine/Plugins/Runtime/Metasound/Source/MetasoundEditor/Private/MetasoundEditorGraphNode.cpp
@@ -823,14 +823,17 @@ const FSlateBrush* UMetasoundEditorGraphNode::GetNodeTitleIconBrush() const
 		return DefaultIcon.GetIcon();
 	}
 
-	const FMetaSoundAssetKey AssetKey(NodeFrontendClass->Metadata);
-	if (const TScriptInterface<IMetaSoundDocumentInterface> MetaSoundDoc = IMetaSoundAssetManager::GetChecked().FindAssetAsDocumentInterface(AssetKey))
+	if (NodeFrontendClass->Metadata.GetType() == EMetasoundFrontendClassType::Graph)
 	{
-		const UClass& Class = MetaSoundDoc->GetBaseMetaSoundUClass();
-		if (const UAssetDefinition* AssetDefinition = UAssetDefinitionRegistry::Get()->GetAssetDefinitionForClass(&Class))
+		const FMetaSoundAssetKey AssetKey(NodeFrontendClass->Metadata);
+		if (const TScriptInterface<IMetaSoundDocumentInterface> MetaSoundDoc = IMetaSoundAssetManager::GetChecked().FindAssetAsDocumentInterface(AssetKey))
 		{
-			const FAssetData AssetData(MetaSoundDoc.GetObject());
-			return AssetDefinition->GetIconBrush(AssetData, AssetData.AssetClassPath.GetAssetName());
+			const UClass& Class = MetaSoundDoc->GetBaseMetaSoundUClass();
+			if (const UAssetDefinition* AssetDefinition = UAssetDefinitionRegistry::Get()->GetAssetDefinitionForClass(&Class))
+			{
+				const FAssetData AssetData(MetaSoundDoc.GetObject());
+				return AssetDefinition->GetIconBrush(AssetData, AssetData.AssetClassPath.GetAssetName());
+			}
 		}
 	}
 
