[MetasoundEditor] Include External graph types in AssetDefinition icon brush lookup condition

#rnx
#Synced-CL 51794084

diff --git a/Engine/Plugins/Runtime/Metasound/Source/MetasoundEditor/Private/MetasoundEditorGraphNode.cpp b/Engine/Plugins/Runtime/Metasound/Source/MetasoundEditor/Private/MetasoundEditorGraphNode.cpp
--- a/Engine/Plugins/Runtime/Metasound/Source/MetasoundEditor/Private/MetasoundEditorGraphNode.cpp
+++ b/Engine/Plugins/Runtime/Metasound/Source/MetasoundEditor/Private/MetasoundEditorGraphNode.cpp
@@ -823,7 +823,8 @@ const FSlateBrush* UMetasoundEditorGraphNode::GetNodeTitleIconBrush() const
 		return DefaultIcon.GetIcon();
 	}
 
-	if (NodeFrontendClass->Metadata.GetType() == EMetasoundFrontendClassType::Graph)
+	EMetasoundFrontendClassType ClassType = NodeFrontendClass->Metadata.GetType();
+	if (ClassType == EMetasoundFrontendClassType::Graph || ClassType == EMetasoundFrontendClassType::External)
 	{
 		const FMetaSoundAssetKey AssetKey(NodeFrontendClass->Metadata);
 		if (const TScriptInterface<IMetaSoundDocumentInterface> MetaSoundDoc = IMetaSoundAssetManager::GetChecked().FindAssetAsDocumentInterface(AssetKey))
