Fix for crash occurring when copying and pasting a variable node into another MetaSound paged graph

[RN] Fix for crash when copying and pasting a variable node into another MetaSound paged graph

diff --git a/Engine/Plugins/Runtime/Metasound/Source/MetasoundEditor/Private/MetasoundEditorDocumentClipboardUtils.cpp b/Engine/Plugins/Runtime/Metasound/Source/MetasoundEditor/Private/MetasoundEditorDocumentClipboardUtils.cpp
--- a/Engine/Plugins/Runtime/Metasound/Source/MetasoundEditor/Private/MetasoundEditorDocumentClipboardUtils.cpp
+++ b/Engine/Plugins/Runtime/Metasound/Source/MetasoundEditor/Private/MetasoundEditorDocumentClipboardUtils.cpp
@@ -275,10 +275,10 @@ namespace Metasound::Editor
 			VariableNode->CreateNewGuid();
 
 			TObjectPtr<UMetasoundEditorGraphVariable>& Variable = VariableNode->Variable;
+			const FMetasoundEditorGraphMemberNodeBreadcrumb PastedBreadcrumb = VariableNode->Breadcrumb;
 			if (!Variable || !Graph.ContainsVariable(*Variable) || Variable->GetFrontendVariable() == nullptr)
 			{
-				const FMetasoundEditorGraphMemberNodeBreadcrumb& Breadcrumb = VariableNode->Breadcrumb;
-				const FName BaseName = Breadcrumb.MemberName;
+				const FName BaseName = PastedBreadcrumb.MemberName;
 				TObjectPtr<UMetasoundEditorGraphVariable> CachedVariable = MappedGeneratedVariableNames.FindRef(BaseName);
 				const FMetasoundFrontendVariable* FrontendVariable = nullptr;
 				if (CachedVariable)
@@ -286,20 +286,20 @@ namespace Metasound::Editor
 					FrontendVariable = CachedVariable->GetFrontendVariable();
 				}
 
-				if (FrontendVariable && FrontendVariable->TypeName == Breadcrumb.DataType)
+				if (FrontendVariable && FrontendVariable->TypeName == PastedBreadcrumb.DataType)
 				{
 					Variable = CachedVariable;
 				}
 				else
 				{
-					const FMetasoundFrontendLiteral* Literal = Breadcrumb.DefaultLiterals.Find(Frontend::DefaultPageID);
+					const FMetasoundFrontendLiteral* Literal = PastedBreadcrumb.DefaultLiterals.Find(Frontend::DefaultPageID);
 					const FName VariableName = DocBuilder.GenerateUniqueVariableName(BaseName.ToString());
 					FrontendVariable = DocBuilder.AddGraphVariable(
 						VariableName,
-						Breadcrumb.DataType,
+						PastedBreadcrumb.DataType,
 						Literal,
-						&Breadcrumb.VertexMetadata.GetDisplayName(),
-						&Breadcrumb.VertexMetadata.GetDescription()
+						&PastedBreadcrumb.VertexMetadata.GetDisplayName(),
+						&PastedBreadcrumb.VertexMetadata.GetDescription()
 					);
 
 					Variable = Graph.FindOrAddVariable(FrontendVariable->Name);
@@ -313,8 +313,9 @@ namespace Metasound::Editor
 			const FMetasoundFrontendVariable* FrontendVariable = Variable->GetFrontendVariable();
 			if (ensure(FrontendVariable))
 			{
-				// Can only have one mutator/setter node
-				const FNodeClassName NodeClassName = VariableNode->GetBreadcrumb().ClassName.ToNodeClassName();
+				// Can only have one mutator/setter node. Check pasted breadcrumb as nodes may have
+				// changed due to it being added above (i.e. node was pasted from another graph).
+				const FNodeClassName NodeClassName = PastedBreadcrumb.ClassName.ToNodeClassName();
 				const bool bMatchesMutatorNodeID = VariableNode->GetNodeID() == FrontendVariable->MutatorNodeID;
 				bool bIsDuplicateMutatorNode = false;
 				if (const FMetasoundFrontendNode* MutatorNode = DocBuilder.FindNode(FrontendVariable->MutatorNodeID))
@@ -339,6 +340,7 @@ namespace Metasound::Editor
 						if (const FMetasoundFrontendNode* NewNode = DocBuilder.AddGraphVariableNode(FrontendVariable->Name, FrontendClass.Metadata.GetType()))
 						{
 							VariableNode->SetNodeID(NewNode->GetID());
+							VariableNode->CacheBreadcrumb(); // Cached again here to ensure node class is up-to-date on the copied node.
 						}
 						else
 						{
