Core module deprecation cleanup for 5.4 and earlier, I did not remove anything still in use.

#Synced-CL 52068627

diff --git a/Engine/Plugins/Runtime/Metasound/Source/MetasoundFrontend/Private/MetasoundFrontendGraphBuilder.cpp b/Engine/Plugins/Runtime/Metasound/Source/MetasoundFrontend/Private/MetasoundFrontendGraphBuilder.cpp
--- a/Engine/Plugins/Runtime/Metasound/Source/MetasoundFrontend/Private/MetasoundFrontendGraphBuilder.cpp
+++ b/Engine/Plugins/Runtime/Metasound/Source/MetasoundFrontend/Private/MetasoundFrontendGraphBuilder.cpp
@@ -1071,10 +1071,7 @@ PRAGMA_ENABLE_DEPRECATION_WARNINGS
 
 	TUniquePtr<FFrontendGraph> FGraphBuilder::CreateGraph(const FMetasoundFrontendDocument& InDocument, const FString& InDebugAssetName, TArrayView<const FGuid> InPageOrder)
 	{
-PRAGMA_DISABLE_DEPRECATION_WARNINGS
-		const FName AssetName = FName(*InDebugAssetName);
-		return CreateGraph(InDocument, FTopLevelAssetPath(AssetName), InPageOrder);
-PRAGMA_ENABLE_DEPRECATION_WARNINGS
+		return CreateGraph(InDocument, FTopLevelAssetPath(InDebugAssetName), InPageOrder);
 	}
 
 	TUniquePtr<FFrontendGraph> FGraphBuilder::CreateGraph(
@@ -1084,10 +1081,7 @@ PRAGMA_ENABLE_DEPRECATION_WARNINGS
 		const FGuid InGraphId,
 		TArrayView<const FGuid> InPageOrder)
 	{
-PRAGMA_DISABLE_DEPRECATION_WARNINGS
-		const FName AssetName = FName(*InDebugAssetName);
-		return CreateGraph(InDocument, InProxies, FTopLevelAssetPath(AssetName), InGraphId, InPageOrder);
-PRAGMA_ENABLE_DEPRECATION_WARNINGS
+		return CreateGraph(InDocument, InProxies, FTopLevelAssetPath(InDebugAssetName), InGraphId, InPageOrder);
 	}
 
 	TUniquePtr<FFrontendGraph> FGraphBuilder::CreateGraph(
@@ -1097,10 +1091,7 @@ PRAGMA_ENABLE_DEPRECATION_WARNINGS
 		const FString& InDebugAssetName,
 		TArrayView<const FGuid> InPageOrder)
 	{
-PRAGMA_DISABLE_DEPRECATION_WARNINGS
-		const FName AssetName = FName(*InDebugAssetName);
-		return CreateGraph(InGraph, InSubgraphs, InDependencies, FTopLevelAssetPath(AssetName), InPageOrder);
-PRAGMA_ENABLE_DEPRECATION_WARNINGS
+		return CreateGraph(InGraph, InSubgraphs, InDependencies, FTopLevelAssetPath(InDebugAssetName), InPageOrder);
 	}
 
 	TUniquePtr<FFrontendGraph> FGraphBuilder::CreateGraph(
@@ -1112,10 +1103,7 @@ PRAGMA_ENABLE_DEPRECATION_WARNINGS
 		const FGuid InGraphId,
 		TArrayView<const FGuid> InPageOrder)
 	{
-PRAGMA_DISABLE_DEPRECATION_WARNINGS
-		const FName AssetName = FName(*InDebugAssetName);
-		return CreateGraph(InGraph, InSubgraphs, InDependencies, InProxyDataCache, FTopLevelAssetPath(AssetName), InGraphId, InPageOrder);
-PRAGMA_ENABLE_DEPRECATION_WARNINGS
+		return CreateGraph(InGraph, InSubgraphs, InDependencies, InProxyDataCache, FTopLevelAssetPath(InDebugAssetName), InGraphId, InPageOrder);
 	}
 }
 
diff --git a/Engine/Plugins/Runtime/Metasound/Source/MetasoundGraphCore/Private/MetasoundGraph.cpp b/Engine/Plugins/Runtime/Metasound/Source/MetasoundGraphCore/Private/MetasoundGraph.cpp
--- a/Engine/Plugins/Runtime/Metasound/Source/MetasoundGraphCore/Private/MetasoundGraph.cpp
+++ b/Engine/Plugins/Runtime/Metasound/Source/MetasoundGraphCore/Private/MetasoundGraph.cpp
@@ -22,9 +22,7 @@ namespace Metasound
 	}
 
 	FGraph::FGraph(const FString& InInstanceName, const FGuid& InInstanceID, TSharedPtr<const IOperatorData> InOperatorData)
-PRAGMA_DISABLE_DEPRECATION_WARNINGS
-		: AssetPath(FTopLevelAssetPath(FName(*InInstanceName)))
-PRAGMA_ENABLE_DEPRECATION_WARNINGS
+		: AssetPath(FTopLevelAssetPath(InInstanceName))
 		, InstanceName(AssetPath.GetPackageName().IsNone() ? AssetPath.GetAssetName() : AssetPath.GetPackageName())
 		, InstanceID(InInstanceID)
 		, Metadata(FNodeClassMetadata::GetEmpty())
diff --git a/Engine/Source/Editor/UnrealEd/Private/PlayLevel.cpp b/Engine/Source/Editor/UnrealEd/Private/PlayLevel.cpp
--- a/Engine/Source/Editor/UnrealEd/Private/PlayLevel.cpp
+++ b/Engine/Source/Editor/UnrealEd/Private/PlayLevel.cpp
@@ -1004,10 +1004,6 @@ void UEditorEngine::RequestPlaySession(const FRequestPlaySessionParams& InParams
 
 	// Play Sessions can use the Game Mode to determine the default Player Start position
 	// or the start position can be overridden by the incoming launch arguments
-	PRAGMA_DISABLE_DEPRECATION_WARNINGS
-	GIsPIEUsingPlayerStart = !InParams.StartLocation.IsSet();
-	PRAGMA_ENABLE_DEPRECATION_WARNINGS
-
 	if (InParams.SessionDestination == EPlaySessionDestinationType::Launcher)
 	{
 		check(InParams.LauncherTargetDevice.IsSet());
diff --git a/Engine/Source/Runtime/ApplicationCore/Public/GenericPlatform/GenericApplication.h b/Engine/Source/Runtime/ApplicationCore/Public/GenericPlatform/GenericApplication.h
--- a/Engine/Source/Runtime/ApplicationCore/Public/GenericPlatform/GenericApplication.h
+++ b/Engine/Source/Runtime/ApplicationCore/Public/GenericPlatform/GenericApplication.h
@@ -396,9 +396,6 @@ struct FDisplayMetrics
 	/** The safe area for less important spill over on TVs (see TitleSafePaddingSize) */
 	FVector4 ActionSafePaddingSize;
 
-	UE_DEPRECATED(4.21, "Please use RebuildDisplayMetrics - it is functionally the same but is clearer about the function cost")
-	static void GetDisplayMetrics(struct FDisplayMetrics& OutDisplayMetrics) { RebuildDisplayMetrics(OutDisplayMetrics); };
-
 	APPLICATIONCORE_API static void RebuildDisplayMetrics(struct FDisplayMetrics& OutDisplayMetrics);
 
 	/** Gets the monitor work area from a position in the global display rect */
diff --git a/Engine/Source/Runtime/AssetRegistry/Private/AssetRegistryArchive.cpp b/Engine/Source/Runtime/AssetRegistry/Private/AssetRegistryArchive.cpp
--- a/Engine/Source/Runtime/AssetRegistry/Private/AssetRegistryArchive.cpp
+++ b/Engine/Source/Runtime/AssetRegistry/Private/AssetRegistryArchive.cpp
@@ -91,9 +91,7 @@ static void LoadBundleEntriesOldVersion(FArchive& Ar, TArray<FAssetBundleEntry>&
 				FString SubPathString;
 				Ar << SubPathString;
 			
-PRAGMA_DISABLE_DEPRECATION_WARNINGS
-				Path = FTopLevelAssetPath(AssetPathName);
-PRAGMA_ENABLE_DEPRECATION_WARNINGS
+				Path = FTopLevelAssetPath(AssetPathName.ToString());
 			}
 		}
 		else
diff --git a/Engine/Source/Runtime/Core/Private/Android/AndroidPlatformMisc.cpp b/Engine/Source/Runtime/Core/Private/Android/AndroidPlatformMisc.cpp
--- a/Engine/Source/Runtime/Core/Private/Android/AndroidPlatformMisc.cpp
+++ b/Engine/Source/Runtime/Core/Private/Android/AndroidPlatformMisc.cpp
@@ -679,12 +679,6 @@ void FAndroidMisc::PlatformHandleSplashScreen(bool ShowSplashScreen)
 #endif
 }
 
-void FAndroidMisc::GetEnvironmentVariable(const TCHAR* VariableName, TCHAR* Result, int32 ResultLength)
-{
-	*Result = 0;
-	// @todo Android : get environment variable.
-}
-
 FString FAndroidMisc::GetEnvironmentVariable(const TCHAR* VariableName)
 {
 	// @todo Android : get environment variable.
diff --git a/Engine/Source/Runtime/Core/Private/Apple/ApplePlatformMisc.cpp b/Engine/Source/Runtime/Core/Private/Apple/ApplePlatformMisc.cpp
--- a/Engine/Source/Runtime/Core/Private/Apple/ApplePlatformMisc.cpp
+++ b/Engine/Source/Runtime/Core/Private/Apple/ApplePlatformMisc.cpp
@@ -54,23 +54,6 @@ void FApplePlatformMisc::MemoryBarrier()
 	__sync_synchronize();
 }
 
-void FApplePlatformMisc::GetEnvironmentVariable(const TCHAR* VariableName, TCHAR* Result, int32 ResultLength)
-{
-	// Replace hyphens with underscores. Some legacy UE environment variables (eg. UE-SharedDataCachePath) are in widespread
-	// usage in their hyphenated form, but are not normally valid shell variables.
-	FString FixedVariableName = VariableName;
-	FixedVariableName.ReplaceInline(TEXT("-"), TEXT("_"));
-	ANSICHAR *AnsiResult = getenv(TCHAR_TO_ANSI(*FixedVariableName));
-	if (AnsiResult)
-	{
-		FCString::Strncpy(Result, ANSI_TO_TCHAR(AnsiResult), ResultLength);
-	}
-	else
-	{
-		*Result = 0;
-	}
-}
-
 FString FApplePlatformMisc::GetEnvironmentVariable(const TCHAR* VariableName)
 {
 	// Replace hyphens with underscores. Some legacy UE environment variables (eg. UE-SharedDataCachePath) are in widespread
diff --git a/Engine/Source/Runtime/Core/Private/GenericPlatform/GenericPlatformMisc.cpp b/Engine/Source/Runtime/Core/Private/GenericPlatform/GenericPlatformMisc.cpp
--- a/Engine/Source/Runtime/Core/Private/GenericPlatform/GenericPlatformMisc.cpp
+++ b/Engine/Source/Runtime/Core/Private/GenericPlatform/GenericPlatformMisc.cpp
@@ -327,18 +327,6 @@ const TCHAR* LexToString( EBuildConfiguration Configuration )
 
 namespace EBuildConfigurations
 {
-	EBuildConfiguration FromString( const FString& Configuration )
-	{
-		EBuildConfiguration Result;
-		LexTryParseString(Result, *Configuration);
-		return Result;
-	}
-
-	const TCHAR* ToString( EBuildConfiguration Configuration )
-	{
-		return LexToString(Configuration);
-	}
-
 	FText ToText( EBuildConfiguration Configuration )
 	{
 		switch (Configuration)
diff --git a/Engine/Source/Runtime/Core/Private/HAL/LLM/LLM.cpp b/Engine/Source/Runtime/Core/Private/HAL/LLM/LLM.cpp
--- a/Engine/Source/Runtime/Core/Private/HAL/LLM/LLM.cpp
+++ b/Engine/Source/Runtime/Core/Private/HAL/LLM/LLM.cpp
@@ -1717,35 +1717,6 @@ bool FLowLevelMemTracker::FindTagByName(const TCHAR* Name, uint64& OutTag, ELLMT
 	return GetGlobals()->FindTagByNameInner(Name, OutTag, InTagSet);
 }
 
-const TCHAR* FLowLevelMemTracker::FindTagName(uint64 Tag) const
-{
-	// This function does not call an Inner function on GetGlobals because it is a deprecated function.
-
-	if (!IsEnabled())
-	{
-		return nullptr;
-	}
-	// Cannot call BootstrapInitialise without shenanigans because this function is const.
-	LLMCheck(GetGlobals()->bInitializedTracking);
-
-	static TMap<uint64, FString> FoundTags;
-	FString* Cached = FoundTags.Find(Tag);
-	if (Cached)
-	{
-		return **Cached;
-	}
-
-	FName DisplayName = FindTagDisplayName(Tag);
-	if (DisplayName.IsNone())
-	{
-		return nullptr;
-	}
-
-	FString& AddedCache = FoundTags.Add(Tag);
-	AddedCache = DisplayName.ToString();
-	return *AddedCache;
-}
-
 FName FLowLevelMemTracker::FindTagDisplayName(uint64 Tag) const
 {
 	if (!IsEnabled())
diff --git a/Engine/Source/Runtime/Core/Private/Internationalization/Internationalization.cpp b/Engine/Source/Runtime/Core/Private/Internationalization/Internationalization.cpp
--- a/Engine/Source/Runtime/Core/Private/Internationalization/Internationalization.cpp
+++ b/Engine/Source/Runtime/Core/Private/Internationalization/Internationalization.cpp
@@ -457,12 +457,6 @@ TArray<FString> FInternationalization::GetPrioritizedCultureNames(const FString&
 	return Implementation->GetPrioritizedCultureNames(Name);
 }
 
-void FInternationalization::GetCulturesWithAvailableLocalization(const TArray<FString>& InLocalizationPaths, TArray<FCultureRef>& OutAvailableCultures, const bool bIncludeDerivedCultures)
-{
-	const TArray<FString> LocalizedCultureNames = TextLocalizationResourceUtil::GetLocalizedCultureNames(InLocalizationPaths);
-	OutAvailableCultures = GetAvailableCultures(LocalizedCultureNames, bIncludeDerivedCultures);
-}
-
 TArray<FCultureRef> FInternationalization::GetAvailableCultures(const TArray<FString>& InCultureNames, const bool bIncludeDerivedCultures)
 {
 	TArray<FCultureRef> AvailableCultures;
diff --git a/Engine/Source/Runtime/Core/Private/Internationalization/Text.cpp b/Engine/Source/Runtime/Core/Private/Internationalization/Text.cpp
--- a/Engine/Source/Runtime/Core/Private/Internationalization/Text.cpp
+++ b/Engine/Source/Runtime/Core/Private/Internationalization/Text.cpp
@@ -2000,21 +2000,6 @@ const TCHAR* FTextStringHelper::ReadFromBuffer(const TCHAR* Buffer, FText& OutVa
 	}
 }
 
-bool FTextStringHelper::ReadFromString(const TCHAR* Buffer, FText& OutValue, const TCHAR* TextNamespace, const TCHAR* PackageNamespace, int32* OutNumCharsRead, const bool bRequiresQuotes, const EStringTableLoadingPolicy InLoadingPolicy)
-{
-	const TCHAR* const Start = Buffer;
-	Buffer = ReadFromBuffer(Buffer, OutValue, TextNamespace, PackageNamespace, bRequiresQuotes);
-	if (Buffer)
-	{
-		if (OutNumCharsRead)
-		{
-			*OutNumCharsRead = UE_PTRDIFF_TO_INT32(Buffer - Start);
-		}
-		return true;
-	}
-	return false;
-}
-
 void FTextStringHelper::WriteToBuffer(FString& Buffer, const FText& Value, const bool bRequiresQuotes, const bool bStripPackageNamespace)
 {
 	// Culture invariant text?
@@ -2044,12 +2029,6 @@ void FTextStringHelper::WriteToBuffer(FString& Buffer, const FText& Value, const
 	}
 }
 
-bool FTextStringHelper::WriteToString(FString& Buffer, const FText& Value, const bool bRequiresQuotes)
-{
-	WriteToBuffer(Buffer, Value, bRequiresQuotes);
-	return true;
-}
-
 bool FTextStringHelper::IsComplexText(const TCHAR* Buffer)
 {
 #define LOC_DEFINE_REGION
diff --git a/Engine/Source/Runtime/Core/Private/Misc/CoreGlobals.cpp b/Engine/Source/Runtime/Core/Private/Misc/CoreGlobals.cpp
--- a/Engine/Source/Runtime/Core/Private/Misc/CoreGlobals.cpp
+++ b/Engine/Source/Runtime/Core/Private/Misc/CoreGlobals.cpp
@@ -393,8 +393,6 @@ bool					GIsCookerLoadingPackage			= false;
  * during PIE, which can happen when PIE initiates a Browse to another level.
  */
 bool					GIsPlayInEditorWorld			= false;
-/** Whether or not PIE was attempting to play from PlayerStart												*/
-bool					GIsPIEUsingPlayerStart			= false;
 /** true if the runtime needs textures to be powers of two													*/
 bool					GPlatformNeedsPowerOfTwoTextures= false;
 /** Time at which FPlatformTime::Seconds() was first initialized (before main)								*/
diff --git a/Engine/Source/Runtime/Core/Private/Misc/FileHelper.cpp b/Engine/Source/Runtime/Core/Private/Misc/FileHelper.cpp
--- a/Engine/Source/Runtime/Core/Private/Misc/FileHelper.cpp
+++ b/Engine/Source/Runtime/Core/Private/Misc/FileHelper.cpp
@@ -422,12 +422,6 @@ bool FFileHelper::LoadFileToStringArray( TArray<FString>& Result, const TCHAR* F
 	return LoadFileToStringArrayWithPredicate(Result, Filename, [](const FString&) { return true;  });
 }
 
-// DEPRECATED
-bool FFileHelper::LoadFileToStringArray(TArray<FString>& Result, const TCHAR* Filename, EHashOptions VerifyFlags)
-{
-	return LoadFileToStringArray(Result, Filename);
-}
-
 bool FFileHelper::LoadFileToStringArrayWithPredicate(TArray<FString>& Result, const TCHAR* Filename, TFunctionRef<bool(const FString&)> Predicate)
 {
 	Result.Empty();
@@ -521,12 +515,6 @@ bool FFileHelper::LoadFileToStringArrayWithPredicate(TArray<FString>& Result, co
 	return true;
 }
 
-// DEPRECATED
-bool FFileHelper::LoadFileToStringArrayWithPredicate(TArray<FString>& Result, const TCHAR* Filename, TFunctionRef<bool(const FString&)> Predicate, EHashOptions VerifyFlags)
-{
-	return LoadFileToStringArrayWithPredicate(Result, Filename, Predicate);
-}
-
 namespace UE::FileHelper::Private
 {
 
diff --git a/Engine/Source/Runtime/Core/Private/Misc/PathViews.cpp b/Engine/Source/Runtime/Core/Private/Misc/PathViews.cpp
--- a/Engine/Source/Runtime/Core/Private/Misc/PathViews.cpp
+++ b/Engine/Source/Runtime/Core/Private/Misc/PathViews.cpp
@@ -117,11 +117,6 @@ FStringView FPathViews::GetBaseFilenameWithPath(const FStringView& InPath)
 	return InPath.LeftChop(GetExtension(InPath, /*bIncludeDot*/ true).Len());
 }
 
-FStringView FPathViews::GetBaseFilename(const FStringView& InPath, bool bRemovePath)
-{
-	return bRemovePath ? GetBaseFilename(InPath) : GetBaseFilenameWithPath(InPath);
-}
-
 FStringView FPathViews::GetPath(const FStringView& InPath)
 {
 	if (const TCHAR* EndPos = Algo::FindLastByPredicate(InPath, UE4PathViews_Private::IsSlashOrBackslash))
diff --git a/Engine/Source/Runtime/Core/Private/Serialization/CustomVersion.cpp b/Engine/Source/Runtime/Core/Private/Serialization/CustomVersion.cpp
--- a/Engine/Source/Runtime/Core/Private/Serialization/CustomVersion.cpp
+++ b/Engine/Source/Runtime/Core/Private/Serialization/CustomVersion.cpp
@@ -264,16 +264,6 @@ const FName FCustomVersion::GetFriendlyName() const
 	return FriendlyName;
 }
 
-const FCustomVersionContainer& FCustomVersionContainer::GetRegistered()
-{
-	FStaticCustomVersionRegistry& Registry = FStaticCustomVersionRegistry::Get();
-
-	// Even though returning a reference isn't thread-safe, we can still synchronize access to the queue
-	UE::TWriteScopeLock Scope(Registry.Lock);
-	Registry.RegisterQueue();
-	return Registry.Registered;
-}
-
 void FCustomVersionContainer::Empty()
 {
 	Versions.Empty();
diff --git a/Engine/Source/Runtime/Core/Private/Unix/UnixPlatformMisc.cpp b/Engine/Source/Runtime/Core/Private/Unix/UnixPlatformMisc.cpp
--- a/Engine/Source/Runtime/Core/Private/Unix/UnixPlatformMisc.cpp
+++ b/Engine/Source/Runtime/Core/Private/Unix/UnixPlatformMisc.cpp
@@ -279,21 +279,6 @@ int32 FUnixPlatformMisc::GetMaxPathLength()
 	return PATH_MAX;
 }
 
-void FUnixPlatformMisc::GetEnvironmentVariable(const TCHAR* InVariableName, TCHAR* Result, int32 ResultLength)
-{
-	FString VariableName = InVariableName;
-	VariableName.ReplaceInline(TEXT("-"), TEXT("_"));
-	ANSICHAR *AnsiResult = secure_getenv(TCHAR_TO_ANSI(*VariableName));
-	if (AnsiResult)
-	{
-		FCString::Strncpy(Result, UTF8_TO_TCHAR(AnsiResult), ResultLength);
-	}
-	else
-	{
-		*Result = 0;
-	}
-}
-
 FString FUnixPlatformMisc::GetEnvironmentVariable(const TCHAR* InVariableName)
 {
 	FString VariableName = InVariableName;
diff --git a/Engine/Source/Runtime/Core/Private/Windows/WindowsPlatformMisc.cpp b/Engine/Source/Runtime/Core/Private/Windows/WindowsPlatformMisc.cpp
--- a/Engine/Source/Runtime/Core/Private/Windows/WindowsPlatformMisc.cpp
+++ b/Engine/Source/Runtime/Core/Private/Windows/WindowsPlatformMisc.cpp
@@ -1259,15 +1259,6 @@ int32 FWindowsPlatformMisc::GetExternalAppMaxPathLength()
 	return MAX_PATH;
 }
 
-void FWindowsPlatformMisc::GetEnvironmentVariable(const TCHAR* VariableName, TCHAR* Result, int32 ResultLength)
-{
-	uint32 Error = ::GetEnvironmentVariableW(VariableName, Result, ResultLength);
-	if (Error <= 0)
-	{		
-		*Result = TEXT('\0');
-	}
-}
-
 FString FWindowsPlatformMisc::GetEnvironmentVariable(const TCHAR* VariableName)
 {
 	// Allocate the data for the string. Loop in case the variable happens to change while running, or the buffer isn't large enough.
diff --git a/Engine/Source/Runtime/Core/Public/Android/AndroidPlatformMisc.h b/Engine/Source/Runtime/Core/Public/Android/AndroidPlatformMisc.h
--- a/Engine/Source/Runtime/Core/Public/Android/AndroidPlatformMisc.h
+++ b/Engine/Source/Runtime/Core/Public/Android/AndroidPlatformMisc.h
@@ -51,9 +51,6 @@ struct FAndroidMisc : public FGenericPlatformMisc
 		return ANDROID_MAX_PATH;
 	}
 
-	UE_DEPRECATED(4.21, "void FPlatformMisc::GetEnvironmentVariable(Name, Result, Length) is deprecated. Use FString FPlatformMisc::GetEnvironmentVariable(Name) instead.")
-	static CORE_API void GetEnvironmentVariable(const TCHAR* VariableName, TCHAR* Result, int32 ResultLength);
-
 	static CORE_API FString GetEnvironmentVariable(const TCHAR* VariableName);
 	static CORE_API const TCHAR* GetSystemErrorMessage(TCHAR* OutBuffer, int32 BufferCount, int32 Error);
 	static CORE_API uint32 GetLastError();
diff --git a/Engine/Source/Runtime/Core/Public/Apple/ApplePlatformMisc.h b/Engine/Source/Runtime/Core/Public/Apple/ApplePlatformMisc.h
--- a/Engine/Source/Runtime/Core/Public/Apple/ApplePlatformMisc.h
+++ b/Engine/Source/Runtime/Core/Public/Apple/ApplePlatformMisc.h
@@ -43,9 +43,6 @@ struct CORE_API FApplePlatformMisc : public FGenericPlatformMisc
 {
 	static CORE_API void PlatformInit();
 
-	UE_DEPRECATED(4.21, "void FPlatformMisc::GetEnvironmentVariable(Name, Result, Length) is deprecated. Use FString FPlatformMisc::GetEnvironmentVariable(Name) instead.")
-	static void GetEnvironmentVariable(const TCHAR* VariableName, TCHAR* Result, int32 ResultLength);
-
 	static FString GetEnvironmentVariable(const TCHAR* VariableName);
 
 #if !UE_BUILD_SHIPPING
diff --git a/Engine/Source/Runtime/Core/Public/CoreGlobals.h b/Engine/Source/Runtime/Core/Public/CoreGlobals.h
--- a/Engine/Source/Runtime/Core/Public/CoreGlobals.h
+++ b/Engine/Source/Runtime/Core/Public/CoreGlobals.h
@@ -509,10 +509,6 @@ namespace UE
 	CORE_API void  SetPlayInEditorID(int32 InPlayInEditorID);
 }
 
-/** Whether or not PIE was attempting to play from PlayerStart */
-UE_DEPRECATED(4.25, "This variable is no longer set. Use !GEditor->GetPlayInEditorSessionInfo()->OriginalRequestParams.HasPlayWorldPlacement() instead.")
-extern CORE_API bool GIsPIEUsingPlayerStart;
-
 /** true if the runtime needs textures to be powers of two */
 extern CORE_API bool GPlatformNeedsPowerOfTwoTextures;
 
@@ -632,7 +628,7 @@ enum class ETaskTag : int32
 	EStaticInit					= 1 << 0,
 	EGameThread					= 1 << 1,
 	ESlateThread				= 1 << 2,
-	EAudioThread UE_DEPRECATED(5.0, "AudioThread was removed and ETaskTag::EAudioThread is not used anymore. Please remove it.") = 1 << 3,
+	// Removed EAudioThread		= 1 << 3,
 	ERenderingThread			= 1 << 4,
 	ERhiThread					= 1 << 5,
 	EAsyncLoadingThread			= 1 << 6,
diff --git a/Engine/Source/Runtime/Core/Public/GenericPlatform/GenericPlatformMisc.h b/Engine/Source/Runtime/Core/Public/GenericPlatform/GenericPlatformMisc.h
--- a/Engine/Source/Runtime/Core/Public/GenericPlatform/GenericPlatformMisc.h
+++ b/Engine/Source/Runtime/Core/Public/GenericPlatform/GenericPlatformMisc.h
@@ -104,45 +104,6 @@ CORE_API const TCHAR* LexToString(EBuildConfiguration Configuration);
 
 namespace EBuildConfigurations
 {
-	UE_DEPRECATED(4.24, "EBuildConfigurations::Type is deprecated. Use EBuildConfiguration instead.")
-	typedef EBuildConfiguration Type;
-
-	UE_DEPRECATED(4.24, "EBuildConfigurations::Unknown is deprecated. Use EBuildConfiguration::Unknown instead.")
-	static constexpr EBuildConfiguration Unknown = EBuildConfiguration::Unknown;
-
-	UE_DEPRECATED(4.24, "EBuildConfigurations::Debug is deprecated. Use EBuildConfiguration::Debug instead.")
-	static constexpr EBuildConfiguration Debug = EBuildConfiguration::Debug;
-
-	UE_DEPRECATED(4.24, "EBuildConfigurations::DebugGame is deprecated. Use EBuildConfiguration::DebugGame instead.")
-	static constexpr EBuildConfiguration DebugGame = EBuildConfiguration::DebugGame;
-
-	UE_DEPRECATED(4.24, "EBuildConfigurations::Development is deprecated. Use EBuildConfiguration::Development instead.")
-	static constexpr EBuildConfiguration Development = EBuildConfiguration::Development;
-
-	UE_DEPRECATED(4.24, "EBuildConfigurations::Test is deprecated. Use EBuildConfiguration::Test instead.")
-	static constexpr EBuildConfiguration Test = EBuildConfiguration::Test;
-
-	UE_DEPRECATED(4.24, "EBuildConfigurations::Shipping is deprecated. Use EBuildConfiguration::Shipping instead.")
-	static constexpr EBuildConfiguration Shipping = EBuildConfiguration::Shipping;
-
-	/**
-	 * Returns the string representation of the specified EBuildConfiguration value.
-	 *
-	 * @param Configuration The string to get the EBuildConfiguration for.
-	 * @return An EBuildConfiguration value.
-	 */
-	UE_DEPRECATED(4.24, "EBuildConfigurations::FromString() is deprecated. Use LexFromString() instead.")
-	CORE_API EBuildConfiguration FromString( const FString& Configuration );
-
-	/**
-	 * Returns the string representation of the specified EBuildConfiguration value.
-	 *
-	 * @param Configuration The value to get the string for.
-	 * @return The string representation.
-	 */
-	UE_DEPRECATED(4.24, "EBuildConfigurations::ToString() is deprecated. Use LexToString() instead.")
-	CORE_API const TCHAR* ToString( EBuildConfiguration Configuration );
-
 	/**
 	 * Returns the localized text representation of the specified EBuildConfiguration value.
 	 *
diff --git a/Engine/Source/Runtime/Core/Public/HAL/LowLevelMemTracker.h b/Engine/Source/Runtime/Core/Public/HAL/LowLevelMemTracker.h
--- a/Engine/Source/Runtime/Core/Public/HAL/LowLevelMemTracker.h
+++ b/Engine/Source/Runtime/Core/Public/HAL/LowLevelMemTracker.h
@@ -522,14 +522,6 @@ private:
 
 } // UE::LLMPrivate
 
-struct UE_DEPRECATED(4.27, "FLLMCustomTag was an implementation detail that has been modified, switch to FLLMTagInfo or to your own local struct") FLLMCustomTag
-{
-	int32 Tag;
-	const TCHAR* Name;
-	FName StatName;
-	FName SummaryStatName;
-};
-
 /** A convenient struct for gathering the fields needed to report in RegisterProjectTag */
 struct FLLMTagInfo
 {
@@ -650,9 +642,6 @@ public:
 	/** Look up the ELLMTag associated with the given display name. */
 	CORE_API bool FindTagByName(const TCHAR* Name, uint64& OutTag, ELLMTagSet InTagSet = ELLMTagSet::None) const;
 
-	UE_DEPRECATED(4.27, "Use FindTagDisplayName instead")
-	CORE_API const TCHAR* FindTagName(uint64 Tag) const;
-
 	/** Get the display name for the given ELLMTag. */
 	CORE_API FName FindTagDisplayName(uint64 Tag) const;
 	CORE_API FName FindPtrDisplayName(void* Ptr) const;
diff --git a/Engine/Source/Runtime/Core/Public/Internationalization/Internationalization.h b/Engine/Source/Runtime/Core/Public/Internationalization/Internationalization.h
--- a/Engine/Source/Runtime/Core/Public/Internationalization/Internationalization.h
+++ b/Engine/Source/Runtime/Core/Public/Internationalization/Internationalization.h
@@ -193,14 +193,6 @@ public:
 
 	CORE_API TArray<FString> GetPrioritizedCultureNames(const FString& Name);
 
-	/**
-	 * Given some paths to look at, populate a list of cultures that we have available localization information for. If bIncludeDerivedCultures, include cultures that are derived from those we have localization data for.
-	 * @note This function was deprecated as it only provides cultures that use localization resources.
-	 *		 FTextLocalizationManager::GetLocalizedCultureNames provides a more complete list when using alternate sources, and FInternationalization::GetAvailableCultures can be used to build a culture array from that list.
-	 */
-	UE_DEPRECATED(4.20, "FInternationalization::GetCulturesWithAvailableLocalization is deprecated in favor of calling FTextLocalizationManager::GetLocalizedCultureNames, potentially followed by FInternationalization::GetAvailableCultures")
-	CORE_API void GetCulturesWithAvailableLocalization(const TArray<FString>& InLocalizationPaths, TArray<FCultureRef>& OutAvailableCultures, const bool bIncludeDerivedCultures);
-
 	/** Given some culture names, populate a list of cultures that are available to be used. If bIncludeDerivedCultures, include cultures that are derived from those we passed. */
 	CORE_API TArray<FCultureRef> GetAvailableCultures(const TArray<FString>& InCultureNames, const bool bIncludeDerivedCultures);
 
diff --git a/Engine/Source/Runtime/Core/Public/Internationalization/Text.h b/Engine/Source/Runtime/Core/Public/Internationalization/Text.h
--- a/Engine/Source/Runtime/Core/Public/Internationalization/Text.h
+++ b/Engine/Source/Runtime/Core/Public/Internationalization/Text.h
@@ -1297,9 +1297,6 @@ public:
 	 * @return The updated buffer after we parsed this text, or nullptr on failure
 	 */
 	static CORE_API const TCHAR* ReadFromBuffer(const TCHAR* Buffer, FText& OutValue, const TCHAR* TextNamespace = nullptr, const TCHAR* PackageNamespace = nullptr, const bool bRequiresQuotes = false);
-	
-	UE_DEPRECATED(4.22, "FTextStringHelper::ReadFromString is deprecated. Use FTextStringHelper::ReadFromBuffer instead.")
-	static CORE_API bool ReadFromString(const TCHAR* Buffer, FText& OutValue, const TCHAR* TextNamespace = nullptr, const TCHAR* PackageNamespace = nullptr, int32* OutNumCharsRead = nullptr, const bool bRequiresQuotes = false, const EStringTableLoadingPolicy InLoadingPolicy = EStringTableLoadingPolicy::FindOrLoad);
 
 	/**
 	 * Write the given FText instance to a stream of text
@@ -1310,9 +1307,6 @@ public:
 	 * @param bStripPackageNamespace True to strip the package namespace from the written NSLOCTEXT value (eg, when saving cooked data)
 	 */
 	static CORE_API void WriteToBuffer(FString& Buffer, const FText& Value, const bool bRequiresQuotes = false, const bool bStripPackageNamespace = false);
-	
-	UE_DEPRECATED(4.22, "FTextStringHelper::WriteToString is deprecated. Use FTextStringHelper::WriteToBuffer instead.")
-	static CORE_API bool WriteToString(FString& Buffer, const FText& Value, const bool bRequiresQuotes = false);
 
 	/**
 	 * Test to see whether a given buffer contains complex text.
diff --git a/Engine/Source/Runtime/Core/Public/Misc/AssertionMacros.h b/Engine/Source/Runtime/Core/Public/Misc/AssertionMacros.h
--- a/Engine/Source/Runtime/Core/Public/Misc/AssertionMacros.h
+++ b/Engine/Source/Runtime/Core/Public/Misc/AssertionMacros.h
@@ -209,9 +209,6 @@ RetType UE_COLD UE_DEBUG_SECTION DispatchCheckVerify (InnerType&& Inner, ArgType
 	#define UE_DEBUG_BREAK_AND_PROMPT_FOR_REMOTE()
 #endif // !UE_BUILD_SHIPPING
 
-#define _DebugBreakAndPromptForRemote() \
-	UE_DEPRECATED_MACRO(5.1, "Use UE_DEBUG_BREAK_AND_PROMPT_FOR_REMOTE.")
-
 #if !UE_BUILD_SHIPPING
 	extern CORE_API bool GIgnoreDebugger;
 	// This is named with PLATFORM_ prefix because UE_DEBUG_BREAK* are conditional on a debugger being detected, and PLATFORM_BREAK isn't
diff --git a/Engine/Source/Runtime/Core/Public/Misc/FileHelper.h b/Engine/Source/Runtime/Core/Public/Misc/FileHelper.h
--- a/Engine/Source/Runtime/Core/Public/Misc/FileHelper.h
+++ b/Engine/Source/Runtime/Core/Public/Misc/FileHelper.h
@@ -162,9 +162,6 @@ struct FFileHelper
 	 */
 	static CORE_API bool LoadFileToStringArray( TArray<FString>& Result, const TCHAR* Filename );
 
-	UE_DEPRECATED(4.26, "LoadFileToStringArray no longer supports VerifyFlags. You can use UE::String::ParseLines to split up a string loaded with LoadFileToString")
-	static CORE_API bool LoadFileToStringArray(TArray<FString>& Result, const TCHAR* Filename, EHashOptions VerifyFlags);
-
 	/**
 	 * Load a text file to an array of strings, filtered by a user-defined predicate. Supports all combination of ANSI/Unicode files and platforms.
 	 *
@@ -174,9 +171,6 @@ struct FFileHelper
 	 */
 	static CORE_API bool LoadFileToStringArrayWithPredicate(TArray<FString>& Result, const TCHAR* Filename, TFunctionRef<bool(const FString&)> Predicate);
 
-	UE_DEPRECATED(4.26, "LoadFileToStringArrayWithPredicate no longer supports VerifyFlags. You can use UE::String::ParseLines to split up a string loaded with LoadFileToString")
-	static CORE_API bool LoadFileToStringArrayWithPredicate(TArray<FString>& Result, const TCHAR* Filename, TFunctionRef<bool(const FString&)> Predicate, EHashOptions VerifyFlags);
-
 	/**
 	 * Load a text file and invoke a visitor for each line. Supports all combination of ANSI/Unicode files and platforms.
 	 *
diff --git a/Engine/Source/Runtime/Core/Public/Misc/PathViews.h b/Engine/Source/Runtime/Core/Public/Misc/PathViews.h
--- a/Engine/Source/Runtime/Core/Public/Misc/PathViews.h
+++ b/Engine/Source/Runtime/Core/Public/Misc/PathViews.h
@@ -53,11 +53,6 @@ public:
 	 */
 	static CORE_API FStringView GetBaseFilenameWithPath(const FStringView& InPath);
 
-	/** Returns the portion of the path before the last dot and optionally after the last separator. */
-	UE_DEPRECATED(4.25, "FPathViews::GetBaseFilename(InPath, bRemovePath) has been superseded by "
-		"FPathViews::GetBaseFilename(InPath) and FPathViews::GetBaseFilenameWithPath(InPath).")
-	static CORE_API FStringView GetBaseFilename(const FStringView& InPath, bool bRemovePath);
-
 	/**
 	 * Returns the portion of the path before the last separator.
 	 *
diff --git a/Engine/Source/Runtime/Core/Public/Misc/TimeGuard.h b/Engine/Source/Runtime/Core/Public/Misc/TimeGuard.h
--- a/Engine/Source/Runtime/Core/Public/Misc/TimeGuard.h
+++ b/Engine/Source/Runtime/Core/Public/Misc/TimeGuard.h
@@ -133,9 +133,6 @@ protected:
 	double			StartTime;
 };
 
-UE_DEPRECATED(4.21, "FLightweightTimeGuard has been renamed to FTimeGuard.")
-typedef FTimeGuard FLightweightTimeGuard;
-
 #define SCOPE_TIME_GUARD(name) \
 	FTimeGuard ANONYMOUS_VARIABLE(TimeGuard)(name);
 
diff --git a/Engine/Source/Runtime/Core/Public/Serialization/CustomVersion.h b/Engine/Source/Runtime/Core/Public/Serialization/CustomVersion.h
--- a/Engine/Source/Runtime/Core/Public/Serialization/CustomVersion.h
+++ b/Engine/Source/Runtime/Core/Public/Serialization/CustomVersion.h
@@ -154,14 +154,6 @@ public:
 	CORE_API void Serialize(FArchive& Ar, ECustomVersionSerializationFormat Format = ECustomVersionSerializationFormat::Latest);
 	CORE_API void Serialize(FStructuredArchive::FSlot Slot, ECustomVersionSerializationFormat Format = ECustomVersionSerializationFormat::Latest);
 
-	/**
-	 * Gets a singleton with the registered versions.
-	 *
-	 * @return The registered version container.
-	 */
-	UE_DEPRECATED(4.24, "Use one of the thread-safe FCurrentCustomVersions methods instead")
-	static CORE_API const FCustomVersionContainer& GetRegistered();
-
 	/**
 	 * Empties the custom version container.
 	 */
diff --git a/Engine/Source/Runtime/Core/Public/Unix/UnixPlatformMisc.h b/Engine/Source/Runtime/Core/Public/Unix/UnixPlatformMisc.h
--- a/Engine/Source/Runtime/Core/Public/Unix/UnixPlatformMisc.h
+++ b/Engine/Source/Runtime/Core/Public/Unix/UnixPlatformMisc.h
@@ -28,9 +28,6 @@ struct FUnixPlatformMisc : public FGenericPlatformMisc
 	static CORE_API void SetCrashHandler(void (* CrashHandler)(const FGenericCrashContext& Context));
 	static CORE_API int32 GetMaxPathLength();
 
-	UE_DEPRECATED(4.21, "void FPlatformMisc::GetEnvironmentVariable(Name, Result, Length) is deprecated. Use FString FPlatformMisc::GetEnvironmentVariable(Name) instead.")
-	static CORE_API void GetEnvironmentVariable(const TCHAR* VariableName, TCHAR* Result, int32 ResultLength);
-
 	static CORE_API FString GetEnvironmentVariable(const TCHAR* VariableName);
 	static CORE_API void SetEnvironmentVar(const TCHAR* VariableName, const TCHAR* Value);
 	static CORE_API TArray<uint8> GetMacAddress();
diff --git a/Engine/Source/Runtime/Core/Public/Windows/WindowsPlatformMisc.h b/Engine/Source/Runtime/Core/Public/Windows/WindowsPlatformMisc.h
--- a/Engine/Source/Runtime/Core/Public/Windows/WindowsPlatformMisc.h
+++ b/Engine/Source/Runtime/Core/Public/Windows/WindowsPlatformMisc.h
@@ -123,9 +123,6 @@ struct FWindowsPlatformMisc
 	static CORE_API int32 GetMaxPathLength();
 	static CORE_API int32 GetExternalAppMaxPathLength();
 
-	UE_DEPRECATED(4.21, "void FPlatformMisc::GetEnvironmentVariable(Name, Result, Length) is deprecated. Use FString FPlatformMisc::GetEnvironmentVariable(Name) instead.")
-	static CORE_API void GetEnvironmentVariable(const TCHAR* VariableName, TCHAR* Result, int32 ResultLength);
-
 	static CORE_API FString GetEnvironmentVariable(const TCHAR* VariableName);
 	static CORE_API void SetEnvironmentVar(const TCHAR* VariableName, const TCHAR* Value);
 
diff --git a/Engine/Source/Runtime/CoreUObject/Private/StructUtils/UserDefinedStruct.cpp b/Engine/Source/Runtime/CoreUObject/Private/StructUtils/UserDefinedStruct.cpp
--- a/Engine/Source/Runtime/CoreUObject/Private/StructUtils/UserDefinedStruct.cpp
+++ b/Engine/Source/Runtime/CoreUObject/Private/StructUtils/UserDefinedStruct.cpp
@@ -220,13 +220,6 @@ void UUserDefinedStruct::PostSaveRoot(FObjectPostSaveRootContext ObjectSaveConte
 	PurgeCookedMetaData();
 }
 
-void UUserDefinedStruct::GetAssetRegistryTags(TArray<FAssetRegistryTag>& OutTags) const
-{
-	PRAGMA_DISABLE_DEPRECATION_WARNINGS;
-	Super::GetAssetRegistryTags(OutTags);
-	PRAGMA_ENABLE_DEPRECATION_WARNINGS;
-}
-
 void UUserDefinedStruct::GetAssetRegistryTags(FAssetRegistryTagsContext Context) const
 {
 	Super::GetAssetRegistryTags(Context);
diff --git a/Engine/Source/Runtime/CoreUObject/Private/UObject/ObjectRedirector.cpp b/Engine/Source/Runtime/CoreUObject/Private/UObject/ObjectRedirector.cpp
--- a/Engine/Source/Runtime/CoreUObject/Private/UObject/ObjectRedirector.cpp
+++ b/Engine/Source/Runtime/CoreUObject/Private/UObject/ObjectRedirector.cpp
@@ -60,13 +60,6 @@ bool UObjectRedirector::NeedsLoadForEditorGame() const
 	return true;
 }
 
-void UObjectRedirector::GetAssetRegistryTags(TArray<FAssetRegistryTag>& OutTags) const
-{
-	PRAGMA_DISABLE_DEPRECATION_WARNINGS;
-	Super::GetAssetRegistryTags(OutTags);
-	PRAGMA_ENABLE_DEPRECATION_WARNINGS;
-}
-
 void UObjectRedirector::GetAssetRegistryTags(FAssetRegistryTagsContext Context) const
 {
 	Super::GetAssetRegistryTags(Context);
diff --git a/Engine/Source/Runtime/CoreUObject/Public/StructUtils/UserDefinedStruct.h b/Engine/Source/Runtime/CoreUObject/Public/StructUtils/UserDefinedStruct.h
--- a/Engine/Source/Runtime/CoreUObject/Public/StructUtils/UserDefinedStruct.h
+++ b/Engine/Source/Runtime/CoreUObject/Public/StructUtils/UserDefinedStruct.h
@@ -109,8 +109,6 @@ public:
 	//~ UObject interface.
 	COREUOBJECT_API virtual void PostDuplicate(bool bDuplicateForPIE) override;
 	COREUOBJECT_API virtual void GetAssetRegistryTags(FAssetRegistryTagsContext Context) const override;
-	UE_DEPRECATED(5.4, "Implement the version that takes FAssetRegistryTagsContext instead.")
-	COREUOBJECT_API virtual void GetAssetRegistryTags(TArray<FAssetRegistryTag>& OutTags) const override;
 	COREUOBJECT_API virtual void PostLoad() override;
 	COREUOBJECT_API virtual void PreSaveRoot(FObjectPreSaveRootContext ObjectSaveContext) override;
 	COREUOBJECT_API virtual void PostSaveRoot(FObjectPostSaveRootContext ObjectSaveContext) override;
diff --git a/Engine/Source/Runtime/CoreUObject/Public/UObject/ObjectRedirector.h b/Engine/Source/Runtime/CoreUObject/Public/UObject/ObjectRedirector.h
--- a/Engine/Source/Runtime/CoreUObject/Public/UObject/ObjectRedirector.h
+++ b/Engine/Source/Runtime/CoreUObject/Public/UObject/ObjectRedirector.h
@@ -38,8 +38,6 @@ class UObjectRedirector : public UObject
 	void Serialize(FStructuredArchive::FRecord Record) override;
 	virtual bool NeedsLoadForEditorGame() const override;
 	virtual void GetAssetRegistryTags(FAssetRegistryTagsContext Context) const override;
-	UE_DEPRECATED(5.4, "Implement the version that takes FAssetRegistryTagsContext instead.")
-	virtual void GetAssetRegistryTags(TArray<FAssetRegistryTag>& OutTags) const override;
 	virtual bool HasNonEditorOnlyReferences() const override
 	{
 		return true;
diff --git a/Engine/Source/Runtime/CoreUObject/Public/UObject/TopLevelAssetPath.h b/Engine/Source/Runtime/CoreUObject/Public/UObject/TopLevelAssetPath.h
--- a/Engine/Source/Runtime/CoreUObject/Public/UObject/TopLevelAssetPath.h
+++ b/Engine/Source/Runtime/CoreUObject/Public/UObject/TopLevelAssetPath.h
@@ -47,14 +47,6 @@ public:
 		TrySetPath(InPackageName, InAssetName);
 	}
 
-	UE_DEPRECATED(5.0, "FNames containing full asset paths have been replaced by FTopLevelAssetPath/FSoftLevelObjectPath."
-		"This function is only for temporary use interfacing with APIs that still produce an FName."
-		"Those APIS should be updated to use FTopLevelAssetPath or FSoftLevelObjectPath.")
-	explicit FTopLevelAssetPath(FName InPath)
-	{
-		TrySetPath(InPath.ToString());
-	}
-
 	/** Construct from string / string view / raw string of a supported character type. */
 	explicit FTopLevelAssetPath(const FString& Path) { TrySetPath(FStringView(Path)); }
 	template<typename CharType, typename = typename std::enable_if<TIsCharType<CharType>::Value>::type>
diff --git a/Engine/Source/Runtime/CoreUObject/Public/UObject/UObjectArray.h b/Engine/Source/Runtime/CoreUObject/Public/UObject/UObjectArray.h
--- a/Engine/Source/Runtime/CoreUObject/Public/UObject/UObjectArray.h
+++ b/Engine/Source/Runtime/CoreUObject/Public/UObject/UObjectArray.h
@@ -30,14 +30,6 @@ namespace Verse
 #	define UE_ENABLE_FUOBJECT_ITEM_PACKING UE_PACK_FUOBJECT_ITEM && !UE_WITH_REMOTE_OBJECT_HANDLE
 #endif
 
-/**
-* Controls whether the number of available elements is being tracked in the ObjObjects array.
-* By default it is only tracked in WITH_EDITOR builds as it adds a small amount of tracking overhead
-*/
-#if !defined(UE_GC_TRACK_OBJ_AVAILABLE)
-#define UE_GC_TRACK_OBJ_AVAILABLE UE_DEPRECATED_MACRO(5.2, "The UE_GC_TRACK_OBJ_AVAILABLE macro has been deprecated because it is no longer necessary.") 1
-#endif
-
 namespace UE::GC::Private
 {
 	class FGCFlags;
diff --git a/Engine/Source/Runtime/CoreUObject/Public/UObject/UObjectIterator.h b/Engine/Source/Runtime/CoreUObject/Public/UObject/UObjectIterator.h
--- a/Engine/Source/Runtime/CoreUObject/Public/UObject/UObjectIterator.h
+++ b/Engine/Source/Runtime/CoreUObject/Public/UObject/UObjectIterator.h
@@ -224,11 +224,6 @@ protected:
  */
 typedef TObjectIteratorBase<EObjectIteratorThreadSafetyOptions::None> FUnsafeObjectIterator;
 
-/** Deprecating the old FObjectIterator class which is now synonymous with FUnsafeObjectIterator */
-UE_DEPRECATED(4.27, "FObjectIterator is not thread safe, use FThreadSafeObjectIterator or FPresistentThreadSafeObjectIterator instead.")
-typedef FUnsafeObjectIterator FObjectIterator;
-
-
 /**
  * Class for iterating through all objects, including class default objects.
  * Note that when Playing In Editor, this will find objects in the
