Fix for missing registration call prior to building preset for the first time

#Synced-CL 51967495

diff --git a/Engine/Plugins/Runtime/Metasound/Source/MetasoundFrontend/Private/DocumentTemplates/MetasoundFrontendPresetTemplate.cpp b/Engine/Plugins/Runtime/Metasound/Source/MetasoundFrontend/Private/DocumentTemplates/MetasoundFrontendPresetTemplate.cpp
--- a/Engine/Plugins/Runtime/Metasound/Source/MetasoundFrontend/Private/DocumentTemplates/MetasoundFrontendPresetTemplate.cpp
+++ b/Engine/Plugins/Runtime/Metasound/Source/MetasoundFrontend/Private/DocumentTemplates/MetasoundFrontendPresetTemplate.cpp
@@ -76,6 +76,12 @@ bool FMetaSoundFrontendPresetTemplate::ConfigureDocument(FMetaSoundFrontendDocum
 	{
 		if (RefObject->IsA(&OutBuilder.GetConstDocumentInterfaceChecked().GetBaseMetaSoundUClass()))
 		{
+			// Ensure the parent is registered for execution before building the preset graph from it.
+			if (FMetasoundAssetBase* ParentAsset = IMetaSoundAssetManager::GetChecked().GetAsAsset(*RefObject))
+			{
+				ParentAsset->UpdateAndRegisterForExecution(IMetaSoundAssetManager::GetChecked().GetRegistrationOptions());
+			}
+
 			const FMetaSoundFrontendDocumentBuilder& ParentBuilder = IDocumentBuilderRegistry::GetChecked().FindOrBeginBuilding(Parent);
 
 			OutBuilder.SetIsGraphEditable(false);
