[CAT] MetasoundSource+Generator+Details Customization (from Jimmy Smith)

CL E: Engine + Generator + Editor integration

Split from CL 51557033. Wires channel-agnostic output through
MetasoundSource, MetasoundGenerator, and MetasoundEditor.

Original work by [REDACTED_PERSON].

Additional changes by [REDACTED_PERSON]:

- MoveTemp fix for TSAN race on OutputCat copy (MetasoundGeneratorBuilder.cpp)
- Corrected misleading comment in CreateSoundGenerator (MetasoundSource.cpp)
- Added SourceFormatName (FName) to FGeneratorInitParams and FMetasoundGenerator
to carry CAT format identity through the pipeline via Environment, avoiding
lossy NumChannels reverse lookup for format determination.

Bug fixes found during testing:
- UE-370803: Editor audio analyzer crash on non-standard channel counts.
AudioBus.h ConvertIntToEAudioBusChannels now clamps to nearest supported
format instead of asserting, with a warning log.
- UE-370811: CAT custom format change required two saves to take effect.
Added PostEditChangeProperty handler for ChannelAgnosticCustomFormat and
ChannelAgnosticFormatChooser. Also resolves NumChannels in
ConformObjectToDocument instead of using placeholder value.
- NumChannels clamped to mixer-compatible values (1,2,4,6,8) via
ClampToMixerChannelCount() in all paths (PostEditChangeProperty,
ConformObjectToDocument, CreateSoundGenerator).
- SetGraph MixerOutput always uses FindDiscreteFormatFromNumChannels
to guarantee mixer compatibility. SourceFormatName is for Environment
(graph-side format identity) only, not mixer output determination.
- Custom Format dropdown on MetaSound Source now restricted to
AudioMixer-compatible formats only (Mono, Stereo, Quad, 5.1, 7.1)
via new GetCastingOptions_AudioMixerOnly() on CatCastingOptionsHelper.
- CAT interface registration gated behind MetasoundExperimental plugin check.

#Synced-CL 51786389

diff --git a/Engine/Plugins/Experimental/MetasoundExperimental/Source/MetasoundExperimentalRuntime/Private/MetasoundCatCastingNode.cpp b/Engine/Plugins/Experimental/MetasoundExperimental/Source/MetasoundExperimentalRuntime/Private/MetasoundCatCastingNode.cpp
--- a/Engine/Plugins/Experimental/MetasoundExperimental/Source/MetasoundExperimentalRuntime/Private/MetasoundCatCastingNode.cpp
+++ b/Engine/Plugins/Experimental/MetasoundExperimental/Source/MetasoundExperimentalRuntime/Private/MetasoundCatCastingNode.cpp
@@ -10,6 +10,7 @@
 #include "MetasoundVertex.h"
 #include "DSP/MultiMono.h"
 #include "TypeFamily/ChannelTypeFamily.h"
+#include "ChannelAgnostic/ChannelAgnosticTypeUtils.h"
 
 #define LOCTEXT_NAMESPACE "MetasoundStandardNodes_CatCastingNode"
 
@@ -216,6 +217,58 @@ TArray<FPropertyTextFName> UMetasoundCatCastingOptionsHelper::GetCastingOptions_
 	return FormatsOptions;
 }
 
+TArray<FPropertyTextFName> UMetasoundCatCastingOptionsHelper::GetCastingOptions_DiscreteOnly()
+{
+	// All non-abstract discrete formats. Excludes soundfield (Atmos, ambisonics) and composite types.
+	struct FDiscreteChecker : public Audio::IChannelTypeVisitor
+	{
+		bool bIsDiscrete = false;
+		PRAGMA_DISABLE_EXPERIMENTAL_WARNINGS
+		void Visit(const Audio::FDiscreteChannelTypeFamily&) override { bIsDiscrete = true; }
+		void Visit(const Audio::FSoundfieldChannelTypeFamily&) override { bIsDiscrete = false; }
+		void Visit(const Audio::FCompositeChannelTypeFamily&) override { bIsDiscrete = false; }
+		PRAGMA_ENABLE_EXPERIMENTAL_WARNINGS
+	};
+
+	TArray<TSharedRef<const Audio::FChannelTypeFamily>> AllFormats = Audio::GetChannelRegistry().GetAllChannelFormats();
+	TArray<FPropertyTextFName> FormatsOptions;
+	for (const TSharedRef<const Audio::FChannelTypeFamily>& Format : AllFormats)
+	{
+		if (Format->IsAbstract())
+		{
+			continue;
+		}
+		FDiscreteChecker Checker;
+		PRAGMA_DISABLE_EXPERIMENTAL_WARNINGS
+		Format->Accept(Checker);
+		PRAGMA_ENABLE_EXPERIMENTAL_WARNINGS
+		if (Checker.bIsDiscrete)
+		{
+			FormatsOptions.Add({.ValueString = Format->GetName(), .DisplayName = FText::FromString(Format->GetFriendlyName())});
+		}
+	}
+	return FormatsOptions;
+}
+
+TArray<FPropertyTextFName> UMetasoundCatCastingOptionsHelper::GetCastingOptions_AudioMixerOnly()
+{
+	// Only the standard AudioMixer formats: Mono 1.0, Stereo 2.0, Quad 4.0, 5.1, 7.1.
+	// These are the exact formats AudioMixer supports for buffer allocation and interleaving.
+	PRAGMA_DISABLE_EXPERIMENTAL_WARNINGS
+	static constexpr int32 MixerChannelCounts[] = { 1, 2, 4, 6, 8 };
+	TArray<FPropertyTextFName> FormatsOptions;
+	for (int32 NumCh : MixerChannelCounts)
+	{
+		if (const TSharedPtr<const Audio::FDiscreteChannelTypeFamily> Format =
+			Audio::FChannelAgnosticUtils::FindDiscreteFormatFromNumChannels(NumCh))
+		{
+			FormatsOptions.Add({.ValueString = Format->GetName(), .DisplayName = FText::FromString(Format->GetFriendlyName())});
+		}
+	}
+	PRAGMA_ENABLE_EXPERIMENTAL_WARNINGS
+	return FormatsOptions;
+}
+
 TInstancedStruct<FMetasoundFrontendClassInterface> FMetaSoundCatCastingNodeConfiguration::OverrideDefaultInterface(
 	const FMetasoundFrontendClass& InNodeClass) const
 {
diff --git a/Engine/Plugins/Experimental/MetasoundExperimental/Source/MetasoundExperimentalRuntime/Private/MetasoundCatCastingNode.h b/Engine/Plugins/Experimental/MetasoundExperimental/Source/MetasoundExperimentalRuntime/Private/MetasoundCatCastingNode.h
--- a/Engine/Plugins/Experimental/MetasoundExperimental/Source/MetasoundExperimentalRuntime/Private/MetasoundCatCastingNode.h
+++ b/Engine/Plugins/Experimental/MetasoundExperimental/Source/MetasoundExperimentalRuntime/Private/MetasoundCatCastingNode.h
@@ -22,6 +22,15 @@ class UMetasoundCatCastingOptionsHelper : public UObject
 
 	UFUNCTION()
 	static TArray<FPropertyTextFName> GetCastingOptions_NoAbstract();
+
+	// Only discrete (mixer-compatible) formats. Excludes soundfield, ambisonics, composite.
+	UFUNCTION()
+	static TArray<FPropertyTextFName> GetCastingOptions_DiscreteOnly();
+
+	// Only the standard AudioMixer formats: Mono 1.0, Stereo 2.0, Quad 4.0, 5.1, 7.1.
+	// For MetaSound Source custom format dropdown — must match what AudioMixer can interleave.
+	UFUNCTION()
+	static TArray<FPropertyTextFName> GetCastingOptions_AudioMixerOnly();
 };
 
 UENUM()
diff --git a/Engine/Plugins/Runtime/Metasound/Source/MetasoundEditor/MetasoundEditor.Build.cs b/Engine/Plugins/Runtime/Metasound/Source/MetasoundEditor/MetasoundEditor.Build.cs
--- a/Engine/Plugins/Runtime/Metasound/Source/MetasoundEditor/MetasoundEditor.Build.cs
+++ b/Engine/Plugins/Runtime/Metasound/Source/MetasoundEditor/MetasoundEditor.Build.cs
@@ -28,6 +28,7 @@ namespace UnrealBuildTool.Rules
 					"MetasoundGenerator",
 					"MetasoundGraphCore",
 					"MetasoundStandardNodes",
+					"Projects",
 					"SignalProcessing",
 					"SubobjectEditor",
 					"ToolMenus",
diff --git a/Engine/Plugins/Runtime/Metasound/Source/MetasoundEditor/Private/MetasoundDetailCustomization.cpp b/Engine/Plugins/Runtime/Metasound/Source/MetasoundEditor/Private/MetasoundDetailCustomization.cpp
--- a/Engine/Plugins/Runtime/Metasound/Source/MetasoundEditor/Private/MetasoundDetailCustomization.cpp
+++ b/Engine/Plugins/Runtime/Metasound/Source/MetasoundEditor/Private/MetasoundDetailCustomization.cpp
@@ -33,6 +33,7 @@
 #include "PropertyRestriction.h"
 #include "SGraphPalette.h"
 #include "ScopedTransaction.h"
+#include "Interfaces/IPluginManager.h"
 #include "Sound/SoundWave.h"
 #include "Styling/CoreStyle.h"
 #include "Styling/SlateColor.h"
@@ -134,6 +135,12 @@ namespace Metasound::Editor
 		const FName RootClass = FName(GetRootClassPropertyPath());
 		return BuildChildPath(RootClass, GET_MEMBER_NAME_CHECKED(FMetasoundFrontendClass, Metadata));
 	}
+	
+	static bool IsExperimentalPluginEnabled()
+	{
+		const TSharedPtr<IPlugin> Plugin = IPluginManager::Get().FindEnabledPlugin(TEXT("MetasoundExperimental"));
+		return Plugin.IsValid();
+	}
 
 	void FMetasoundDetailCustomization::CustomizeDetails(IDetailLayoutBuilder& DetailLayout)
 	{
@@ -252,11 +259,114 @@ namespace Metasound::Editor
 						));
 					}
 				}
+				
 
 				// Invalid for UMetaSounds
 				TSharedPtr<IPropertyHandle> OutputFormat = DetailLayout.GetProperty(GET_MEMBER_NAME_CHECKED(UMetaSoundSource, OutputFormat));
 				if (OutputFormat.IsValid() && OutputFormat->IsValidHandle())
 				{
+					const UEnum* OutputFormatEnum = StaticEnum<EMetaSoundOutputAudioFormat>();
+					check(OutputFormatEnum);
+
+					const bool bAllowExperimental = IsExperimentalPluginEnabled();
+
+					// Build list of selectable values
+					OutputFormatOptions.Reset();
+					OutputFormatOptions.Reserve(OutputFormatEnum->NumEnums());
+
+					for (int32 Index = 0; Index < OutputFormatEnum->NumEnums(); ++Index)
+					{
+						const bool bIsHidden = OutputFormatEnum->HasMetaData(TEXT("Hidden"), Index);
+						const FName Name = OutputFormatEnum->GetNameByIndex(Index);
+						// UEnum::GetNameByIndex returns the C++ identifier (e.g. "COUNT", "MAX") — use case-insensitive
+					// comparison since UMETA names don't always match the case of the suffix check.
+					const FString NameStr = Name.ToString();
+					const bool bIsCount = NameStr.EndsWith(TEXT("Count"), ESearchCase::IgnoreCase) || NameStr.EndsWith(TEXT("MAX"), ESearchCase::IgnoreCase);
+						
+						if (bIsCount)
+						{
+							continue;
+						}
+
+						if (bIsHidden && !bAllowExperimental)
+						{
+							continue;
+						}
+
+						const int64 Value = OutputFormatEnum->GetValueByIndex(Index);
+						OutputFormatOptions.Add(MakeShared<int64>(Value));
+					}
+
+					auto GetCurrentValue = [OutputFormat]() -> int64
+					{
+						uint8 Raw = 0;
+						if (OutputFormat->GetValue(Raw) == FPropertyAccess::Success)
+						{
+							return (int64)Raw;
+						}
+						return 0;
+					};
+
+					TSharedPtr<int64> InitiallySelected;
+					{
+						const int64 Current = GetCurrentValue();
+						for (const TSharedPtr<int64>& Opt : OutputFormatOptions)
+						{
+							if (Opt.IsValid() && *Opt == Current)
+							{
+								InitiallySelected = Opt;
+								break;
+							}
+						}
+						if (!InitiallySelected.IsValid() && OutputFormatOptions.Num() > 0)
+						{
+							InitiallySelected = OutputFormatOptions[0];
+						}
+					}
+
+					TSharedRef<SWidget> OutputFormatValueWidget =
+						SNew(SComboBox<TSharedPtr<int64>>)
+						.OptionsSource(&OutputFormatOptions)
+						.InitiallySelectedItem(InitiallySelected)
+						.IsEnabled(TAttribute<bool>::Create(TAttribute<bool>::FGetter::CreateSP(this,
+							&FMetaSoundDetailCustomizationBase::IsFormatEditable)))
+						.OnGenerateWidget_Lambda([OutputFormatEnum](TSharedPtr<int64> InItem)
+							{
+								if (!InItem.IsValid())
+								{
+									return SNew(STextBlock).Text(FText::GetEmpty());
+								}
+
+								const int32 EnumIndex = OutputFormatEnum->GetIndexByValue(*InItem);
+								const FText Display = (EnumIndex != INDEX_NONE)
+									? OutputFormatEnum->GetDisplayNameTextByIndex(EnumIndex)
+									: FText::FromString(TEXT("<Invalid>"));
+
+								return SNew(STextBlock).Text(Display);
+							})
+						.OnSelectionChanged_Lambda([OutputFormat](TSharedPtr<int64> NewValue, ESelectInfo::Type)
+							{
+								if (NewValue.IsValid())
+								{
+									OutputFormat->SetValue((uint8)*NewValue);
+								}
+							})
+						[
+							SNew(STextBlock)
+							.Text_Lambda([OutputFormatEnum, OutputFormat]()
+								{
+									uint8 Raw = 0;
+									if (OutputFormat->GetValue(Raw) != FPropertyAccess::Success)
+									{
+										return FText::GetEmpty();
+									}
+
+									const int32 EnumIndex = OutputFormatEnum->GetIndexByValue((int64)Raw);
+									return (EnumIndex != INDEX_NONE)
+										? OutputFormatEnum->GetDisplayNameTextByIndex(EnumIndex)
+										: FText::FromString(TEXT("<Invalid>"));
+								})
+						];
 					if (MetaSoundSource.IsValid())
 					{
 						OutputFormat->SetOnPropertyValuePreChange(FSimpleDelegate::CreateLambda([Source = MetaSoundSource]()
@@ -284,9 +394,7 @@ namespace Metasound::Editor
 						}));
 					}
 
-					TSharedRef<SWidget> OutputFormatValueWidget = OutputFormat->CreatePropertyValueWidget();
-					OutputFormatValueWidget->SetEnabled(TAttribute<bool>::Create(TAttribute<bool>::FGetter::CreateSP(this, &FMetasoundDetailCustomization::IsEditable)));
-
+					OutputFormatValueWidget->SetEnabled(TAttribute<bool>::Create(TAttribute<bool>::FGetter::CreateSP(this, &FMetaSoundDetailCustomizationBase::IsEditable)));
 					static const FText OutputFormatName = LOCTEXT("MetasoundOutputFormatPropertyName", "Output Format");
 					GeneralCategoryBuilder.AddCustomRow(OutputFormatName)
 					.NameContent()
@@ -302,6 +410,115 @@ namespace Metasound::Editor
 					OutputFormat->MarkHiddenByCustomization();
 				}
 
+				TSharedPtr<IPropertyHandle> OutputFormatChooser = DetailLayout.GetProperty(GET_MEMBER_NAME_CHECKED(UMetaSoundSource, ChannelAgnosticFormatChooser));
+				if (OutputFormatChooser.IsValid() && OutputFormatChooser->IsValidHandle() && IsExperimentalPluginEnabled())
+				{
+					if (MetaSoundSource.IsValid())
+					{
+						OutputFormatChooser->SetOnPropertyValuePreChange(FSimpleDelegate::CreateLambda([Source = MetaSoundSource]()
+						{
+							if (Source.IsValid())
+							{
+								TSharedPtr<FEditor> ParentEditor = FGraphBuilder::GetEditorForMetasound(*Source.Get());
+								if (ParentEditor.IsValid())
+								{
+									ParentEditor->Stop();
+								};
+							}
+						}));
+
+						OutputFormatChooser->SetOnPropertyValueChanged(FSimpleDelegate::CreateLambda([Source = MetaSoundSource]()
+						{
+							if (Source.IsValid())
+							{
+								TSharedPtr<FEditor> ParentEditor = FGraphBuilder::GetEditorForMetasound(*Source.Get());
+								if (ParentEditor.IsValid())
+								{
+									ParentEditor->CreateAnalyzers(*Source.Get());
+								};
+							}
+						}));
+					}
+
+					TSharedRef<SWidget> OutputFormatValueWidget = OutputFormatChooser->CreatePropertyValueWidget();
+					OutputFormatValueWidget->SetEnabled(TAttribute<bool>::Create(TAttribute<bool>::FGetter::CreateSP(this, &FMetaSoundDetailCustomizationBase::IsFormatEditable)));
+
+					static const FText OutputFormatName = LOCTEXT("MetasoundOutputFormatChooserPropertyName", "Output Format Chooser");
+					GeneralCategoryBuilder.AddCustomRow(OutputFormatName)
+					.NameContent()
+					[
+						OutputFormatChooser->CreatePropertyNameWidget()
+					]
+					.ValueContent()
+					[
+						OutputFormatValueWidget
+					]
+					.RowTag("OutputFormatChooser");
+					
+					OutputFormatChooser->MarkHiddenByCustomization();
+				}
+
+				// Custom format name — visible only when ChannelAgnostic + Custom chooser.
+				TSharedPtr<IPropertyHandle> CustomFormatName = DetailLayout.GetProperty(GET_MEMBER_NAME_CHECKED(UMetaSoundSource, ChannelAgnosticCustomFormat));
+				if (CustomFormatName.IsValid() && CustomFormatName->IsValidHandle() && IsExperimentalPluginEnabled())
+				{
+					// Stop playback before format change and recreate analyzers after,
+					// matching the pattern used by OutputFormat and OutputFormatChooser.
+					if (MetaSoundSource.IsValid())
+					{
+						CustomFormatName->SetOnPropertyValuePreChange(FSimpleDelegate::CreateLambda([Source = MetaSoundSource]()
+						{
+							if (Source.IsValid())
+							{
+								TSharedPtr<FEditor> ParentEditor = FGraphBuilder::GetEditorForMetasound(*Source.Get());
+								if (ParentEditor.IsValid())
+								{
+									ParentEditor->Stop();
+								};
+							}
+						}));
+
+						CustomFormatName->SetOnPropertyValueChanged(FSimpleDelegate::CreateLambda([Source = MetaSoundSource]()
+						{
+							if (Source.IsValid())
+							{
+								TSharedPtr<FEditor> ParentEditor = FGraphBuilder::GetEditorForMetasound(*Source.Get());
+								if (ParentEditor.IsValid())
+								{
+									ParentEditor->CreateAnalyzers(*Source.Get());
+								};
+							}
+						}));
+					}
+
+					TSharedRef<SWidget> CustomFormatWidget = CustomFormatName->CreatePropertyValueWidget();
+					CustomFormatWidget->SetEnabled(TAttribute<bool>::Create(TAttribute<bool>::FGetter::CreateSP(this, &FMetaSoundDetailCustomizationBase::IsEditable)));
+
+					static const FText CustomFormatLabel = LOCTEXT("MetasoundCustomFormatPropertyName", "Custom Format Name");
+					GeneralCategoryBuilder.AddCustomRow(CustomFormatLabel)
+					.NameContent()
+					[
+						CustomFormatName->CreatePropertyNameWidget()
+					]
+					.ValueContent()
+					[
+						CustomFormatWidget
+					]
+					.Visibility(TAttribute<EVisibility>::CreateLambda([MetaSoundSource]()
+					{
+						if (MetaSoundSource.IsValid()
+							&& MetaSoundSource->OutputFormat == EMetaSoundOutputAudioFormat::ChannelAgnostic
+							&& MetaSoundSource->ChannelAgnosticFormatChooser == EMetasoundChannelAgnosticSourceFormatChooser::Custom)
+						{
+							return EVisibility::Visible;
+						}
+						return EVisibility::Collapsed;
+					}))
+					.RowTag("CustomFormatName");
+
+					CustomFormatName->MarkHiddenByCustomization();
+				}
+
 				// Updates FText properties on open editors if required
 				{
 					FSimpleDelegate RegisterOnChange = FSimpleDelegate::CreateLambda([this]()
diff --git a/Engine/Plugins/Runtime/Metasound/Source/MetasoundEditor/Private/MetasoundDetailCustomization.h b/Engine/Plugins/Runtime/Metasound/Source/MetasoundEditor/Private/MetasoundDetailCustomization.h
--- a/Engine/Plugins/Runtime/Metasound/Source/MetasoundEditor/Private/MetasoundDetailCustomization.h
+++ b/Engine/Plugins/Runtime/Metasound/Source/MetasoundEditor/Private/MetasoundDetailCustomization.h
@@ -30,6 +30,7 @@ namespace Metasound
 
 			UE_DEPRECATED(5.8, "Use virtual IsEditable instead")
 			bool IsGraphEditable() const;
+			bool IsFormatEditable() const { return IsEditable(); }
 
 			virtual bool IsEditable() const;
 
@@ -57,6 +58,8 @@ namespace Metasound
 			FName GetMetadataPropertyPath() const;
 
 			FName DocumentPropertyName;
+			TArray<TSharedPtr<int64>> OutputFormatOptions;
+			TSharedPtr<SComboBox<TSharedPtr<int64>>> OutputFormatCombo;
 		};
 
 		class FMetasoundPagesDetailCustomization : public FMetaSoundDetailCustomizationBase
diff --git a/Engine/Plugins/Runtime/Metasound/Source/MetasoundEngine/MetasoundEngine.Build.cs b/Engine/Plugins/Runtime/Metasound/Source/MetasoundEngine/MetasoundEngine.Build.cs
--- a/Engine/Plugins/Runtime/Metasound/Source/MetasoundEngine/MetasoundEngine.Build.cs
+++ b/Engine/Plugins/Runtime/Metasound/Source/MetasoundEngine/MetasoundEngine.Build.cs
@@ -27,7 +27,8 @@ namespace UnrealBuildTool.Rules
 					"MetasoundGraphCore",
 					"MetasoundGenerator",
 					"SignalProcessing",
-					"AudioChannelAgnosticCore"
+					"AudioChannelAgnosticCore",
+					"Projects"
 				}
 			);
 
diff --git a/Engine/Plugins/Runtime/Metasound/Source/MetasoundEngine/Private/Interfaces/MetasoundInputFormatInterfaces.cpp b/Engine/Plugins/Runtime/Metasound/Source/MetasoundEngine/Private/Interfaces/MetasoundInputFormatInterfaces.cpp
--- a/Engine/Plugins/Runtime/Metasound/Source/MetasoundEngine/Private/Interfaces/MetasoundInputFormatInterfaces.cpp
+++ b/Engine/Plugins/Runtime/Metasound/Source/MetasoundEngine/Private/Interfaces/MetasoundInputFormatInterfaces.cpp
@@ -5,6 +5,7 @@
 #include "Interfaces/MetasoundOutputFormatInterfaces.h"
 #include "Metasound.h"
 #include "MetasoundAudioBuffer.h"
+#include "MetasoundChannelAgnosticType.h"
 #include "MetasoundSource.h"
 #include "MetasoundTrigger.h"
 #include "Templates/SharedPointer.h"
@@ -798,5 +799,65 @@ namespace Metasound::Engine
 		}
 	} // namespace InputFormatSevenDotOneInterface
 #undef AUDIO_PARAMETER_INTERFACE_NAMESPACE
+	
+	
+#define AUDIO_PARAMETER_INTERFACE_NAMESPACE "UE.InputFormat.Cat"
+	namespace InputFormatChannelAgnosticInterface
+	{
+		const FMetasoundFrontendVersion& GetVersion()
+		{
+			static const FMetasoundFrontendVersion Version = { AUDIO_PARAMETER_INTERFACE_NAMESPACE, { 1, 0 } };
+			return Version;
+		}
+
+		namespace Inputs
+		{
+			const FName CatIn = AUDIO_PARAMETER_INTERFACE_MEMBER_DEFINE("Audio:0");
+		} // namespace Inputs
+
+		Audio::FParameterInterfacePtr CreateInterface()
+		{
+			struct FInterface : public Audio::FParameterInterface
+			{
+				FInterface()
+					: FParameterInterface(InputFormatChannelAgnosticInterface::GetVersion().Name, InputFormatChannelAgnosticInterface::GetVersion().Number.ToInterfaceVersion())
+				{
+					UClassOptions = InputFormatPrivate::GetUClassOptions();
+					Inputs =
+					{
+						{
+							LOCTEXT("InputChannelAgnosticInterfaceInputCatName", "Channel Agnostic"),
+							LOCTEXT("ChannelAgnostic_Cat_AudioDescription", "Channel Agnostic Input."),
+							GetMetasoundDataTypeName<FChannelAgnosticType>(),
+							Inputs::CatIn,
+							FText::GetEmpty(), // Required Text
+							100
+						}
+					};
+				}
+			};
+
+			return MakeShared<FInterface>();
+		}
+
+		TArray<FMetasoundFrontendInterfaceBinding> CreateBindings()
+		{
+			return TArray<FMetasoundFrontendInterfaceBinding>
+			{
+				{
+					OutputFormatChannelAgnosticInterface::GetVersion(),
+					GetVersion(),
+						10,
+					{
+						{
+							OutputFormatChannelAgnosticInterface::Outputs::CatOut, Inputs::CatIn
+						}
+					}
+				}
+			};
+		}
+	} // namespace InputFormatChannelAgnosticInterface
+#undef AUDIO_PARAMETER_INTERFACE_NAMESPACE
+
 } // namespace Metasound::Engine
 #undef LOCTEXT_NAMESPACE
diff --git a/Engine/Plugins/Runtime/Metasound/Source/MetasoundEngine/Private/Interfaces/MetasoundInterface.cpp b/Engine/Plugins/Runtime/Metasound/Source/MetasoundEngine/Private/Interfaces/MetasoundInterface.cpp
--- a/Engine/Plugins/Runtime/Metasound/Source/MetasoundEngine/Private/Interfaces/MetasoundInterface.cpp
+++ b/Engine/Plugins/Runtime/Metasound/Source/MetasoundEngine/Private/Interfaces/MetasoundInterface.cpp
@@ -6,6 +6,7 @@
 #include "AudioParameterControllerInterface.h"
 #include "IAudioParameterInterfaceRegistry.h"
 #include "IAudioParameterTransmitter.h"
+#include "Interfaces/IPluginManager.h"
 #include "Interfaces/MetasoundFrontendInterfaceRegistry.h"
 #include "Interfaces/MetasoundFrontendSourceInterface.h"
 #include "Interfaces/MetasoundInputFormatInterfaces.h"
@@ -164,6 +165,14 @@ namespace Metasound::Engine
 			Reg.RegisterInterface(MakeUnique<FInterfaceRegistryEntry>(OutputFormatFiveDotOneInterface::CreateInterface()));
 			Reg.RegisterInterface(MakeUnique<FInterfaceRegistryEntry>(OutputFormatSevenDotOneInterface::CreateInterface()));
 		}
+
+		// Channel Agnostic (CAT) format interfaces — only registered when MetasoundExperimental plugin is enabled.
+		// Without this gate, CAT output would appear as an option even without the experimental plugin loaded.
+		if (IPluginManager::Get().FindEnabledPlugin(TEXT("MetasoundExperimental")))
+		{
+			Reg.RegisterInterface(MakeUnique<FInterfaceRegistryEntry>(InputFormatChannelAgnosticInterface::CreateInterface()));
+			Reg.RegisterInterface(MakeUnique<FInterfaceRegistryEntry>(OutputFormatChannelAgnosticInterface::CreateInterface()));
+		}
 	}
 
 	void RegisterExternalInterfaces()
diff --git a/Engine/Plugins/Runtime/Metasound/Source/MetasoundEngine/Private/Interfaces/MetasoundOutputFormatInterfaces.cpp b/Engine/Plugins/Runtime/Metasound/Source/MetasoundEngine/Private/Interfaces/MetasoundOutputFormatInterfaces.cpp
--- a/Engine/Plugins/Runtime/Metasound/Source/MetasoundEngine/Private/Interfaces/MetasoundOutputFormatInterfaces.cpp
+++ b/Engine/Plugins/Runtime/Metasound/Source/MetasoundEngine/Private/Interfaces/MetasoundOutputFormatInterfaces.cpp
@@ -2,8 +2,10 @@
 #include "Interfaces/MetasoundOutputFormatInterfaces.h"
 
 #include "IAudioParameterInterfaceRegistry.h"
+#include "Interfaces/IPluginManager.h"
 #include "Metasound.h"
 #include "MetasoundAudioBuffer.h"
+#include "MetasoundChannelAgnosticType.h"
 #include "MetasoundSource.h"
 #include "MetasoundTrigger.h"
 #include "Templates/SharedPointer.h"
@@ -129,6 +131,20 @@ namespace Metasound::Engine
 				107
 			};
 		}
+		
+		Audio::FParameterInterface::FOutput GetChannelAgnosticOutput(const FName& InVertexName)
+		{
+			return
+			{
+				LOCTEXT("OutputFormatSurroundInterface_GeneratedChannelAgnosticDisplayName", "Out Channel Agnostic"),
+				LOCTEXT("OutputFormatSurroundInterface_GeneratedChannelAgnosticDescription", "The resulting channel agnostic output audio."),
+				GetMetasoundDataTypeName<FChannelAgnosticType>(),
+				InVertexName,
+				FText::GetEmpty(), // RequiredText
+				EAudioParameterType::None,
+				107
+			};
+		}	
 
 		TArray<Audio::FParameterInterface::FClassOptions> CreateUClassOptions(bool bIsSourceDefault = false)
 		{
@@ -404,14 +420,52 @@ namespace Metasound::Engine
 
 			return MakeShared<FInterface>();
 		}
-	} // namespace OutputFormatFiveDotOneInterface
+	} // namespace OutputFormatSevenDotOneInterface
 #undef AUDIO_PARAMETER_INTERFACE_NAMESPACE
 
+
+#define AUDIO_PARAMETER_INTERFACE_NAMESPACE "UE.OutputFormat.Cat"
+	namespace OutputFormatChannelAgnosticInterface
+	{
+		const FMetasoundFrontendVersion& GetVersion()
+		{
+			static const FMetasoundFrontendVersion Version = { AUDIO_PARAMETER_INTERFACE_NAMESPACE, { 1, 0 } };
+			return Version;
+		}
+
+		namespace Outputs
+		{
+			const FName CatOut = AUDIO_PARAMETER_INTERFACE_MEMBER_DEFINE("Audio:0");
+		}
+
+		Audio::FParameterInterfacePtr CreateInterface()
+		{
+			struct FInterface : public Audio::FParameterInterface
+			{
+				FInterface()
+					: FParameterInterface(OutputFormatChannelAgnosticInterface::GetVersion().Name, OutputFormatChannelAgnosticInterface::GetVersion().Number.ToInterfaceVersion())
+				{
+					using namespace OutputFormatPrivate;
+
+					UClassOptions = CreateUClassOptions();
+
+					Outputs =
+					{
+						GetChannelAgnosticOutput(Outputs::CatOut),
+					};
+				}
+			};
+
+			return MakeShared<FInterface>();
+		}
+	} // namespace OutputFormatChannelAgnosticInterface 
+#undef AUDIO_PARAMETER_INTERFACE_NAMESPACE
+	
 	const FOutputAudioFormatInfoMap& GetOutputAudioFormatInfo()
 	{
 		auto CreateFormatInfoMap = []()
 		{
-			return FOutputAudioFormatInfoMap
+			FOutputAudioFormatInfoMap Map
 			{
 				{
 					EMetaSoundOutputAudioFormat::Mono,
@@ -475,6 +529,20 @@ namespace Metasound::Engine
 					}
 				}
 			};
+
+			// Only include CAT output format when MetasoundExperimental plugin is enabled.
+			if (IPluginManager::Get().FindEnabledPlugin(TEXT("MetasoundExperimental")))
+			{
+				Map.Add(EMetaSoundOutputAudioFormat::ChannelAgnostic,
+				{
+					OutputFormatChannelAgnosticInterface::GetVersion(),
+					{
+						OutputFormatChannelAgnosticInterface::Outputs::CatOut,
+					}
+				});
+			}
+
+			return Map;
 		};
 
 		static const FOutputAudioFormatInfoMap Map = CreateFormatInfoMap();
diff --git a/Engine/Plugins/Runtime/Metasound/Source/MetasoundEngine/Private/MetasoundOperatorCacheSubsystem.cpp b/Engine/Plugins/Runtime/Metasound/Source/MetasoundEngine/Private/MetasoundOperatorCacheSubsystem.cpp
--- a/Engine/Plugins/Runtime/Metasound/Source/MetasoundEngine/Private/MetasoundOperatorCacheSubsystem.cpp
+++ b/Engine/Plugins/Runtime/Metasound/Source/MetasoundEngine/Private/MetasoundOperatorCacheSubsystem.cpp
@@ -2,6 +2,8 @@
 #include "MetasoundOperatorCacheSubsystem.h"
 
 #include "AudioMixerDevice.h"
+#include "Interfaces/MetasoundFrontendSourceInterface.h"
+#include "TypeFamily/ChannelTypeFamily.h"
 #include "MetasoundAssetManager.h"
 #include "MetasoundGenerator.h"
 #include "MetasoundGeneratorModule.h"
@@ -38,6 +40,65 @@ namespace Metasound::OperatorCachePrivate
 		FOperatorBuilderSettings BuilderSettings = FOperatorBuilderSettings::GetDefaultSettings();
 		// Graph analyzer currently only enabled for preview sounds (but can theoretically be supported for all sounds)
 		BuilderSettings.bPopulateInternalDataReferences = InParams.bIsPreviewSound;
+	
+		// Clamp helper — mirrors ClampToMixerChannelCount in MetasoundSource.cpp.
+		// AudioMixer only supports 1, 2, 4, 6, 8 channels.
+		auto ClampToMixer = [](int32 InCh) -> int32
+		{
+			if (InCh <= 1) return 1;
+			if (InCh <= 2) return 2;
+			if (InCh <= 4) return 4;
+			if (InCh <= 6) return 6;
+			return 8;
+		};
+
+		int32 NumChannels;
+		FName SourceFormatName;
+		if (InMetaSound.OutputFormat == EMetaSoundOutputAudioFormat::ChannelAgnostic)
+		{
+			// Three-branch format resolution, mirroring CreateSoundGenerator:
+			// 1) Custom: user-specified format name.
+			// 2) Auto: detect from wave proxy channel counts via CalculateAutoSourceFormat.
+			// 3) Fallback: use device channel count.
+			if (InMetaSound.ChannelAgnosticFormatChooser == EMetasoundChannelAgnosticSourceFormatChooser::Custom
+				&& !InMetaSound.ChannelAgnosticCustomFormat.IsNone())
+			{
+				PRAGMA_DISABLE_EXPERIMENTAL_WARNINGS
+				if (const Audio::FChannelTypeFamily* Found = Audio::GetChannelRegistry().FindChannel(InMetaSound.ChannelAgnosticCustomFormat))
+				{
+					SourceFormatName = Found->GetName();
+					NumChannels = ClampToMixer(Found->NumChannels());
+				}
+				PRAGMA_ENABLE_EXPERIMENTAL_WARNINGS
+				else
+				{
+					NumChannels = ClampToMixer(InParams.NumChannels);
+				}
+			}
+			// Auto: determine format from default wave proxy channel counts.
+			// Empty param array causes CalculateAutoSourceFormat to use RuntimeInput defaults,
+			// which is correct for precaching (no runtime parameter overrides available).
+			else if (TSharedPtr<const Audio::FChannelTypeFamily> SP = InMetaSound.CalculateAutoSourceFormat({}); SP.IsValid())
+			{
+				SourceFormatName = SP->GetName();
+				NumChannels = ClampToMixer(SP->NumChannels());
+			}
+			else
+			{
+				NumChannels = ClampToMixer(InParams.NumChannels);
+			}
+
+			// Embed SourceFormatName in the Environment so graph nodes can read it at build time.
+			// Mirrors the Environment.SetValue call in CreateSoundGenerator.
+			if (!SourceFormatName.IsNone())
+			{
+				Environment.SetValue<FName>(SourceInterface::Environment::SourceFormatName, SourceFormatName);
+			}
+		}
+		else
+		{
+			NumChannels = InMetaSound.GetOutputAudioChannelOrder().Num();
+		}
 
 		return FGeneratorInitParams
 		{
@@ -47,7 +108,9 @@ namespace Metasound::OperatorCachePrivate
 			.Environment = Environment,
 			.AudioOutputNames = InMetaSound.GetOutputAudioChannelOrder(),
 			.bBuildSynchronous = true,
-			.AssetPath = FTopLevelAssetPath(&InMetaSound)
+			.AssetPath = FTopLevelAssetPath(&InMetaSound),
+			.NumChannels = NumChannels,
+			.SourceFormatName = SourceFormatName,
 		};
 	}
 } // namespace Metasound::OperatorCachePrivate
diff --git a/Engine/Plugins/Runtime/Metasound/Source/MetasoundEngine/Private/MetasoundSource.cpp b/Engine/Plugins/Runtime/Metasound/Source/MetasoundEngine/Private/MetasoundSource.cpp
--- a/Engine/Plugins/Runtime/Metasound/Source/MetasoundEngine/Private/MetasoundSource.cpp
+++ b/Engine/Plugins/Runtime/Metasound/Source/MetasoundEngine/Private/MetasoundSource.cpp
@@ -7,7 +7,10 @@
 #include "Algo/Transform.h"
 #include "AssetRegistry/AssetRegistryModule.h"
 #include "AudioDeviceManager.h"
+#include "AudioMixerDevice.h"
+#include "AudioMixerBlueprintLibrary.h"
 #include "Containers/Ticker.h"
+#include "DocumentTemplates/MetasoundFrontendPresetTemplate.h"
 #include "IAudioParameterInterfaceRegistry.h"
 #include "Interfaces/ITargetPlatform.h"
 #include "Interfaces/MetasoundOutputFormatInterfaces.h"
@@ -44,7 +47,11 @@
 #include "MetasoundTrace.h"
 #include "MetasoundTrigger.h"
 #include "MetasoundUObjectRegistry.h"
-#include "DocumentTemplates/MetasoundFrontendPresetTemplate.h"
+#include "MetasoundWave.h"
+#include "MetasoundChannelAgnosticType.h"
+#include "ChannelAgnostic/ChannelAgnosticTypeUtils.h"
+#include "DSP/WavFormatUtils.h"
+#include "Sound/ISoundWaveContainer.h"
 #include "UObject/ObjectSaveContext.h"
 #include "UObject/ScriptInterface.h"
 
@@ -62,6 +69,18 @@
 
 #define LOCTEXT_NAMESPACE "MetaSound"
 
+// Clamp a CAT format's channel count to the nearest mixer-compatible value (1, 2, 4, 6, 8).
+// The AudioMixer only supports these counts for buffer allocation and interleaving.
+// The actual CAT format identity is carried separately via SourceFormatName for transcoding.
+static int32 ClampToMixerChannelCount(int32 InChannels)
+{
+	if (InChannels <= 1) return 1;
+	if (InChannels <= 2) return 2;
+	if (InChannels <= 4) return 4;
+	if (InChannels <= 6) return 6;
+	return 8;
+}
+
 namespace Metasound
 {
 	namespace ConsoleVariables
@@ -456,6 +475,11 @@ UMetaSoundSource::UMetaSoundSource(const FObjectInitializer& ObjectInitializer)
 	bProcedural = true;
 	bRequiresStopFade = true;
 	NumChannels = 1;
+
+	// Note: Can't call FDiscreteChannelAgnosticType::GetDefaultCatFormat() here as the CDO
+	// constructor runs during module load before the channel registry is initialized.
+	static const FName DefaultDiscreteFormat(TEXT("Cat:Mono1Dot0"));
+	ChannelAgnosticCustomFormat = DefaultDiscreteFormat;
 }
 
 const UClass& UMetaSoundSource::GetBaseMetaSoundUClass() const
@@ -507,6 +531,25 @@ void UMetaSoundSource::PostEditChangeProperty(FPropertyChangedEvent& InEvent)
 	{
 		PostEditChangeOutputFormat();
 	}
+	// When the CAT custom format or format chooser changes, resolve NumChannels immediately
+	// so the mixer allocates the correct buffer size on the next play. (UE-370811)
+	// NumChannels must be mixer-compatible (1, 2, 4, 6, 8) — the CAT format identity is
+	// carried separately via SourceFormatName, and transcoding handles the conversion.
+	if (InEvent.GetPropertyName() == GET_MEMBER_NAME_CHECKED(UMetaSoundSource, ChannelAgnosticCustomFormat) ||
+		InEvent.GetPropertyName() == GET_MEMBER_NAME_CHECKED(UMetaSoundSource, ChannelAgnosticFormatChooser))
+	{
+		if (OutputFormat == EMetaSoundOutputAudioFormat::ChannelAgnostic &&
+			ChannelAgnosticFormatChooser == EMetasoundChannelAgnosticSourceFormatChooser::Custom &&
+			!ChannelAgnosticCustomFormat.IsNone())
+		{
+			PRAGMA_DISABLE_EXPERIMENTAL_WARNINGS
+			if (const Audio::FChannelTypeFamily* Found = Audio::GetChannelRegistry().FindChannel(ChannelAgnosticCustomFormat))
+			{
+				NumChannels = ClampToMixerChannelCount(Found->NumChannels());
+			}
+			PRAGMA_ENABLE_EXPERIMENTAL_WARNINGS
+		}
+	}
 	if (InEvent.GetMemberPropertyName() == GET_MEMBER_NAME_CHECKED(UMetaSoundSource, SampleRateOverride) ||
 		InEvent.GetMemberPropertyName() == GET_MEMBER_NAME_CHECKED(UMetaSoundSource, BlockRateOverride) ||
 		InEvent.GetPropertyName() == GET_MEMBER_NAME_CHECKED(UMetaSoundSource, QualitySetting) )
@@ -596,6 +639,7 @@ void UMetaSoundSource::PostEditChangeQualitySettings()
 }
 #endif // WITH_EDITOR
 
+
 bool UMetaSoundSource::ConformObjectToDocument()
 {
 	using namespace Metasound::Engine;
@@ -610,7 +654,33 @@ bool UMetaSoundSource::ConformObjectToDocument()
 	{
 		if (RootMetasoundDocument.Interfaces.Contains(Pair.Value.InterfaceVersion))
 		{
-			if ((OutputFormat != Pair.Key) || (NumChannels != Pair.Value.OutputVertexChannelOrder.Num()))
+			// Channel Agnostic (CAT) output: resolve NumChannels from the custom format
+			// if one is selected, otherwise default to 1 (mono). NumChannels here is the
+			// mixer-facing channel count, clamped to supported values (1,2,4,6,8).
+			if (Pair.Key == EMetaSoundOutputAudioFormat::ChannelAgnostic)
+			{
+				// Resolve NumChannels now if a custom format is selected and the registry is available.
+				// Previously this was hardcoded to 1 as a placeholder, causing a two-save sync bug (UE-370811).
+				int32 ResolvedChannels = 1;
+				if (ChannelAgnosticFormatChooser == EMetasoundChannelAgnosticSourceFormatChooser::Custom && !ChannelAgnosticCustomFormat.IsNone())
+				{
+					PRAGMA_DISABLE_EXPERIMENTAL_WARNINGS
+					if (const Audio::FChannelTypeFamily* Found = Audio::GetChannelRegistry().FindChannel(ChannelAgnosticCustomFormat))
+					{
+						ResolvedChannels = ClampToMixerChannelCount(Found->NumChannels());
+					}
+					PRAGMA_ENABLE_EXPERIMENTAL_WARNINGS
+				}
+
+				// Only mark altered if something actually changed (avoids unnecessary asset dirtying).
+				if (OutputFormat != Pair.Key || NumChannels != ResolvedChannels)
+				{
+					OutputFormat = Pair.Key;
+					NumChannels = ResolvedChannels;
+					bDidAlterObjectData = true;
+				}
+			}
+			else if ((OutputFormat != Pair.Key) || (NumChannels != Pair.Value.OutputVertexChannelOrder.Num()))
 			{
 				OutputFormat = Pair.Key;
 				NumChannels = Pair.Value.OutputVertexChannelOrder.Num();
@@ -740,6 +810,79 @@ void UMetaSoundSource::SerializeCookedQualitySettings(const FName PlatformName,
 	}
 }
 
+static TSharedPtr<const Audio::FDiscreteChannelTypeFamily> NumChannelsToFormat(const int32 NumChannels)
+{
+	return Audio::FChannelAgnosticUtils::FindDiscreteFormatFromNumChannels(NumChannels);
+}
+
+TSharedPtr<const Audio::FChannelTypeFamily> UMetaSoundSource::CalculateAutoSourceFormat(const TArray<FAudioParameter>& InParams) const
+{
+	using namespace Audio;
+
+	// This only applies if we're set to Channel Agnostic Output.
+	if (OutputFormat == EMetaSoundOutputAudioFormat::ChannelAgnostic)
+	{
+		TMap<FName, TArray<FSoundWaveProxyPtr>> Proxies;
+		auto CollectProxies = [&Proxies](const FName Key, const TArray<TSharedPtr<IProxyData>>& ObjectProxies)
+		{
+			for (const TSharedPtr<IProxyData>& ProxyData : ObjectProxies)
+			{
+				if (const ISoundWaveContainer* Container = QueryInterface<ISoundWaveContainer>(ProxyData.Get()))
+				{
+					Proxies.FindOrAdd(Key).Append(Container->GetContainedWaveProxies());
+				}
+			}
+		};
+
+		// Gather all referenced sound waves.
+		// Use the greatest channel count as the format.
+		if (RuntimeInputData.bIsValid.load())
+		{
+			TArray<FRuntimeInput> Values;
+			RuntimeInputData.InputMap.GenerateValueArray(Values);
+
+			for (const FRuntimeInput& RuntimeInput : Values)
+			{
+				if (const FAudioParameter* Found = InParams.FindByPredicate([&RuntimeInput](const FAudioParameter& Param) -> bool { return Param.ParamName == RuntimeInput.Name; }))
+				{
+					CollectProxies(RuntimeInput.Name, Found->ObjectProxies);
+				}
+				else
+				{
+					CollectProxies(RuntimeInput.Name, RuntimeInput.DefaultParameter.ObjectProxies);
+				}
+			}
+		}
+		else
+		{
+			for (const FAudioParameter& Param : InParams)
+			{
+				CollectProxies(Param.ParamName, Param.ObjectProxies);
+			}
+		}
+
+		// Find Greatest channel count.
+		int32 GreatestNumChannels = 1;
+		for (const TTuple<FName, TArray<TSharedPtr<FSoundWaveProxy>>>& ProxyEntry : Proxies)
+		{
+			if (const TSharedPtr<FSoundWaveProxy>* Greatest = Algo::MaxElementBy(ProxyEntry.Value, [](const TSharedPtr<FSoundWaveProxy>& WaveProxy) -> int32
+				{
+					return WaveProxy->GetNumChannels();
+				}))
+			{
+				GreatestNumChannels = FMath::Max<int32>(GreatestNumChannels, (*Greatest)->GetNumChannels());
+			}
+		}
+		
+		// Convert to Format Name.
+		if (const TSharedPtr<const Audio::FDiscreteChannelTypeFamily> SP = NumChannelsToFormat(GreatestNumChannels))
+		{
+			return SP;	
+		}
+	}
+	return nullptr;
+}
+
 #if WITH_EDITOR
 void UMetaSoundSource::SetReferencedAssets(TSet<Metasound::Frontend::IMetaSoundAssetManager::FAssetRef>&& InAssetRefs)
 {
@@ -1066,6 +1209,51 @@ ISoundGeneratorPtr UMetaSoundSource::CreateSoundGenerator(const FSoundGeneratorI
 	SampleRate = InSettings.GetSampleRate();
 
 	FMetasoundEnvironment Environment = CreateEnvironment(InParams);
+	
+	// If the output is Channel Agnostic, determine the format and channel count.
+	// Priority: 1) user-specified custom format, 2) auto-detect from wave proxy assets,
+	// 3) fall back to the audio mixer's channel count.
+	// SourceFormatName carries format identity through to the generator so it can set up
+	// transcoding without a lossy reverse lookup from NumChannels.
+	FName SourceFormatName;
+	if (OutputFormat == EMetaSoundOutputAudioFormat::ChannelAgnostic)
+	{
+		// Custom: user has specified a format name directly.
+		if (ChannelAgnosticFormatChooser == EMetasoundChannelAgnosticSourceFormatChooser::Custom && !ChannelAgnosticCustomFormat.IsNone())
+		{
+			PRAGMA_DISABLE_EXPERIMENTAL_WARNINGS
+			if (const Audio::FChannelTypeFamily* Found = Audio::GetChannelRegistry().FindChannel(ChannelAgnosticCustomFormat))
+			{
+				SourceFormatName = Found->GetName();
+				Environment.SetValue<FName>(SourceInterface::Environment::SourceFormatName, SourceFormatName);
+				NumChannels = ClampToMixerChannelCount(Found->NumChannels());
+			}
+			PRAGMA_ENABLE_EXPERIMENTAL_WARNINGS
+			else
+			{
+				UE_LOG(LogMetaSound, Warning, TEXT("Custom format '%s' not found in channel registry"), *ChannelAgnosticCustomFormat.ToString());
+				NumChannels = ClampToMixerChannelCount(InParams.NumChannels);
+			}
+		}
+		// Auto: determine format from wave proxy channel counts.
+		else if (TSharedPtr<const Audio::FChannelTypeFamily> SP = CalculateAutoSourceFormat(InDefaultParameters); SP.IsValid())
+		{
+			SourceFormatName = SP->GetName();
+			Environment.SetValue<FName>(SourceInterface::Environment::SourceFormatName, SourceFormatName);
+			NumChannels = ClampToMixerChannelCount(SP->NumChannels());
+		}
+		else
+		{
+			NumChannels = ClampToMixerChannelCount(InParams.NumChannels); // Clamp mixer channel count to supported values (1,2,4,6,8).
+		}
+	}
+	else
+	{
+		// Otherwise the number of AudioOutput is our NumChannels.
+		NumChannels = GetOutputAudioChannelOrder().Num(); 
+	}
+
+	
 	FParameterRouter& Router = UMetaSoundSource::GetParameterRouter();
 	TSharedPtr<TSpscQueue<FMetaSoundParameterTransmitter::FParameter>> DataChannel = Router.FindOrCreateDataChannelForReader(InParams.AudioDeviceID, InParams.InstanceID);
 
@@ -1123,7 +1311,9 @@ ISoundGeneratorPtr UMetaSoundSource::CreateSoundGenerator(const FSoundGeneratorI
 				.DataChannel = DataChannel,
 				.GraphRenderCost = GraphRenderCost,
 				.ClassName = GetClass()->GetFName(),
-				.AssetPath = GetAssetPathChecked()
+				.AssetPath = GetAssetPathChecked(),
+				.NumChannels =  NumChannels,
+				.SourceFormatName = SourceFormatName
 			},
 			DynamicTransactor->CreateTransformQueue(InSettings, Environment, GraphRenderCost) // Create transaction queue
 		};
@@ -1162,7 +1352,9 @@ ISoundGeneratorPtr UMetaSoundSource::CreateSoundGenerator(const FSoundGeneratorI
 				.DefaultParameters = MoveTemp(InDefaultParameters),
 				.bBuildSynchronous = bBuildSynchronous,
 				.DataChannel = DataChannel,
-				.ClassName = GetClass()->GetFName()
+				.ClassName = GetClass()->GetFName(),
+				.NumChannels = NumChannels,
+				.SourceFormatName = SourceFormatName
 			};
 			InitParams.AssetPath = GetAssetPathChecked();
 			Generator = MakeShared<FMetasoundConstGraphGenerator>(MoveTemp(InitParams));
@@ -1643,6 +1835,7 @@ Metasound::FMetasoundEnvironment UMetaSoundSource::CreateEnvironment(const FSoun
 	Environment.SetValue<Audio::FDeviceId>(SourceInterface::Environment::DeviceID, InParams.AudioDeviceID);
 	Environment.SetValue<int32>(SourceInterface::Environment::AudioMixerNumOutputFrames, InParams.AudioMixerNumOutputFrames);
 	Environment.SetValue<float>(SourceInterface::Environment::AudioMixerSampleRate, InParams.SampleRate);
+	Environment.SetValue<int32>(SourceInterface::Environment::AudioMixerNumChannels, InParams.NumChannels);
 
 #if WITH_METASOUND_DEBUG_ENVIRONMENT
 	Environment.SetValue<FString>(SourceInterface::Environment::GraphName, GetFullName());
diff --git a/Engine/Plugins/Runtime/Metasound/Source/MetasoundEngine/Public/Interfaces/MetasoundInputFormatInterfaces.h b/Engine/Plugins/Runtime/Metasound/Source/MetasoundEngine/Public/Interfaces/MetasoundInputFormatInterfaces.h
--- a/Engine/Plugins/Runtime/Metasound/Source/MetasoundEngine/Public/Interfaces/MetasoundInputFormatInterfaces.h
+++ b/Engine/Plugins/Runtime/Metasound/Source/MetasoundEngine/Public/Interfaces/MetasoundInputFormatInterfaces.h
@@ -83,4 +83,17 @@ namespace Metasound::Engine
 		METASOUNDENGINE_API Audio::FParameterInterfacePtr CreateInterface();
 		METASOUNDENGINE_API TArray<FMetasoundFrontendInterfaceBinding> CreateBindings();
 	} // namespace InputFormatSevenDotOneInterface
+	
+	namespace InputFormatChannelAgnosticInterface
+	{
+		namespace Inputs
+		{
+			METASOUNDENGINE_API const extern FName CatIn;
+		} // namespace Inputs
+
+		METASOUNDENGINE_API const FMetasoundFrontendVersion& GetVersion();
+		METASOUNDENGINE_API Audio::FParameterInterfacePtr CreateInterface();
+		METASOUNDENGINE_API TArray<FMetasoundFrontendInterfaceBinding> CreateBindings();
+	} // namespace InputFormatChannelAgnosticInterface 
+	
 } // namespace Metasound::Engine
diff --git a/Engine/Plugins/Runtime/Metasound/Source/MetasoundEngine/Public/Interfaces/MetasoundOutputFormatInterfaces.h b/Engine/Plugins/Runtime/Metasound/Source/MetasoundEngine/Public/Interfaces/MetasoundOutputFormatInterfaces.h
--- a/Engine/Plugins/Runtime/Metasound/Source/MetasoundEngine/Public/Interfaces/MetasoundOutputFormatInterfaces.h
+++ b/Engine/Plugins/Runtime/Metasound/Source/MetasoundEngine/Public/Interfaces/MetasoundOutputFormatInterfaces.h
@@ -20,11 +20,10 @@ enum class EMetaSoundOutputAudioFormat : uint8
 	Quad,
 	FiveDotOne UMETA(DisplayName = "5.1"),
 	SevenDotOne UMETA(DisplayName = "7.1"),
-
+	ChannelAgnostic  UMETA(DisplayName = "Channel Agnostic (Experimental)", Hidden),
 	COUNT UMETA(Hidden)
 };
 
-
 namespace Metasound::Engine
 {
 	namespace OutputFormatMonoInterface
@@ -113,6 +112,16 @@ namespace Metasound::Engine
 		METASOUNDENGINE_API Audio::FParameterInterfacePtr CreateInterface();
 	} // namespace OutputFormatSevenDotOneInterface
 
+	namespace OutputFormatChannelAgnosticInterface
+	{
+		namespace Outputs
+		{
+			METASOUNDENGINE_API const extern FName CatOut;
+		} // namespace Outputs
+
+		METASOUNDENGINE_API const FMetasoundFrontendVersion& GetVersion();
+		METASOUNDENGINE_API Audio::FParameterInterfacePtr CreateInterface();
+	} // namespace OutputFormatChannelAgnosticInterface
 	// Contains information on output audio formats
 	struct FOutputAudioFormatInfo
 	{
diff --git a/Engine/Plugins/Runtime/Metasound/Source/MetasoundEngine/Public/MetasoundSource.h b/Engine/Plugins/Runtime/Metasound/Source/MetasoundEngine/Public/MetasoundSource.h
--- a/Engine/Plugins/Runtime/Metasound/Source/MetasoundEngine/Public/MetasoundSource.h
+++ b/Engine/Plugins/Runtime/Metasound/Source/MetasoundEngine/Public/MetasoundSource.h
@@ -7,6 +7,7 @@
 #include "Interfaces/MetasoundOutputFormatInterfaces.h"
 #include "Metasound.h"
 #include "MetasoundAssetBase.h"
+#include "MetasoundChannelAgnosticAutoFormat.h"
 #include "MetasoundFrontend.h"
 #include "MetasoundFrontendDocument.h"
 #include "MetasoundFrontendRegistries.h"
@@ -152,6 +153,14 @@ public:
 	UPROPERTY(EditAnywhere, BlueprintReadOnly, Category = Metasound)
 	EMetaSoundOutputAudioFormat OutputFormat;
 
+	UPROPERTY(EditAnywhere, BlueprintReadOnly, Category = Metasound, meta = (DisplayName = "Source Format (exp)", EditConditionHides, EditCondition = "OutputFormat == EMetaSoundOutputAudioFormat::ChannelAgnostic"))
+	EMetasoundChannelAgnosticSourceFormatChooser ChannelAgnosticFormatChooser = EMetasoundChannelAgnosticSourceFormatChooser::Auto;
+
+	// The custom channel-agnostic format name (e.g. "Cat:Discrete:Stereo"). Only used when ChannelAgnosticFormatChooser == Custom.
+	// Restricted to discrete (mixer-compatible) formats only — soundfield/ambisonics are not supported by AudioMixer.
+	UPROPERTY(EditAnywhere, BlueprintReadOnly, Category = Metasound, meta = (DisplayName = "Custom Format Name", GetOptions="MetasoundCatCastingOptionsHelper.GetCastingOptions_AudioMixerOnly", EditConditionHides, EditCondition = "OutputFormat == EMetaSoundOutputAudioFormat::ChannelAgnostic && ChannelAgnosticFormatChooser == EMetasoundChannelAgnosticSourceFormatChooser::Custom"))
+	FName ChannelAgnosticCustomFormat;
+
 #if WITH_EDITORONLY_DATA
 	// The QualitySetting MetaSound will use, as defined in 'MetaSound' Settings.
 	UPROPERTY(EditAnywhere, BlueprintReadWrite, AssetRegistrySearchable, meta = (GetOptions="MetasoundEngine.MetaSoundSettings.GetQualityNames"), Category = "Metasound")
@@ -370,6 +379,15 @@ public:
 	UE_API Metasound::FMetasoundEnvironment CreateEnvironment(const FSoundGeneratorInitParams& InParams) const;
 	UE_API const TArray<Metasound::FVertexName>& GetOutputAudioChannelOrder() const;
 
+	/**
+	 * Determines the "Auto" Channel Agnostic Format for this Source.
+	 * Consults the document layer and makes a best guess at the appropriate format
+	 * from wave proxy channel counts. Can be called before graph evaluation.
+	 * @param InParams - Audio parameters to inspect for wave proxies. Pass empty array to use defaults.
+	 * @return Shared pointer to the resolved channel type family, or nullptr if no format could be determined.
+	 */
+	UE_API TSharedPtr<const Audio::FChannelTypeFamily> CalculateAutoSourceFormat(const TArray<FAudioParameter>& InParams) const;
+
 	/** Find the Source related to this Preset.
 	 * 
 	 * If this MetaSound is a preset and preset graph inflation is enabled, this
@@ -436,6 +454,7 @@ private:
 	UE_API void ResolveQualitySettings(const UMetaSoundSettings* Settings);
 	UE_API void SerializeCookedQualitySettings(const FName PlatformName, FArchive& Ar);
 	TPimplPtr<Metasound::SourcePrivate::FCookedQualitySettings> CookedQualitySettings;
+	
 };
 
 #undef UE_API
diff --git a/Engine/Plugins/Runtime/Metasound/Source/MetasoundGenerator/MetasoundGenerator.Build.cs b/Engine/Plugins/Runtime/Metasound/Source/MetasoundGenerator/MetasoundGenerator.Build.cs
--- a/Engine/Plugins/Runtime/Metasound/Source/MetasoundGenerator/MetasoundGenerator.Build.cs
+++ b/Engine/Plugins/Runtime/Metasound/Source/MetasoundGenerator/MetasoundGenerator.Build.cs
@@ -14,7 +14,8 @@ namespace UnrealBuildTool.Rules
 					"SignalProcessing",
 					"Engine",
 					"AudioExtensions",
-					"TraceLog"
+					"TraceLog",
+					"AudioChannelAgnosticCore"
 				}
 			);
 
diff --git a/Engine/Plugins/Runtime/Metasound/Source/MetasoundGenerator/Private/MetasoundGenerator.cpp b/Engine/Plugins/Runtime/Metasound/Source/MetasoundGenerator/Private/MetasoundGenerator.cpp
--- a/Engine/Plugins/Runtime/Metasound/Source/MetasoundGenerator/Private/MetasoundGenerator.cpp
+++ b/Engine/Plugins/Runtime/Metasound/Source/MetasoundGenerator/Private/MetasoundGenerator.cpp
@@ -23,6 +23,9 @@
 #include "Interfaces/MetasoundFrontendSourceInterface.h"
 #include "Modules/ModuleManager.h"
 #include "AutoRTFM.h"
+#include "ChannelAgnostic/ChannelAgnosticTranscoding.h"
+#include "ChannelAgnostic/ChannelAgnosticTypeUtils.h"
+#include "DSP/MultiMono.h"
 #include "ProfilingDebugging/AssetMetadataTrace.h"
 
 #ifndef ENABLE_METASOUNDGENERATOR_INVALID_SAMPLE_VALUE_LOGGING 
@@ -260,6 +263,8 @@ namespace Metasound
 		InParams.DefaultParameters = {};
 		InParams.DataChannel.Reset();
 		InParams.GraphRenderCost.Reset();
+		InParams.NumChannels = 0;
+		InParams.SourceFormatName = NAME_None;
 	}
 
 	void FMetasoundDynamicGraphGeneratorInitParams::Reset(FMetasoundDynamicGraphGeneratorInitParams& InParams)
@@ -302,7 +307,8 @@ namespace Metasound
 	void FMetasoundGenerator::InitBase(FGeneratorInitParams& InInitParams)
 	{
 		AssetPath = InInitParams.AssetPath;
-		NumChannels = InInitParams.AudioOutputNames.Num();
+		NumChannels = InInitParams.NumChannels;
+		SourceFormatName = InInitParams.SourceFormatName;
 		NumSamplesPerExecute = NumChannels * NumFramesPerExecute;
 
 #if ENABLE_METASOUND_GENERATOR_INSTANCE_COUNTING
@@ -372,7 +378,76 @@ PRAGMA_ENABLE_DEPRECATION_WARNINGS
 		VertexInterfaceData = MoveTemp(InData->VertexInterfaceData);
 
 		GraphOutputAudio.Reset();
-		if (InData->OutputBuffers.Num() == NumChannels)
+		GraphCatOutput.Reset();
+
+		if (InData->OutputChannelAgnosticRef.IsSet())
+		{
+			// The MixerOutput must be a format AudioMixer supports (1/2/4/6/8 channels).
+			// SourceFormatName carries the CAT format identity for the graph side (via Environment),
+			// but the mixer-facing output is always determined from NumChannels (already clamped to
+			// mixer-compatible values). This is intentionally lossy — the transcoder handles the
+			// conversion from whatever the graph produces to the mixer's discrete format.
+			//
+			// IMPORTANT: Do NOT early-return from this block. SetGraph must complete operator
+			// installation below regardless of CAT setup success. Early returns after buffer
+			// resets and MoveTemp leave the generator in a stranded partial-init state with
+			// no retry path (the pending graph has already been consumed).
+			const TSharedPtr<const Audio::FDiscreteChannelTypeFamily> MixerFormat =
+				Audio::FChannelAgnosticUtils::FindDiscreteFormatFromNumChannels(NumChannels);
+			const Audio::FChannelTypeFamily* MixerChannel = nullptr;
+			if (MixerFormat.IsValid())
+			{
+				MixerChannel = Audio::GetChannelRegistry().FindConcreteChannel(MixerFormat->GetName());
+			}
+
+			if (ensureMsgf(MixerChannel, TEXT("Failed to resolve mixer channel format for %d channels. CAT output will be silent."), NumChannels))
+			{
+				// Check if the graph output format already matches the mixer format.
+				// If so, skip the intermediate buffer and transcoder entirely — just interleave
+				// directly from the graph output. This avoids allocating a redundant MixerOutput
+				// buffer and a no-op memcpy transcode on every block. (Addresses Phil Popp's
+				// memory bloat concern: most CAT sources output a mixer-compatible format.)
+				const Audio::FChannelTypeFamily& GraphOutputType = InData->OutputChannelAgnosticRef->Get()->GetType();
+				const bool bIdentity = (&GraphOutputType == MixerChannel) || (GraphOutputType.GetName() == MixerChannel->GetName());
+
+				if (bIdentity)
+				{
+					// Identity: graph output IS the mixer format. No transcode, no extra buffer.
+					GraphCatOutput.Emplace(FChannelAgnosticState{
+						*InData->OutputChannelAgnosticRef,	// GraphOutput
+						{},					// MixerOutput — not allocated
+						{},					// Transcoder — not needed
+						true					// bIdentityTranscode
+					});
+				}
+				else
+				{
+					// Non-identity: allocate mixer buffer and resolve transcoder.
+					Audio::FChannelAgnosticType MixOut(*MixerChannel, NumFramesPerExecute);
+					Audio::ChannelAgnosticTranscoder::FTranscoder Transcoder = Audio::FTranscoderResolver::Resolve(
+						GraphOutputType,
+						{
+							.ToType = MixOut.GetType(),
+							.TranscodeMethod = Audio::EChannelTranscodeMethod::MixUpOrDown,
+							.MixMethod = Audio::EChannelMapMonoUpmixMethod::EqualPower,
+						});
+					if (ensureMsgf(Transcoder.IsSet(), TEXT("Failed to resolve CAT transcoder from '%s' to '%s'. CAT output will be silent."),
+						*GraphOutputType.GetName().ToString(), *MixerChannel->GetName().ToString()))
+					{
+						GraphCatOutput.Emplace(FChannelAgnosticState{
+							*InData->OutputChannelAgnosticRef,	// GraphOutput
+							MoveTemp(MixOut),			// MixerOutput
+							MoveTemp(Transcoder),			// Transcoder
+							false					// bIdentityTranscode
+						});
+					}
+					// On transcoder failure: GraphCatOutput stays empty → InterleaveGeneratedAudio
+					// skips CAT path → silence for this source, but generator completes init normally.
+				}
+			}
+			// On MixerChannel failure: GraphCatOutput stays empty → silence, but generator init completes.
+		}
+		else if (InData->OutputBuffers.Num() == NumChannels)
 		{
 			if (InData->OutputBuffers.Num() > 0)
 			{
@@ -405,9 +480,13 @@ PRAGMA_ENABLE_DEPRECATION_WARNINGS
 		RootExecuter.SetOperator(MoveTemp(InData->GraphOperator));
 
 
-		// Query the graph output to get the number of output audio channels.
-		// Multichannel version:
-		check(NumChannels == GraphOutputAudio.Num());
+		// Validate that either discrete output buffers or CAT output is configured.
+		// If CAT setup failed (ensureMsgf above), both may be empty — that's a valid
+		// degraded state (silence), not a crash. Only assert if we have neither AND
+		// NumChannels expects output (which would cause OOB in InterleaveGeneratedAudio).
+		ensureMsgf(NumChannels == GraphOutputAudio.Num() || GraphCatOutput.IsSet() || NumChannels == 0,
+			TEXT("MetaSound generator has %d expected channels but no audio output configured (GraphOutputAudio=%d, CatOutput=%s)"),
+			NumChannels, GraphOutputAudio.Num(), GraphCatOutput.IsSet() ? TEXT("set") : TEXT("unset"));
 
 		if (NumSamplesPerExecute > 0)
 		{
@@ -468,6 +547,7 @@ PRAGMA_ENABLE_DEPRECATION_WARNINGS
 		RootExecuter.ReleaseOperator();
 		VertexInterfaceData = FVertexInterfaceData();
 		GraphOutputAudio.Reset();
+		GraphCatOutput.Reset();
 		OnFinishedTriggerRef = TDataWriteReference<FTrigger>::CreateNew(OperatorSettings);
 		GraphAnalyzer.Reset();
 		OutputAnalyzers.Reset();
@@ -797,19 +877,60 @@ PRAGMA_ENABLE_DEPRECATION_WARNINGS
 		MetasoundGeneratorPrivate::LogInvalidAudioSampleValues(AssetPath, GraphOutputAudio);
 #endif // if ENABLE_METASOUNDGENERATOR_INVALID_SAMPLE_VALUE_LOGGING
 
-		// Iterate over channels
-		for (int32 ChannelIndex = 0; ChannelIndex < NumChannels; ChannelIndex++)
+		if (GraphCatOutput.IsSet())
 		{
-			const FAudioBuffer& InputBuffer = *GraphOutputAudio[ChannelIndex];
+			// Pointer to the multi-mono data we'll interleave from.
+			const Audio::FChannelAgnosticType* SourceData = nullptr;
 
-			const float* InputPtr = InputBuffer.GetData();
-			float* OutputPtr = &InterleavedAudioBuffer.GetData()[ChannelIndex];
+			if (GraphCatOutput->bIdentityTranscode)
+			{
+				// Fast path: graph output format matches mixer format — no transcode needed,
+				// no intermediate buffer allocated. Read directly from graph output.
+				SourceData = GraphCatOutput->GraphOutput.operator->();
+			}
+			else
+			{
+				// Transcode path: graph output format differs from mixer format.
+				// Convert into the pre-allocated MixerOutput buffer, then interleave from that.
+				check(GraphCatOutput->MixerOutput.IsSet());
+				Audio::TStackArrayOfPointers<const float> Src = Audio::MakeMultiMonoPointersFromView(GraphCatOutput->GraphOutput->GetRawMultiMono(), NumFramesPerExecute, GraphCatOutput->GraphOutput->NumChannels());
+				Audio::TStackArrayOfPointers<float> Dst = Audio::MakeMultiMonoPointersFromView(GraphCatOutput->MixerOutput->GetRawMultiMono(), NumFramesPerExecute, GraphCatOutput->MixerOutput->NumChannels());
+				GraphCatOutput->Transcoder(Src, Dst, NumFramesPerExecute);
+				SourceData = &GraphCatOutput->MixerOutput.GetValue();
+			}
 
-			// Assign values to output for single channel.
-			for (int32 FrameIndex = 0; FrameIndex < NumFramesPerExecute; FrameIndex++)
+			// Interleave from SourceData into the output buffer.
+			for (int32 ChannelIndex = 0; ChannelIndex < NumChannels; ChannelIndex++)
 			{
-				*OutputPtr = InputPtr[FrameIndex];
-				OutputPtr += NumChannels;
+				const TArrayView<const float> InputBuffer = SourceData->GetChannel(ChannelIndex);
+
+				const float* InputPtr = InputBuffer.GetData();
+				float* OutputPtr = &InterleavedAudioBuffer.GetData()[ChannelIndex];
+
+				// Assign values to output for single channel.
+				for (int32 FrameIndex = 0; FrameIndex < NumFramesPerExecute; FrameIndex++)
+				{
+					*OutputPtr = InputPtr[FrameIndex];
+					OutputPtr += NumChannels;
+				}
+			}
+		}
+		else
+		{
+			// Iterate over channels
+			for (int32 ChannelIndex = 0; ChannelIndex < NumChannels; ChannelIndex++)
+			{
+				const FAudioBuffer& InputBuffer = *GraphOutputAudio[ChannelIndex];
+
+				const float* InputPtr = InputBuffer.GetData();
+				float* OutputPtr = &InterleavedAudioBuffer.GetData()[ChannelIndex];
+
+				// Assign values to output for single channel.
+				for (int32 FrameIndex = 0; FrameIndex < NumFramesPerExecute; FrameIndex++)
+				{
+					*OutputPtr = InputPtr[FrameIndex];
+					OutputPtr += NumChannels;
+				}
 			}
 		}
 		// TODO: memcpy for single channel. 
@@ -1141,7 +1262,18 @@ PRAGMA_ENABLE_DEPRECATION_WARNINGS
 		int32 AudioOutputIndex = AudioOutputNames.IndexOfByKey(InVertexName);
 		if (INDEX_NONE != AudioOutputIndex)
 		{
-			GraphOutputAudio[AudioOutputIndex] = InGraphOutputs.GetDataReadReference<FAudioBuffer>(InVertexName);
+			// For CAT sources, the audio output vertex carries FChannelAgnosticType instead of
+			// FAudioBuffer. GraphOutputAudio is empty in the CAT path, so update GraphCatOutput
+			// instead. MixerOutput buffer and Transcoder are format-dependent (not reference-
+			// dependent) and remain valid across reference updates.
+			if (GraphCatOutput.IsSet())
+			{
+				GraphCatOutput->GraphOutput = InGraphOutputs.GetDataReadReference<FChannelAgnosticType>(InVertexName);
+			}
+			else
+			{
+				GraphOutputAudio[AudioOutputIndex] = InGraphOutputs.GetDataReadReference<FAudioBuffer>(InVertexName);
+			}
 		}
 		// If it's the on-finished trigger, update the internal reference
 		else if (Frontend::SourceOneShotInterface::Outputs::OnFinished == InVertexName)
diff --git a/Engine/Plugins/Runtime/Metasound/Source/MetasoundGenerator/Private/MetasoundGeneratorBuilder.cpp b/Engine/Plugins/Runtime/Metasound/Source/MetasoundGenerator/Private/MetasoundGeneratorBuilder.cpp
--- a/Engine/Plugins/Runtime/Metasound/Source/MetasoundGenerator/Private/MetasoundGeneratorBuilder.cpp
+++ b/Engine/Plugins/Runtime/Metasound/Source/MetasoundGenerator/Private/MetasoundGeneratorBuilder.cpp
@@ -3,6 +3,7 @@
 #include "MetasoundGeneratorBuilder.h"
 
 #include "AudioParameter.h"
+#include "MetasoundChannelAgnosticType.h"
 #include "MetasoundDynamicOperatorTransactor.h"
 #include "MetasoundFrontendDataTypeRegistry.h"
 #include "MetasoundOperatorBuilder.h"
@@ -138,6 +139,26 @@ namespace Metasound
 
 	namespace GeneratorBuilder
 	{
+		TOptional<FChannelAgnosticTypeReadRef> FindOutputChannelAgnostic(const TArray<FVertexName>& InAudioVertexNames, const FVertexInterfaceData& InVertexData, const FOperatorSettings& InOperatorSettings, const FTopLevelAssetPath& InAssetPath)
+		{
+			const FOutputVertexInterfaceData& OutputVertexData = InVertexData.GetOutputs();
+
+			for (const FVertexName& AudioOutputName : InAudioVertexNames)
+			{
+				if (OutputVertexData.IsVertexBound(AudioOutputName))
+				{
+					if (const FAnyDataReference* Ref = OutputVertexData.FindDataReference(AudioOutputName))
+					{
+						if (IsCastable(Ref->GetDataTypeName(), GetMetasoundDataTypeName<FChannelAgnosticType>()))
+						{
+							return Ref->GetDataReadReference<FChannelAgnosticType>();
+						}
+					}
+				}
+			}
+			return {};
+		}
+		
 		TArray<FAudioBufferReadRef> FindOutputAudioBuffers(const TArray<FVertexName>& InAudioVertexNames, const FVertexInterfaceData& InVertexData, const FOperatorSettings& InOperatorSettings, const FTopLevelAssetPath& InAssetPath)
 		{
 			TArray<FAudioBufferReadRef> OutputBuffers;
@@ -147,11 +168,13 @@ namespace Metasound
 			// Get output audio buffers.
 			for (const FVertexName& AudioOutputName : InAudioVertexNames)
 			{
-				if (!OutputVertexData.IsVertexBound(AudioOutputName))
+				if (const FAnyDataReference* Ref = OutputVertexData.FindDataReference(AudioOutputName))
 				{
-					UE_LOG(LogMetaSound, Warning, TEXT("MetasoundSource [%s] does not contain audio output [%s] in output"), *InAssetPath.ToString(), *AudioOutputName.ToString());
+					if (Ref->GetDataTypeName() == Metasound::GetMetasoundDataTypeName<FAudioBuffer>())
+					{
+						OutputBuffers.Add(OutputVertexData.GetOrConstructDataReadReference<FAudioBuffer>(AudioOutputName, InOperatorSettings));
+					}
 				}
-				OutputBuffers.Add(OutputVertexData.GetOrConstructDataReadReference<FAudioBuffer>(AudioOutputName, InOperatorSettings));
 			}
 
 			return OutputBuffers;
@@ -333,7 +356,9 @@ namespace Metasound
 			FTriggerWriteRef PlayTrigger = VertexData.GetInputs().GetOrCreateDefaultDataWriteReference<FTrigger>(SourceInterface::Inputs::OnPlay, InOperatorSettings);
 
 			// Get outputs
-			TArray<FAudioBufferReadRef> OutputBuffers = FindOutputAudioBuffers(InInitParams.AudioOutputNames, VertexData, InOperatorSettings, InInitParams.AssetPath);
+			TArray<FAudioBufferReadRef> OutputAudioBuffers = FindOutputAudioBuffers(InInitParams.AudioOutputNames, VertexData, InOperatorSettings, InInitParams.AssetPath);
+			TOptional<FChannelAgnosticTypeReadRef> OutputCat = FindOutputChannelAgnostic(InInitParams.AudioOutputNames, VertexData, InOperatorSettings, InInitParams.AssetPath);
+			
 			FTriggerReadRef FinishTrigger = TDataReadReferenceFactory<FTrigger>::CreateExplicitArgs(InOperatorSettings, false);
 
 			if (const FAnyDataReference* FinishTriggerRef = VertexData.GetOutputs().FindDataReference(SourceOneShotInterface::Outputs::OnFinished))
@@ -357,9 +382,10 @@ namespace Metasound
 				MoveTemp(ParameterSetters),
 				MoveTemp(ParameterPackSetters),
 				MoveTemp(InAnalyzer),
-				MoveTemp(OutputBuffers),
+				MoveTemp(OutputAudioBuffers),
 				MoveTemp(PlayTrigger),
 				MoveTemp(FinishTrigger),
+				MoveTemp(OutputCat)
 			};
 		}
 
diff --git a/Engine/Plugins/Runtime/Metasound/Source/MetasoundGenerator/Public/MetasoundGenerator.h b/Engine/Plugins/Runtime/Metasound/Source/MetasoundGenerator/Public/MetasoundGenerator.h
--- a/Engine/Plugins/Runtime/Metasound/Source/MetasoundGenerator/Public/MetasoundGenerator.h
+++ b/Engine/Plugins/Runtime/Metasound/Source/MetasoundGenerator/Public/MetasoundGenerator.h
@@ -1,6 +1,7 @@
 // Copyright Epic Games, Inc. All Rights Reserved.
 #pragma once
 
+#include "MetasoundChannelAgnosticType.h"
 #include "MetasoundExecutableOperator.h"
 #include "MetasoundGraphOperator.h"
 #include "MetasoundInstanceCounter.h"
@@ -16,6 +17,7 @@
 
 #include "Analysis/MetasoundFrontendGraphAnalyzer.h"
 #include "Async/AsyncWork.h"
+#include "ChannelAgnostic/ChannelAgnosticTranscoding.h"
 #include "Containers/MpscQueue.h"
 #include "Containers/SpscQueue.h"
 #include "Sound/SoundGenerator.h"
@@ -80,6 +82,7 @@ namespace Metasound
 			TArray<TDataReadReference<FAudioBuffer>> OutputBuffers;
 			FTriggerWriteRef TriggerOnPlayRef;
 			FTriggerReadRef TriggerOnFinishRef;
+			TOptional<FChannelAgnosticTypeReadRef> OutputChannelAgnosticRef;
 		};
 	}
 
@@ -121,6 +124,12 @@ namespace Metasound
 		TSharedPtr<FGraphRenderCost> GraphRenderCost;
 		FName ClassName;
 		FTopLevelAssetPath AssetPath;
+		int32 NumChannels = 0;
+
+		// CAT format identity (e.g. "Cat:Mono1Dot0"). Empty for non-CAT sources.
+		// Carried to the generator for Environment population (graph-side format identity).
+		// Mixer output format is always resolved from NumChannels via FindDiscreteFormatFromNumChannels.
+		FName SourceFormatName;
 
 		static UE_API void Reset(FGeneratorInitParams& InParams);
 	};
@@ -349,7 +358,16 @@ namespace Metasound
 		FVertexInterfaceData VertexInterfaceData;
 		TArray<FVertexInterfaceChange> VertexInterfaceChangesSinceLastBroadcast;
 
-		TArray<FAudioBufferReadRef> GraphOutputAudio;
+		TArray<FAudioBufferReadRef> GraphOutputAudio; 
+		
+		struct FChannelAgnosticState
+		{
+			FChannelAgnosticTypeReadRef GraphOutput;
+			TOptional<Audio::FChannelAgnosticType> MixerOutput;		// Only allocated when transcode is needed (graph format != mixer format).
+			Audio::ChannelAgnosticTranscoder::FTranscoder Transcoder;	// Empty when identity — no transcode needed.
+			bool bIdentityTranscode = false;				// True when graph output format matches mixer format (skip transcode + extra buffer).
+		};
+		TOptional<FChannelAgnosticState> GraphCatOutput;			// This is the CAT reference from the graph.
 
 		// Triggered when metasound is finished
 		FTriggerReadRef OnFinishedTriggerRef;
@@ -374,6 +392,7 @@ namespace Metasound
 
 		int32 FinishSample = INDEX_NONE;
 		int32 NumChannels = 0;
+		FName SourceFormatName; // Graph-side CAT format identity. Set in Environment upstream (MetasoundSource); held here for reference only. Mixer output uses NumChannels separately.
 		int32 NumFramesPerExecute = 0;
 		int32 NumSamplesPerExecute = 0;
 
diff --git a/Engine/Source/Runtime/Engine/Classes/Sound/AudioBus.h b/Engine/Source/Runtime/Engine/Classes/Sound/AudioBus.h
--- a/Engine/Source/Runtime/Engine/Classes/Sound/AudioBus.h
+++ b/Engine/Source/Runtime/Engine/Classes/Sound/AudioBus.h
@@ -32,10 +32,22 @@ namespace AudioBusUtils
 		case 6:	return EAudioBusChannels::FivePointOne;
 		case 8:	return EAudioBusChannels::SevenPointOne;
 		default:
-			// TODO alex.perez: LogAudio undeclared identifier, is it ok to add the header that has LogAudio, or should we move this function to another place?
-			//UE_LOG(LogAudio, Error, TEXT("Number of channels: %d not available in AudioBusChannels configurations. Make sure that the number of channels used is 1, 2, 4, 6 or 8."), InValue);
-			check(false);
-			return EAudioBusChannels::Mono;
+		{
+			// Channel count doesn't map to a standard bus configuration (e.g. CAT formats like 4.1).
+			// Clamp to the nearest supported configuration rather than crashing. (UE-370803)
+			// Log once to avoid spam if unexpectedly called from a hot path.
+			static bool bLoggedOnce = false;
+			if (!bLoggedOnce)
+			{
+				bLoggedOnce = true;
+				UE_LOG(LogCore, Warning, TEXT("AudioBus channel count %d is not a supported configuration (1, 2, 4, 6, 8). Clamping to nearest supported format for visualization."), InValue);
+			}
+			if (InValue <= 1)      return EAudioBusChannels::Mono;
+			else if (InValue <= 2) return EAudioBusChannels::Stereo;
+			else if (InValue <= 4) return EAudioBusChannels::Quad;
+			else if (InValue <= 6) return EAudioBusChannels::FivePointOne;
+			else                   return EAudioBusChannels::SevenPointOne;
+		}
 		}
 	};
 }
