Fix round 2 for registering MetaSound parent when initializing new preset asset

#rnx
#Synced-CL 51793207

diff --git a/Engine/Plugins/Runtime/Metasound/Source/MetasoundFrontend/Private/DocumentTemplates/MetasoundFrontendPresetTemplate.cpp b/Engine/Plugins/Runtime/Metasound/Source/MetasoundFrontend/Private/DocumentTemplates/MetasoundFrontendPresetTemplate.cpp
--- a/Engine/Plugins/Runtime/Metasound/Source/MetasoundFrontend/Private/DocumentTemplates/MetasoundFrontendPresetTemplate.cpp
+++ b/Engine/Plugins/Runtime/Metasound/Source/MetasoundFrontend/Private/DocumentTemplates/MetasoundFrontendPresetTemplate.cpp
@@ -12,6 +12,7 @@
 #include "MetasoundFrontendDocument.h"
 #include "MetasoundFrontendDocumentBuilder.h"
 #include "MetasoundFrontendDocumentBuilderRegistry.h"
+#include "MetasoundGlobals.h"
 #include "MetasoundTrigger.h"
 
 #define LOCTEXT_NAMESPACE "MetaSound"
@@ -144,12 +145,26 @@ const FMetaSoundFrontendDocumentTemplate::FEditorOptions& FMetaSoundFrontendPres
 
 void FMetaSoundFrontendPresetTemplate::OnAssetInitialized(TArray<UObject*> SelectedObjects, FMetaSoundFrontendDocumentBuilder& OutBuilder)
 {
+	using namespace Metasound::Frontend;
+
 	for (UObject* SelectedObject : SelectedObjects)
 	{
-		if (&OutBuilder.GetConstDocumentInterfaceChecked().GetBaseMetaSoundUClass() == SelectedObject->GetClass())
+		if (SelectedObject)
 		{
-			Parent = SelectedObject;
-			return;
+			if (&OutBuilder.GetConstDocumentInterfaceChecked().GetBaseMetaSoundUClass() == SelectedObject->GetClass())
+			{
+				Parent = SelectedObject;
+
+				if (FMetasoundAssetBase* ParentAsset = IMetaSoundAssetManager::GetChecked().GetAsAsset(*SelectedObject))
+				{
+					const bool bCanExecute = Metasound::CanEverExecuteGraph();
+					if (bCanExecute)
+					{
+						ParentAsset->UpdateAndRegisterForExecution(IMetaSoundAssetManager::GetChecked().GetRegistrationOptions());
+					}
+				}
+				return;
+			}
 		}
 	}
 }
@@ -172,7 +187,11 @@ void FMetaSoundFrontendPresetTemplate::OnPropertyChanged(const FPropertyChangedE
 				}
 				else
 				{
-					ParentAsset->UpdateAndRegisterForExecution(IMetaSoundAssetManager::GetChecked().GetRegistrationOptions());
+					const bool bCanExecute = Metasound::CanEverExecuteGraph();
+					if (bCanExecute)
+					{
+						ParentAsset->UpdateAndRegisterForExecution(IMetaSoundAssetManager::GetChecked().GetRegistrationOptions());
+					}
 				}
 			}
 		}
diff --git a/Engine/Plugins/Runtime/Metasound/Source/MetasoundFrontend/Public/DocumentTemplates/MetasoundFrontendDocumentTemplate.h b/Engine/Plugins/Runtime/Metasound/Source/MetasoundFrontend/Public/DocumentTemplates/MetasoundFrontendDocumentTemplate.h
--- a/Engine/Plugins/Runtime/Metasound/Source/MetasoundFrontend/Public/DocumentTemplates/MetasoundFrontendDocumentTemplate.h
+++ b/Engine/Plugins/Runtime/Metasound/Source/MetasoundFrontend/Public/DocumentTemplates/MetasoundFrontendDocumentTemplate.h
@@ -91,9 +91,7 @@ struct FMetaSoundFrontendDocumentTemplate
 	UE_EXPERIMENTAL(5.8, "MetaSound Templates are in active development and subject to change")
 	UE_API virtual const FEditorOptions& GetEditorOptions() const;
 
-#if WITH_EDITOR
 	UE_API virtual EDataValidationResult IsDataValid(const FMetaSoundFrontendDocumentBuilder& Builder, FDataValidationContext& InOutContext) const;
-#endif // WITH_EDITOR
 
 	UE_EXPERIMENTAL(5.8, "MetaSound Templates are in active development and subject to change")
 	UE_API virtual void OnAssetInitialized(TArray<UObject*> SelectedObjects, FMetaSoundFrontendDocumentBuilder& OutBuilder);
