Fix CIS engine test issues by disabling MetaSound CAT type registration if the metasound experimental plugin is not enabled.

#rnx
#Synced-CL 52211756

diff --git a/Engine/Plugins/Runtime/Metasound/Source/MetasoundFrontend/Private/MetasoundChannelAgnosticType.cpp b/Engine/Plugins/Runtime/Metasound/Source/MetasoundFrontend/Private/MetasoundChannelAgnosticType.cpp
--- a/Engine/Plugins/Runtime/Metasound/Source/MetasoundFrontend/Private/MetasoundChannelAgnosticType.cpp
+++ b/Engine/Plugins/Runtime/Metasound/Source/MetasoundFrontend/Private/MetasoundChannelAgnosticType.cpp
@@ -137,20 +137,28 @@ namespace Metasound
 	static_assert(::Metasound::TInputNode<FChannelAgnosticType, EVertexAccessType::Reference>::bCanRegister);
 	static_assert(::Metasound::Frontend::TMetasoundDataTypeRegistration<FChannelAgnosticType>::bCanRegister);
 	
-	REGISTER_METASOUND_DATATYPE(FChannelAgnosticType, "Cat");
+	// Registering of data type is performed in module startup because it should
+	// only be registered if the MetaSoundExperimental plugin is enabled. 
+	DEFINE_METASOUND_DATA_TYPE(FChannelAgnosticType, "Cat"); 
 	REGISTER_METASOUND_POLY_TYPE(FChannelAgnosticType);
 
 	static_assert(TEnableConstructorVertex<FSoundfieldChannelAgnosticType>::Value);
 	static_assert(CHasFactoryFunctionWithSettings<FSoundfieldChannelAgnosticType>);
 	static_assert(TLiteralTraits<FSoundfieldChannelAgnosticType>::bIsParsableFromAnyLiteralType);
 
-	REGISTER_METASOUND_DATATYPE(FCompositeChannelAgnosticType, "Cat:Composite");
+	// Registering of data type is performed in module startup because it should
+	// only be registered if the MetaSoundExperimental plugin is enabled. 
+	DEFINE_METASOUND_DATA_TYPE(FCompositeChannelAgnosticType, "Cat:Composite");
 	REGISTER_METASOUND_POLY_TYPE(FCompositeChannelAgnosticType);
 
-	REGISTER_METASOUND_DATATYPE(FDiscreteChannelAgnosticType, "Cat:Discrete");
+	// Registering of data type is performed in module startup because it should
+	// only be registered if the MetaSoundExperimental plugin is enabled. 
+	DEFINE_METASOUND_DATA_TYPE(FDiscreteChannelAgnosticType, "Cat:Discrete");
 	REGISTER_METASOUND_POLY_TYPE(FDiscreteChannelAgnosticType);
 
-	REGISTER_METASOUND_DATATYPE(FSoundfieldChannelAgnosticType, "Cat:Soundfield");
+	// Registering of data type is performed in module startup because it should
+	// only be registered if the MetaSoundExperimental plugin is enabled. 
+	DEFINE_METASOUND_DATA_TYPE(FSoundfieldChannelAgnosticType, "Cat:Soundfield");
 	REGISTER_METASOUND_POLY_TYPE(FSoundfieldChannelAgnosticType);
 	
 	namespace ChannelAgnosticTypePrivate 
diff --git a/Engine/Plugins/Runtime/Metasound/Source/MetasoundFrontend/Private/MetasoundFrontend.cpp b/Engine/Plugins/Runtime/Metasound/Source/MetasoundFrontend/Private/MetasoundFrontend.cpp
--- a/Engine/Plugins/Runtime/Metasound/Source/MetasoundFrontend/Private/MetasoundFrontend.cpp
+++ b/Engine/Plugins/Runtime/Metasound/Source/MetasoundFrontend/Private/MetasoundFrontend.cpp
@@ -8,12 +8,14 @@
 #include "HAL/IConsoleManager.h"
 #include "Interfaces/IPluginManager.h"
 #include "MetasoundAudioBuffer.h"
+#include "MetasoundChannelAgnosticType.h"
 #include "MetasoundDataTypeRegistrationMacro.h"
 #include "MetasoundFrontendDocument.h"
 #include "MetasoundFrontendDocumentVersioning.h"
 #include "MetasoundFrontendGraph.h"
 #include "MetasoundFrontendNodeClassRegistryPrivate.h"
 #include "MetasoundFrontendNodeTemplateRegistry.h"
+#include "MetasoundFrontendModuleRegistrationMacros.h"
 #include "MetasoundFrontendRegistryContainer.h"
 #include "MetasoundFrontendRegistryTransaction.h"
 #include "MetasoundJsonBackend.h"
@@ -191,6 +193,7 @@ class FMetasoundFrontendModule : public IModuleInterface
 {
 	virtual void StartupModule() override
 	{
+		using namespace Metasound;
 		using namespace Metasound::Frontend;
 
 		RegisterNodeTemplate(MakeUnique<FAudioAnalyzerNodeTemplate>());
@@ -207,10 +210,20 @@ class FMetasoundFrontendModule : public IModuleInterface
 		{
 			FNodeClassRegistry::Get().NotifyPluginUnmounted(Plugin);
 		});
+
+		// Only register/unregister CAT types if the experimental plugin is enabled. 
+		if (IsExperimentalPluginEnabled())
+		{
+			RegisterDataType<FChannelAgnosticType>(METASOUND_GET_MODULE_LIST_INFO);
+			RegisterDataType<FCompositeChannelAgnosticType>(METASOUND_GET_MODULE_LIST_INFO);
+			RegisterDataType<FDiscreteChannelAgnosticType>(METASOUND_GET_MODULE_LIST_INFO);
+			RegisterDataType<FSoundfieldChannelAgnosticType>(METASOUND_GET_MODULE_LIST_INFO);
+		}
 	}
 
 	virtual void ShutdownModule() override
 	{
+		using namespace Metasound;
 		using namespace Metasound::Frontend;
 
 		IPluginManager::Get().OnPluginUnmounted().Remove(OnPluginUnmountedHandle);
@@ -223,6 +236,15 @@ class FMetasoundFrontendModule : public IModuleInterface
 		UnregisterNodeTemplate(FAudioAnalyzerNodeTemplate::ClassName, FAudioAnalyzerNodeTemplate::VersionNumber);
 		UnregisterNodeTemplate(FInputNodeTemplate::ClassName, FInputNodeTemplate::VersionNumber);
 		UnregisterNodeTemplate(FRerouteNodeTemplate::ClassName, FRerouteNodeTemplate::VersionNumber);
+
+		// Only register/unregister CAT types if the experimental plugin is enabled. 
+		if (IsExperimentalPluginEnabled())
+		{
+			UnregisterDataType<FChannelAgnosticType>(METASOUND_GET_MODULE_LIST_INFO);
+			UnregisterDataType<FCompositeChannelAgnosticType>(METASOUND_GET_MODULE_LIST_INFO);
+			UnregisterDataType<FDiscreteChannelAgnosticType>(METASOUND_GET_MODULE_LIST_INFO);
+			UnregisterDataType<FSoundfieldChannelAgnosticType>(METASOUND_GET_MODULE_LIST_INFO);
+		}
 	}
 
 private:
