Fix out of bounds index issue on for ADSR node when combination of attack, decay and release exceed capacity of integer.

#rnx
#Synced-CL 52211756

diff --git a/Engine/Plugins/Runtime/Metasound/Source/MetasoundStandardNodes/Private/MetasoundADSREnvelopeNode.cpp b/Engine/Plugins/Runtime/Metasound/Source/MetasoundStandardNodes/Private/MetasoundADSREnvelopeNode.cpp
--- a/Engine/Plugins/Runtime/Metasound/Source/MetasoundStandardNodes/Private/MetasoundADSREnvelopeNode.cpp
+++ b/Engine/Plugins/Runtime/Metasound/Source/MetasoundStandardNodes/Private/MetasoundADSREnvelopeNode.cpp
@@ -653,13 +653,17 @@ namespace Metasound
 
 		void UpdateParams()
 		{
+			// Clamp each segment sample count to Max/4 so that the sum of all three 
+			// can never overflow int32. 
+			static constexpr float MaxSegmentSampleCount = static_cast<float>(TNumericLimits<int32>::Max() / 4);
+
 			float AttackTimeSeconds = AttackTime->GetSeconds();
 			float DecayTimeSeconds = DecayTime->GetSeconds();
 			float ReleaseTimeSeconds = ReleaseTime->GetSeconds();
-			EnvState.AttackSampleCount = FMath::Max(1, static_cast<int32>(SampleRate * FMath::Max(0.0f, AttackTimeSeconds)));
-			EnvState.DecaySampleCount = FMath::Max(0, static_cast<int32>(SampleRate * FMath::Max(0.0f, DecayTimeSeconds)));
+			EnvState.AttackSampleCount = FMath::Max(1, static_cast<int32>(FMath::Min(SampleRate * FMath::Max(0.0f, AttackTimeSeconds), MaxSegmentSampleCount)));
+			EnvState.DecaySampleCount = FMath::Max(0, static_cast<int32>(FMath::Min(SampleRate * FMath::Max(0.0f, DecayTimeSeconds), MaxSegmentSampleCount)));
 			EnvState.SustainLevel = FMath::Clamp(*SustainLevel, 0.0f, 1.0f);
-			EnvState.ReleaseSampleCount = FMath::Max(0, static_cast<int32>(SampleRate * FMath::Max(0.0f, ReleaseTimeSeconds)));
+			EnvState.ReleaseSampleCount = FMath::Max(0, static_cast<int32>(FMath::Min(SampleRate * FMath::Max(0.0f, ReleaseTimeSeconds), MaxSegmentSampleCount)));
 			EnvState.AttackCurveFactor = FMath::Max(KINDA_SMALL_NUMBER, *AttackCurveFactor);
 			EnvState.DecayCurveFactor = FMath::Max(KINDA_SMALL_NUMBER, *DecayCurveFactor);
 			EnvState.ReleaseCurveFactor = FMath::Max(KINDA_SMALL_NUMBER, *ReleaseCurveFactor);
