Fix faliing CIS for Audio.Metasound.Generator.*
Remove NumChannels from FGeneratorInitParams; centralize CAT format resolution

Remove the redundant NumChannels field from Metasound::FGeneratorInitParams. The generator now derives the mixer-compatible channel count internally:
- CAT path: from SourceFormatName via channel registry lookup + clamp to mixer-compatible values (1, 2, 4, 6, 8).
- Discrete path: from AudioOutputNames.Num().

This ensures CAT (Channel Agnostic Type) format resolution lives in a single place ? UMetaSoundSource::ResolveSourceFormatAndNumChannels() ? which is called by both CreateSoundGenerator() and the operator cache subsystem. Previously
the same Custom > Auto > Fallback logic was duplicated across both call sites.

CAT code paths remain gated behind the MetasoundExperimental plugin: the ChannelAgnostic output format is only registered when the plugin is enabled, and the generator's CAT transcoding setup is naturally gated by
OutputChannelAgnosticRef.IsSet(). Old discrete workflows are unchanged.

Files:
- MetasoundGenerator.h ? Removed NumChannels from FGeneratorInitParams
- MetasoundGenerator.cpp ? Added ResolveNumChannelsFromFormat() helper; InitBase() derives NumChannels internally
- MetasoundSource.h ? Added public ResolveSourceFormatAndNumChannels() declaration
- MetasoundSource.cpp ? Extracted format/channel resolution into ResolveSourceFormatAndNumChannels()
- MetasoundOperatorCacheSubsystem.cpp ? Replaced duplicated resolution logic with call to centralized method

Testing: Audio.Metasound (1867 passed), Audio.Metasound.Generator (15 passed) in EngineTest; Audio.Metasound (43 passed) in AudioQA (MetasoundExperimental enabled)

#Synced-CL 52211756

diff --git a/Engine/Plugins/Runtime/Metasound/Source/MetasoundEngine/Private/MetasoundOperatorCacheSubsystem.cpp b/Engine/Plugins/Runtime/Metasound/Source/MetasoundEngine/Private/MetasoundOperatorCacheSubsystem.cpp
--- a/Engine/Plugins/Runtime/Metasound/Source/MetasoundEngine/Private/MetasoundOperatorCacheSubsystem.cpp
+++ b/Engine/Plugins/Runtime/Metasound/Source/MetasoundEngine/Private/MetasoundOperatorCacheSubsystem.cpp
@@ -2,8 +2,6 @@
 #include "MetasoundOperatorCacheSubsystem.h"
 
 #include "AudioMixerDevice.h"
-#include "Interfaces/MetasoundFrontendSourceInterface.h"
-#include "TypeFamily/ChannelTypeFamily.h"
 #include "MetasoundAssetManager.h"
 #include "MetasoundGenerator.h"
 #include "MetasoundGeneratorModule.h"
@@ -41,63 +39,12 @@ namespace Metasound::OperatorCachePrivate
 		// Graph analyzer currently only enabled for preview sounds (but can theoretically be supported for all sounds)
 		BuilderSettings.bPopulateInternalDataReferences = InParams.bIsPreviewSound;
 	
-		// Clamp helper — mirrors ClampToMixerChannelCount in MetasoundSource.cpp.
-		// AudioMixer only supports 1, 2, 4, 6, 8 channels.
-		auto ClampToMixer = [](int32 InCh) -> int32
+		// Resolve format and channel count via the centralized method on UMetaSoundSource.
+		// Empty default parameters — precaching has no runtime parameter overrides.
+		const UMetaSoundSource::FResolvedSourceFormat ResolvedFormat = InMetaSound.ResolveSourceFormatAndNumChannels(InParams, {}, Environment);
+		if (ResolvedFormat.NumChannels == INDEX_NONE)
 		{
-			if (InCh <= 1) return 1;
-			if (InCh <= 2) return 2;
-			if (InCh <= 4) return 4;
-			if (InCh <= 6) return 6;
-			return 8;
-		};
-
-		int32 NumChannels;
-		FName SourceFormatName;
-		if (InMetaSound.OutputFormat == EMetaSoundOutputAudioFormat::ChannelAgnostic)
-		{
-			// Three-branch format resolution, mirroring CreateSoundGenerator:
-			// 1) Custom: user-specified format name.
-			// 2) Auto: detect from wave proxy channel counts via CalculateAutoSourceFormat.
-			// 3) Fallback: use device channel count.
-			if (InMetaSound.ChannelAgnosticFormatChooser == EMetasoundChannelAgnosticSourceFormatChooser::Custom
-				&& !InMetaSound.ChannelAgnosticCustomFormat.IsNone())
-			{
-				PRAGMA_DISABLE_EXPERIMENTAL_WARNINGS
-				if (const Audio::FChannelTypeFamily* Found = Audio::GetChannelRegistry().FindChannel(InMetaSound.ChannelAgnosticCustomFormat))
-				{
-					SourceFormatName = Found->GetName();
-					NumChannels = ClampToMixer(Found->NumChannels());
-				}
-				PRAGMA_ENABLE_EXPERIMENTAL_WARNINGS
-				else
-				{
-					NumChannels = ClampToMixer(InParams.NumChannels);
-				}
-			}
-			// Auto: determine format from default wave proxy channel counts.
-			// Empty param array causes CalculateAutoSourceFormat to use RuntimeInput defaults,
-			// which is correct for precaching (no runtime parameter overrides available).
-			else if (TSharedPtr<const Audio::FChannelTypeFamily> SP = InMetaSound.CalculateAutoSourceFormat({}); SP.IsValid())
-			{
-				SourceFormatName = SP->GetName();
-				NumChannels = ClampToMixer(SP->NumChannels());
-			}
-			else
-			{
-				NumChannels = ClampToMixer(InParams.NumChannels);
-			}
-
-			// Embed SourceFormatName in the Environment so graph nodes can read it at build time.
-			// Mirrors the Environment.SetValue call in CreateSoundGenerator.
-			if (!SourceFormatName.IsNone())
-			{
-				Environment.SetValue<FName>(SourceInterface::Environment::SourceFormatName, SourceFormatName);
-			}
-		}
-		else
-		{
-			NumChannels = InMetaSound.GetOutputAudioChannelOrder().Num();
+			return {};
 		}
 
 		return FGeneratorInitParams
@@ -109,8 +56,7 @@ namespace Metasound::OperatorCachePrivate
 			.AudioOutputNames = InMetaSound.GetOutputAudioChannelOrder(),
 			.bBuildSynchronous = true,
 			.AssetPath = FTopLevelAssetPath(&InMetaSound),
-			.NumChannels = NumChannels,
-			.SourceFormatName = SourceFormatName,
+			.SourceFormatName = ResolvedFormat.SourceFormatName,
 		};
 	}
 } // namespace Metasound::OperatorCachePrivate
diff --git a/Engine/Plugins/Runtime/Metasound/Source/MetasoundEngine/Private/MetasoundSource.cpp b/Engine/Plugins/Runtime/Metasound/Source/MetasoundEngine/Private/MetasoundSource.cpp
--- a/Engine/Plugins/Runtime/Metasound/Source/MetasoundEngine/Private/MetasoundSource.cpp
+++ b/Engine/Plugins/Runtime/Metasound/Source/MetasoundEngine/Private/MetasoundSource.cpp
@@ -1194,66 +1194,74 @@ bool UMetaSoundSource::ImplementsParameterInterface(Audio::FParameterInterfacePt
 	return GetConstDocument().Interfaces.Contains(Version);
 }
 
-ISoundGeneratorPtr UMetaSoundSource::CreateSoundGenerator(const FSoundGeneratorInitParams& InParams, TArray<FAudioParameter>&& InDefaultParameters)
+UMetaSoundSource::FResolvedSourceFormat UMetaSoundSource::ResolveSourceFormatAndNumChannels(const FSoundGeneratorInitParams& InParams, const TArray<FAudioParameter>& InDefaultParameters, Metasound::FMetasoundEnvironment& InOutEnvironment)
 {
-	using namespace Metasound;
 	using namespace Metasound::Frontend;
-	using namespace Metasound::Engine;
-	using namespace Metasound::SourcePrivate;
-
-	METASOUND_LLM_SCOPE;
-	METASOUND_TRACE_CPUPROFILER_EVENT_SCOPE(UMetaSoundSource::CreateSoundGenerator);
-
-	FOperatorSettings InSettings = GetOperatorSettings(static_cast<FSampleRate>(InParams.SampleRate));
 
-	SampleRate = InSettings.GetSampleRate();
-
-	FMetasoundEnvironment Environment = CreateEnvironment(InParams);
-	
-	// If the output is Channel Agnostic, determine the format and channel count.
-	// Priority: 1) user-specified custom format, 2) auto-detect from wave proxy assets,
-	// 3) fall back to the audio mixer's channel count.
-	// SourceFormatName carries format identity through to the generator so it can set up
-	// transcoding without a lossy reverse lookup from NumChannels.
-	FName SourceFormatName;
+	FResolvedSourceFormat Result;
 	if (OutputFormat == EMetaSoundOutputAudioFormat::ChannelAgnostic)
 	{
-		// Custom: user has specified a format name directly.
+		// Priority: 1) user-specified custom format, 2) auto-detect from wave proxy assets,
+		// 3) return error condition for undetermined type
 		if (ChannelAgnosticFormatChooser == EMetasoundChannelAgnosticSourceFormatChooser::Custom && !ChannelAgnosticCustomFormat.IsNone())
 		{
 			PRAGMA_DISABLE_EXPERIMENTAL_WARNINGS
 			if (const Audio::FChannelTypeFamily* Found = Audio::GetChannelRegistry().FindChannel(ChannelAgnosticCustomFormat))
 			{
-				SourceFormatName = Found->GetName();
-				Environment.SetValue<FName>(SourceInterface::Environment::SourceFormatName, SourceFormatName);
-				NumChannels = ClampToMixerChannelCount(Found->NumChannels());
+				Result.SourceFormatName = Found->GetName();
+				InOutEnvironment.SetValue<FName>(SourceInterface::Environment::SourceFormatName, Result.SourceFormatName);
+				Result.NumChannels = ClampToMixerChannelCount(Found->NumChannels());
 			}
 			PRAGMA_ENABLE_EXPERIMENTAL_WARNINGS
 			else
 			{
 				UE_LOG(LogMetaSound, Warning, TEXT("Custom format '%s' not found in channel registry"), *ChannelAgnosticCustomFormat.ToString());
-				NumChannels = ClampToMixerChannelCount(InParams.NumChannels);
+				Result.NumChannels = ClampToMixerChannelCount(InParams.NumChannels);
 			}
 		}
-		// Auto: determine format from wave proxy channel counts.
 		else if (TSharedPtr<const Audio::FChannelTypeFamily> SP = CalculateAutoSourceFormat(InDefaultParameters); SP.IsValid())
 		{
-			SourceFormatName = SP->GetName();
-			Environment.SetValue<FName>(SourceInterface::Environment::SourceFormatName, SourceFormatName);
-			NumChannels = ClampToMixerChannelCount(SP->NumChannels());
+			Result.SourceFormatName = SP->GetName();
+			InOutEnvironment.SetValue<FName>(SourceInterface::Environment::SourceFormatName, Result.SourceFormatName);
+			Result.NumChannels = ClampToMixerChannelCount(SP->NumChannels());
 		}
 		else
 		{
-			NumChannels = ClampToMixerChannelCount(InParams.NumChannels); // Clamp mixer channel count to supported values (1,2,4,6,8).
+			UE_LOG(LogMetaSound, Error, TEXT("Failed to determine concrete format for ChannelAgnotistic output on source \"%s\""), *GetName()) ;
 		}
 	}
 	else
 	{
-		// Otherwise the number of AudioOutput is our NumChannels.
-		NumChannels = GetOutputAudioChannelOrder().Num(); 
+		Result.NumChannels = GetOutputAudioChannelOrder().Num();
 	}
 
-	
+	return Result;
+}
+
+ISoundGeneratorPtr UMetaSoundSource::CreateSoundGenerator(const FSoundGeneratorInitParams& InParams, TArray<FAudioParameter>&& InDefaultParameters)
+{
+	using namespace Metasound;
+	using namespace Metasound::Frontend;
+	using namespace Metasound::Engine;
+	using namespace Metasound::SourcePrivate;
+
+	METASOUND_LLM_SCOPE;
+	METASOUND_TRACE_CPUPROFILER_EVENT_SCOPE(UMetaSoundSource::CreateSoundGenerator);
+
+	FOperatorSettings InSettings = GetOperatorSettings(static_cast<FSampleRate>(InParams.SampleRate));
+
+	SampleRate = InSettings.GetSampleRate();
+
+	FMetasoundEnvironment Environment = CreateEnvironment(InParams);
+	const UMetaSoundSource::FResolvedSourceFormat ResolvedFormat = ResolveSourceFormatAndNumChannels(InParams, InDefaultParameters, Environment);
+	if (ResolvedFormat.NumChannels == INDEX_NONE)
+	{
+		// Failed to resolve number of channels. This should not happen, but if it does,
+		// do not attempt to render audio.
+		return nullptr;
+	}
+	NumChannels = ResolvedFormat.NumChannels;
+
 	FParameterRouter& Router = UMetaSoundSource::GetParameterRouter();
 	TSharedPtr<TSpscQueue<FMetaSoundParameterTransmitter::FParameter>> DataChannel = Router.FindOrCreateDataChannelForReader(InParams.AudioDeviceID, InParams.InstanceID);
 
@@ -1312,8 +1320,7 @@ ISoundGeneratorPtr UMetaSoundSource::CreateSoundGenerator(const FSoundGeneratorI
 				.GraphRenderCost = GraphRenderCost,
 				.ClassName = GetClass()->GetFName(),
 				.AssetPath = GetAssetPathChecked(),
-				.NumChannels =  NumChannels,
-				.SourceFormatName = SourceFormatName
+				.SourceFormatName = ResolvedFormat.SourceFormatName
 			},
 			DynamicTransactor->CreateTransformQueue(InSettings, Environment, GraphRenderCost) // Create transaction queue
 		};
@@ -1353,8 +1360,7 @@ ISoundGeneratorPtr UMetaSoundSource::CreateSoundGenerator(const FSoundGeneratorI
 				.bBuildSynchronous = bBuildSynchronous,
 				.DataChannel = DataChannel,
 				.ClassName = GetClass()->GetFName(),
-				.NumChannels = NumChannels,
-				.SourceFormatName = SourceFormatName
+				.SourceFormatName = ResolvedFormat.SourceFormatName
 			};
 			InitParams.AssetPath = GetAssetPathChecked();
 			Generator = MakeShared<FMetasoundConstGraphGenerator>(MoveTemp(InitParams));
diff --git a/Engine/Plugins/Runtime/Metasound/Source/MetasoundEngine/Public/MetasoundSource.h b/Engine/Plugins/Runtime/Metasound/Source/MetasoundEngine/Public/MetasoundSource.h
--- a/Engine/Plugins/Runtime/Metasound/Source/MetasoundEngine/Public/MetasoundSource.h
+++ b/Engine/Plugins/Runtime/Metasound/Source/MetasoundEngine/Public/MetasoundSource.h
@@ -92,11 +92,21 @@ class UMetaSoundSource : public USoundWaveProcedural, public FMetasoundAssetBase
 
 	friend struct Metasound::Engine::FAssetHelper;
 	friend class UMetaSoundSourceBuilder;
-	
+
 	//Forward declare
 	class FAudioParameterCollector;
 
 public:
+	/** Result of ResolveSourceFormatAndNumChannels. */
+	struct FResolvedSourceFormat
+	{
+		/** The resolved source format name (NAME_None for discrete/non-CAT sources). */
+		FName SourceFormatName;
+
+		/** The resolved channel count, or INDEX_NONE on failure. */
+		int32 NumChannels = INDEX_NONE;
+	};
+
 	// FRuntimeInput represents an input to a MetaSound which can be manipulated.
 	struct UE_INTERNAL FRuntimeInput
 	{
@@ -386,8 +396,26 @@ public:
 	 * @param InParams - Audio parameters to inspect for wave proxies. Pass empty array to use defaults.
 	 * @return Shared pointer to the resolved channel type family, or nullptr if no format could be determined.
 	 */
+	UE_INTERNAL
 	UE_API TSharedPtr<const Audio::FChannelTypeFamily> CalculateAutoSourceFormat(const TArray<FAudioParameter>& InParams) const;
 
+	/**
+	 * Resolves SourceFormatName and NumChannels based on the current OutputFormat.
+	 * For ChannelAgnostic sources, determines format via Custom > Auto 
+	 * priority and populates the Environment. For discrete sources, derives
+	 * NumChannels from GetOutputAudioChannelOrder().
+	 *
+	 * Does NOT modify USoundWave::NumChannels -- the caller is responsible for
+	 * assigning the returned value where appropriate.
+	 *
+	 * @param InParams - Sound generator init params (used for fallback channel count).
+	 * @param InDefaultParameters - Audio parameters for auto-format detection (pass empty for precache).
+	 * @param InOutEnvironment - Environment to populate with SourceFormatName.
+	 * @return FResolvedSourceFormat containing the format name and channel count.
+	 */
+	UE_INTERNAL
+	UE_API FResolvedSourceFormat ResolveSourceFormatAndNumChannels(const FSoundGeneratorInitParams& InParams, const TArray<FAudioParameter>& InDefaultParameters, Metasound::FMetasoundEnvironment& InOutEnvironment);
+
 	/** Find the Source related to this Preset.
 	 * 
 	 * If this MetaSound is a preset and preset graph inflation is enabled, this
@@ -397,6 +425,7 @@ public:
 	 * 		AND/OR
 	 *  	- Has modified constructor pin overrides.
 	 */
+	UE_INTERNAL
 	UE_API const UMetaSoundSource& FindFirstNoninflatableSource(Metasound::FMetasoundEnvironment& InOutEnvironment, TFunctionRef<void(const UMetaSoundSource&)> OnTraversal) const;
 
 private:
@@ -409,6 +438,7 @@ private:
 	mutable FCriticalSection GeneratorMapCriticalSection;
 	TArray<Metasound::FGeneratorInstanceInfo> Generators;
 	UE_API void TrackGenerator(Metasound::FGeneratorInstanceInfo&& GeneratorInfo);
+
 	UE_API void ForgetGenerator(ISoundGeneratorPtr Generator);
 
 	PRAGMA_DISABLE_INTERNAL_WARNINGS
diff --git a/Engine/Plugins/Runtime/Metasound/Source/MetasoundGenerator/Private/MetasoundGenerator.cpp b/Engine/Plugins/Runtime/Metasound/Source/MetasoundGenerator/Private/MetasoundGenerator.cpp
--- a/Engine/Plugins/Runtime/Metasound/Source/MetasoundGenerator/Private/MetasoundGenerator.cpp
+++ b/Engine/Plugins/Runtime/Metasound/Source/MetasoundGenerator/Private/MetasoundGenerator.cpp
@@ -263,7 +263,6 @@ namespace Metasound
 		InParams.DefaultParameters = {};
 		InParams.DataChannel.Reset();
 		InParams.GraphRenderCost.Reset();
-		InParams.NumChannels = 0;
 		InParams.SourceFormatName = NAME_None;
 	}
 
@@ -304,11 +303,44 @@ namespace Metasound
 		return OnSetGraph.Add(Delegate);
 	}
 
+	// Clamp a channel count to the nearest mixer-compatible value (1, 2, 4, 6, 8).
+	static int32 ClampToMixerChannelCount(int32 InChannels)
+	{
+		if (InChannels <= 1) return 1;
+		if (InChannels <= 2) return 2;
+		if (InChannels <= 4) return 4;
+		if (InChannels <= 6) return 6;
+		return 8;
+	}
+
+	// Resolve the mixer-compatible channel count from a CAT SourceFormatName.
+	static int32 ResolveNumChannelsFromFormat(FName InSourceFormatName)
+	{
+		PRAGMA_DISABLE_EXPERIMENTAL_WARNINGS
+		if (const Audio::FChannelTypeFamily* Found = Audio::GetChannelRegistry().FindChannel(InSourceFormatName))
+		{
+			return ClampToMixerChannelCount(Found->NumChannels());
+		}
+		PRAGMA_ENABLE_EXPERIMENTAL_WARNINGS
+		UE_LOG(LogMetaSound, Warning, TEXT("CAT format '%s' not found in channel registry; defaulting to stereo."), *InSourceFormatName.ToString());
+		return 2;
+	}
+
 	void FMetasoundGenerator::InitBase(FGeneratorInitParams& InInitParams)
 	{
 		AssetPath = InInitParams.AssetPath;
-		NumChannels = InInitParams.NumChannels;
 		SourceFormatName = InInitParams.SourceFormatName;
+
+		if (!SourceFormatName.IsNone())
+		{
+			// CAT path: derive mixer-compatible channel count from format name.
+			NumChannels = ResolveNumChannelsFromFormat(SourceFormatName);
+		}
+		else
+		{
+			// Discrete path: channel count equals number of audio output vertices.
+			NumChannels = InInitParams.AudioOutputNames.Num();
+		}
 		NumSamplesPerExecute = NumChannels * NumFramesPerExecute;
 
 #if ENABLE_METASOUND_GENERATOR_INSTANCE_COUNTING
diff --git a/Engine/Plugins/Runtime/Metasound/Source/MetasoundGenerator/Public/MetasoundGenerator.h b/Engine/Plugins/Runtime/Metasound/Source/MetasoundGenerator/Public/MetasoundGenerator.h
--- a/Engine/Plugins/Runtime/Metasound/Source/MetasoundGenerator/Public/MetasoundGenerator.h
+++ b/Engine/Plugins/Runtime/Metasound/Source/MetasoundGenerator/Public/MetasoundGenerator.h
@@ -124,11 +124,11 @@ namespace Metasound
 		TSharedPtr<FGraphRenderCost> GraphRenderCost;
 		FName ClassName;
 		FTopLevelAssetPath AssetPath;
-		int32 NumChannels = 0;
 
 		// CAT format identity (e.g. "Cat:Mono1Dot0"). Empty for non-CAT sources.
-		// Carried to the generator for Environment population (graph-side format identity).
-		// Mixer output format is always resolved from NumChannels via FindDiscreteFormatFromNumChannels.
+		// The generator derives the mixer-compatible channel count internally:
+		//   - CAT path:     from SourceFormatName via channel registry lookup + clamp.
+		//   - Discrete path: from AudioOutputNames.Num().
 		FName SourceFormatName;
 
 		static UE_API void Reset(FGeneratorInitParams& InParams);
