[MetaSound] Fix AD Envelope attack phase skipped when Attack Time driven by Value (Float) node with Init Value of 0

#Synced-CL 51959729

diff --git a/Engine/Plugins/Runtime/Metasound/Source/MetasoundStandardNodes/Private/MetasoundADEnvelopeNode.cpp b/Engine/Plugins/Runtime/Metasound/Source/MetasoundStandardNodes/Private/MetasoundADEnvelopeNode.cpp
--- a/Engine/Plugins/Runtime/Metasound/Source/MetasoundStandardNodes/Private/MetasoundADEnvelopeNode.cpp
+++ b/Engine/Plugins/Runtime/Metasound/Source/MetasoundStandardNodes/Private/MetasoundADEnvelopeNode.cpp
@@ -362,12 +362,6 @@ namespace Metasound
 			EnvState.DecayCurveFactor = FMath::Max(KINDA_SMALL_NUMBER, *DecayCurveFactor);
 			EnvState.bLooping = *bLooping;
 			EnvState.bHardReset = *bHardReset;
-
-			// if there is no attack phase, we jump to 1.f on the first frame
-			if(EnvState.AttackSampleCount == 0)
-			{
-				EnvState.CurrentEnvelopeValue = 1.f;
-			}
 		}
 
 		void Reset(const IOperator::FResetParams& InParams)
