Document Configurator & Associated Code Reflection Updates
- Add Variable Node Support
- Add Page Support by Name

#rnx

#Synced-CL 52327832

diff --git a/Engine/Plugins/Runtime/Metasound/Source/MetasoundEditor/Private/MetasoundEditorDocumentClipboardUtils.cpp b/Engine/Plugins/Runtime/Metasound/Source/MetasoundEditor/Private/MetasoundEditorDocumentClipboardUtils.cpp
--- a/Engine/Plugins/Runtime/Metasound/Source/MetasoundEditor/Private/MetasoundEditorDocumentClipboardUtils.cpp
+++ b/Engine/Plugins/Runtime/Metasound/Source/MetasoundEditor/Private/MetasoundEditorDocumentClipboardUtils.cpp
@@ -668,9 +668,10 @@ namespace Metasound::Editor
 		using namespace Frontend;
 
 		FMetasoundFrontendGraphLayer Layer;
-		Layer.Graph.PageID = FGuid::NewGuid();
+		Layer.Graph.PageID = Builder.GetBuildPageID();
 
 		TSet<FGuid> NodeIDs;
+		TSet<FGuid> AddedVariableIDs;
 		TMap<FGuid, FMetasoundFrontendClass> DependenciesToAdd;
 		for (UObject* Object : Selection)
 		{
@@ -735,17 +736,49 @@ namespace Metasound::Editor
 					}
 				}
 			}
+			else if (const UMetasoundEditorGraphVariableNode* VariableNode = Cast<UMetasoundEditorGraphVariableNode>(Object))
+			{
+				const FGuid VariableNodeID = VariableNode->GetNodeID();
+				if (const FMetasoundFrontendNode* FrontendNode = Builder.FindNode(VariableNodeID))
+				{
+					if (const FMetasoundFrontendClass* Dependency = Builder.FindDependency(FrontendNode->ClassID))
+					{
+						Layer.Graph.Nodes.Add(*FrontendNode);
+						NodeIDs.Add(FrontendNode->GetID());
+						DependenciesToAdd.Add(Dependency->ID, *Dependency);
+
+						// Add the associated variable definition (once per variable)
+						if (const FMetasoundFrontendVariable* Variable = Builder.FindGraphVariableByNodeID(VariableNodeID))
+						{
+							if (!AddedVariableIDs.Contains(Variable->ID))
+							{
+								AddedVariableIDs.Add(Variable->ID);
+								Layer.Graph.Variables.Add(*Variable);
+
+								// Also add the hidden variable node and its dependency if not already present
+								if (Variable->VariableNodeID.IsValid() && !NodeIDs.Contains(Variable->VariableNodeID))
+								{
+									if (const FMetasoundFrontendNode* VarNode = Builder.FindNode(Variable->VariableNodeID))
+									{
+										Layer.Graph.Nodes.Add(*VarNode);
+										NodeIDs.Add(VarNode->GetID());
+										if (const FMetasoundFrontendClass* VarDep = Builder.FindDependency(VarNode->ClassID))
+										{
+											DependenciesToAdd.Add(VarDep->ID, *VarDep);
+										}
+									}
+								}
+							}
+						}
+					}
+				}
+			}
 		}
 
 		for (const UObject* Object : Selection)
 		{
 			if (const UMetasoundEditorGraphNode* EditorNode = Cast<const UMetasoundEditorGraphNode>(Object))
 			{
-				if (EditorNode->IsA<UMetasoundEditorGraphVariableNode>())
-				{
-					continue;
-				}
-
 				const FMetasoundFrontendNode& FrontendNode = EditorNode->GetFrontendNodeChecked();
 				for (const FMetasoundFrontendVertex& Vertex : FrontendNode.Interface.Inputs)
 				{
diff --git a/Engine/Plugins/Runtime/Metasound/Source/MetasoundEditor/Private/MetasoundNodeDetailCustomization.cpp b/Engine/Plugins/Runtime/Metasound/Source/MetasoundEditor/Private/MetasoundNodeDetailCustomization.cpp
--- a/Engine/Plugins/Runtime/Metasound/Source/MetasoundEditor/Private/MetasoundNodeDetailCustomization.cpp
+++ b/Engine/Plugins/Runtime/Metasound/Source/MetasoundEditor/Private/MetasoundNodeDetailCustomization.cpp
@@ -1107,7 +1107,7 @@ PRAGMA_ENABLE_DEPRECATION_WARNINGS
 					TSharedPtr<IPropertyHandle> PagedGraphsHandle = RootGraphHandle->GetChildHandle(FName(TEXT("PagedGraphs")));
 
 					const FGuid& PageID = DocBuilder.GetBuildPageID();
-					const int32 PageIndex = DocBuilder.FindPageIndex(PageID);
+					const int32 PageIndex = DocBuilder.FindGraphPageIndex(PageID);
 					TSharedRef<IPropertyHandle> GraphHandle = PagedGraphsHandle->AsArray()->GetElement(PageIndex);
 					TSharedPtr<IPropertyHandle> NodesHandle = GraphHandle->GetChildHandle(GET_MEMBER_NAME_CHECKED(FMetasoundFrontendGraph, Nodes));
 
diff --git a/Engine/Plugins/Runtime/Metasound/Source/MetasoundEngine/Private/DocumentTemplates/MetasoundDocumentConfigurator.cpp b/Engine/Plugins/Runtime/Metasound/Source/MetasoundEngine/Private/DocumentTemplates/MetasoundDocumentConfigurator.cpp
--- a/Engine/Plugins/Runtime/Metasound/Source/MetasoundEngine/Private/DocumentTemplates/MetasoundDocumentConfigurator.cpp
+++ b/Engine/Plugins/Runtime/Metasound/Source/MetasoundEngine/Private/DocumentTemplates/MetasoundDocumentConfigurator.cpp
@@ -12,6 +12,7 @@
 #include "MetasoundFrontendDocumentBuilder.h"
 #include "MetasoundFrontendLiteral.h"
 #include "MetasoundFrontendSearchEngine.h"
+#include "MetasoundSettings.h"
 #include "NodeTemplates/MetasoundFrontendNodeTemplateInput.h"
 #include "UObject/SoftObjectPath.h"
 
@@ -101,6 +102,11 @@ namespace Metasound::Engine
 			TStringBuilder<0> InputNodeCode; // Input Template nodes
 #endif // WITH_EDITORONLY_DATA
 
+			TStringBuilder<0> VariableCode;
+			TStringBuilder<0> VariableSetNodeCode;
+			TStringBuilder<0> VariableGetNodeCode;
+			TStringBuilder<0> VariableGetDelayedNodeCode;
+
 			TStringBuilder<0> OutputCode;
 			TStringBuilder<0> OutputEdgesCode;
 
@@ -258,6 +264,119 @@ namespace Metasound::Engine
 				}
 			}
 
+			void ReflectVariables()
+			{
+				if (Layer.Graph.Variables.IsEmpty())
+				{
+					return;
+				}
+
+				VariableCode = TEXT("\n// Variables\n");
+				VariableCode += TEXT(".Add<FDC::FVariable>(\n{\n");
+
+				for (const FMetasoundFrontendVariable& Variable : Layer.Graph.Variables)
+				{
+					// Configurations/MetaSound Builder Subsystem currently do not support pages,
+					// so attempt to harvest literal with Default PageID.
+					VariableCode += TEXT("\t{ ");
+					{
+						VariableCode.Appendf(TEXT(".Name = \"%s\", "), *Variable.Name.ToString());
+						VariableCode.Appendf(TEXT(".DataType = \"%s\", "), *Variable.TypeName.ToString());
+						VariableCode.Appendf(TEXT(".DefaultValue = ")); FormatDefaultCode(Variable.Literal, VariableCode);
+					}
+					VariableCode += TEXT("},\n");
+				}
+
+				VariableCode += TEXT("})");
+
+				// Build a map from variable node IDs to variable names for use in variable node reflection
+				bool bHasSetNodes = false;
+				bool bHasGetNodes = false;
+				bool bHasGetDelayedNodes = false;
+
+				VariableSetNodeCode = TEXT("\n// Variable Set Nodes\n");
+				VariableSetNodeCode += TEXT(".Add<FDC::FVariableSetNode>(\n{\n");
+
+				VariableGetNodeCode = TEXT("\n// Variable Get Nodes\n");
+				VariableGetNodeCode += TEXT(".Add<FDC::FVariableGetNode>(\n{\n");
+
+				VariableGetDelayedNodeCode = TEXT("\n// Variable Get Delayed Nodes\n");
+				VariableGetDelayedNodeCode += TEXT(".Add<FDC::FVariableGetDelayedNode>(\n{\n");
+
+				auto CacheNodeInfo = [this](const FString& BaseName, const FMetasoundFrontendNode& Node)
+				{
+					int32 Index = 0;
+					FString FullName = BaseName;
+					while (NodeInstanceNames.Contains(FullName))
+					{
+						FullName = FString::Printf(TEXT("%s_%i"), *BaseName, ++Index);
+					}
+					NodeInstanceNames.Add(FullName);
+					NodeIDToInfo.Add(Node.GetID(), FConfigNodeInfo { .InstanceName = FullName, .Node = &Node });
+					return FullName;
+				};
+
+				for (const FMetasoundFrontendVariable& Variable : Layer.Graph.Variables)
+				{
+					auto ReflectVariableNode = [&](const FGuid& InNodeID, FStringBuilderBase& OutCode, bool& bOutHasNodes, const TCHAR* StructName)
+					{
+						auto IsNodeWithID = [&InNodeID](const FMetasoundFrontendNode& Node) { return Node.GetID() == InNodeID; };
+						if (const FMetasoundFrontendNode* Node = Layer.Graph.Nodes.FindByPredicate(IsNodeWithID))
+						{
+							FVector2D Location = FVector2D::Zero();
+							if (!Node->Style.Display.Locations.IsEmpty())
+							{
+								Location = Node->Style.Display.Locations.Array().Last().Value;
+							}
+
+							const FString NodeName = CacheNodeInfo(FString::Printf(TEXT("%s_%s"), StructName, *Variable.Name.ToString()), *Node);
+							bOutHasNodes = true;
+							OutCode.Appendf(
+								TEXT("\t{ .Name = \"%s\", .VariableName = \"%s\", .Location = FVector2D(%sf, %sf) },\n"),
+								*NodeName,
+								*Variable.Name.ToString(),
+								*FString::SanitizeFloat(Location.X),
+								*FString::SanitizeFloat(Location.Y)
+							);
+						}
+					};
+
+					// Mutator (Set) node
+					if (Variable.MutatorNodeID.IsValid())
+					{
+						ReflectVariableNode(Variable.MutatorNodeID, VariableSetNodeCode, bHasSetNodes, TEXT("Set"));
+					}
+
+					// Accessor (Get) nodes
+					for (const FGuid& AccessorID : Variable.AccessorNodeIDs)
+					{
+						ReflectVariableNode(AccessorID, VariableGetNodeCode, bHasGetNodes, TEXT("Get"));
+					}
+
+					// Deferred Accessor (Get Delayed) nodes
+					for (const FGuid& DeferredID : Variable.DeferredAccessorNodeIDs)
+					{
+						ReflectVariableNode(DeferredID, VariableGetDelayedNodeCode, bHasGetDelayedNodes, TEXT("GetDelayed"));
+					}
+				}
+
+				auto FinishCode = [](bool bHadContent, FStringBuilderBase& CodeBuilder)
+				{
+					if (bHadContent)
+					{
+						CodeBuilder += TEXT("})");
+					}
+					else
+					{
+						CodeBuilder = { };
+					}
+				};
+
+				FinishCode(bHasSetNodes, VariableSetNodeCode);
+				FinishCode(bHasGetNodes, VariableGetNodeCode);
+				FinishCode(bHasGetDelayedNodes, VariableGetDelayedNodeCode);
+			}
+
 			// Edges from inputs are not distinct in reflected code as inputs are characterized
 			// visually by template nodes which are included in the typical edge addition case.
 			void ReflectEdges()
@@ -636,6 +755,15 @@ namespace Metasound::Engine
 							}
 							break;
 
+							case EMetasoundFrontendClassType::Variable:
+							case EMetasoundFrontendClassType::VariableAccessor:
+							case EMetasoundFrontendClassType::VariableDeferredAccessor:
+							case EMetasoundFrontendClassType::VariableMutator:
+							{
+								// Variable nodes are reflected via ReflectVariables, not here.
+							}
+							break;
+
 							default:
 							{
 								break;
@@ -707,6 +835,15 @@ namespace Metasound::Engine
 
 			FString Execute()
 			{
+				FString SetBuildPageCode;
+				{
+					const FName PageName = UMetaSoundSettings::GetPageName(Layer.Graph.PageID);
+					if (PageName != Frontend::DefaultPageName)
+					{
+						SetBuildPageCode = FString::Printf(TEXT("\n// Page Setup\n.SetBuildPage(\"%s\")\n"), *PageName.ToString());
+					}
+				}
+
 				CacheDependencies();
 
 #if WITH_EDITORONLY_DATA
@@ -721,6 +858,7 @@ namespace Metasound::Engine
 
 				ReflectInputs();
 				ReflectOutputs();
+				ReflectVariables();
 				ReflectEdges();
 
 				ReflectNodeInputDefaults();
@@ -745,10 +883,20 @@ namespace Metasound::Engine
 				LayerCode.Append(TEXT("FDC Configurator(FDC::FArgs { .Builder = &OutBuilder });\n"));
 				LayerCode.Append(TEXT("return Configurator\n"));
 
+				if (!SetBuildPageCode.IsEmpty())
+				{
+					LayerCode.Append(SetBuildPageCode);
+				}
+
 				AddSection(InterfaceCode);
 				AddSection(InputCode);
 				AddSection(OutputCode);
 
+				AddSection(VariableCode);
+				AddSection(VariableSetNodeCode);
+				AddSection(VariableGetNodeCode);
+				AddSection(VariableGetDelayedNodeCode);
+
 				AddSection(InputNodeCode);
 				AddSection(NodeCode);
 				AddSection(NodeInputDefaultsCode);
@@ -762,6 +910,13 @@ namespace Metasound::Engine
 		};
 	} // namespace ConfigCodeReflectorPrivate
 
+#if WITH_EDITORONLY_DATA
+	FDocumentConfigurator::FDocumentConfigurator(UMetaSoundBuilderBase& InBuilder)
+	{
+		Builder.Reset(&InBuilder);
+	}
+#endif // WITH_EDITORONLY_DATA
+
 	FDocumentConfigurator::FDocumentConfigurator(const FDocumentConfigurator::FArgs& Args)
 #if WITH_EDITORONLY_DATA
 		: Offset(Args.Offset)
@@ -814,14 +969,11 @@ namespace Metasound::Engine
 		if (StepResult == EMetaSoundBuilderResult::Succeeded)
 		{
 			Builder->SetNodeLocation(OutputHandle.NodeID, Input.Location + Offset, StepResult);
-			if (StepResult == EMetaSoundBuilderResult::Failed)
-			{
-				Result = StepResult;
-			}
 		}
-		else
+
+		if (Result == EMetaSoundBuilderResult::Succeeded)
 		{
-			Result = EMetaSoundBuilderResult::Failed;
+			Result = StepResult;
 		}
 
 		return *this;
@@ -830,6 +982,8 @@ namespace Metasound::Engine
 #if WITH_EDITORONLY_DATA
 	FDocumentConfigurator& FDocumentConfigurator::Add(FDocumentConfigurator::FInputNode InputNode)
 	{
+		using namespace Frontend;
+
 		EMetaSoundBuilderResult StepResult = EMetaSoundBuilderResult::Failed;
 
 		FMetaSoundNodeHandle NodeHandle;
@@ -842,14 +996,12 @@ namespace Metasound::Engine
 		}
 		else
 		{
-			using namespace Frontend;
-
 			if (const INodeTemplate* Template = INodeTemplateRegistry::Get().FindTemplate(FInputNodeTemplate::ClassName); ensure(Template))
 			{
 				if (const FMetasoundFrontendClassInput* Input = Builder->GetConstBuilder().FindGraphInput(InputNode.InputName))
 				{
 					FNodeTemplateGenerateInterfaceParams Params { .InputsToConnect = { Input->TypeName }, .OutputsToConnect = { Input->TypeName } };
-					if (const FMetasoundFrontendNode* TemplateNode = Builder->GetBuilder().AddNodeByTemplate(*Template, MoveTemp(Params)))
+					if (const FMetasoundFrontendNode* TemplateNode = Builder->GetBuilder().AddNodeByTemplate(*Template, MoveTemp(Params), FGuid::NewGuid()))
 					{
 						Builder->ConnectGraphInputToNode(InputNode.InputName, FMetaSoundNodeHandle(TemplateNode->GetID()), "Value", StepResult);
 						if (StepResult == EMetaSoundBuilderResult::Succeeded)
@@ -893,6 +1045,103 @@ namespace Metasound::Engine
 		return *this;
 	}
 
+	FDocumentConfigurator& FDocumentConfigurator::Add(FDocumentConfigurator::FVariable Variable)
+	{
+		const FMetasoundFrontendVariable* NewVariable = Builder->GetBuilder().AddGraphVariable(Variable.Name, Variable.DataType, &Variable.DefaultValue);
+
+		if (!NewVariable)
+		{
+			Result = EMetaSoundBuilderResult::Failed;
+		}
+
+		return *this;
+	}
+
+	FDocumentConfigurator& FDocumentConfigurator::Add(FDocumentConfigurator::FVariableSetNode Node)
+	{
+		EMetaSoundBuilderResult StepResult = EMetaSoundBuilderResult::Failed;
+
+		if (NodeNameToHandle.Contains(Node.Name))
+		{
+			UE_LOG(LogMetaSound, Error, TEXT("Node with name '%s' already exists. Failed to add variable set node to MetaSound graph"), *Node.Name.ToString());
+		}
+		else
+		{
+			if (const FMetasoundFrontendNode* MutatorNode = Builder->GetBuilder().AddGraphVariableMutatorNode(Node.VariableName, FGuid::NewGuid()))
+			{
+				FMetaSoundNodeHandle NodeHandle;
+				NodeHandle.NodeID = MutatorNode->GetID();
+				NodeNameToHandle.Add(Node.Name, NodeHandle);
+				Builder->SetNodeLocation(NodeHandle, Node.Location + Offset, StepResult);
+			}
+			else
+			{
+				StepResult = EMetaSoundBuilderResult::Failed;
+			}
+		}
+
+		if (Result == EMetaSoundBuilderResult::Succeeded)
+		{
+			Result = StepResult;
+		}
+
+		return *this;
+	}
+
+	FDocumentConfigurator& FDocumentConfigurator::Add(FDocumentConfigurator::FVariableGetNode Node)
+	{
+		EMetaSoundBuilderResult StepResult = EMetaSoundBuilderResult::Failed;
+
+		if (NodeNameToHandle.Contains(Node.Name))
+		{
+			UE_LOG(LogMetaSound, Error, TEXT("Node with name '%s' already exists. Failed to add variable get node to MetaSound graph"), *Node.Name.ToString());
+		}
+		else
+		{
+			if (const FMetasoundFrontendNode* AccessorNode = Builder->GetBuilder().AddGraphVariableAccessorNode(Node.VariableName, FGuid::NewGuid()))
+			{
+				FMetaSoundNodeHandle NodeHandle;
+				NodeHandle.NodeID = AccessorNode->GetID();
+				NodeNameToHandle.Add(Node.Name, NodeHandle);
+				Builder->SetNodeLocation(NodeHandle, Node.Location + Offset, StepResult);
+			}
+		}
+
+		if (Result == EMetaSoundBuilderResult::Succeeded)
+		{
+			Result = StepResult;
+		}
+
+		return *this;
+	}
+
+	FDocumentConfigurator& FDocumentConfigurator::Add(FDocumentConfigurator::FVariableGetDelayedNode Node)
+	{
+		EMetaSoundBuilderResult StepResult = EMetaSoundBuilderResult::Failed;
+
+		if (NodeNameToHandle.Contains(Node.Name))
+		{
+			UE_LOG(LogMetaSound, Error, TEXT("Node with name '%s' already exists. Failed to add variable get delayed node to MetaSound graph"), *Node.Name.ToString());
+		}
+		else
+		{
+			if (const FMetasoundFrontendNode* AccessorNode = Builder->GetBuilder().AddGraphVariableDeferredAccessorNode(Node.VariableName, FGuid::NewGuid()))
+			{
+				FMetaSoundNodeHandle NodeHandle;
+				NodeHandle.NodeID = AccessorNode->GetID();
+				NodeNameToHandle.Add(Node.Name, NodeHandle);
+				Builder->SetNodeLocation(NodeHandle, Node.Location + Offset, StepResult);
+			}
+		}
+
+		if (Result == EMetaSoundBuilderResult::Succeeded)
+		{
+			Result = StepResult;
+		}
+
+		return *this;
+	}
+
 	FDocumentConfigurator& FDocumentConfigurator::Add(FDocumentConfigurator::FNode Node)
 	{
 		EMetaSoundBuilderResult StepResult = EMetaSoundBuilderResult::Failed;
@@ -905,7 +1154,11 @@ namespace Metasound::Engine
 		}
 		else if (Node.ClassName.IsValid())
 		{
-			NodeHandle = Builder->AddNodeByClassName(Node.ClassName, StepResult);
+			if (const FMetasoundFrontendNode* NewNode = Builder->GetBuilder().AddNodeByClassName(Node.ClassName, 1, FGuid::NewGuid()))
+			{
+				NodeHandle.NodeID = NewNode->GetID();
+				StepResult = EMetaSoundBuilderResult::Succeeded;
+			}
 		}
 		else if (Node.AssetPath.IsValid())
 		{
@@ -932,13 +1185,19 @@ namespace Metasound::Engine
 
 	FDocumentConfigurator& FDocumentConfigurator::Add(FDocumentConfigurator::FEdge Edge)
 	{
+		using namespace Frontend;
+
 		EMetaSoundBuilderResult StepResult = EMetaSoundBuilderResult::Failed;
 
 		if (const FMetaSoundNodeHandle* OutputNodeHandle = NodeNameToHandle.Find(Edge.OutputNode))
 		{
 			if (const FMetaSoundNodeHandle* InputNodeHandle = NodeNameToHandle.Find(Edge.InputNode))
 			{
-				Builder->ConnectNodes(*OutputNodeHandle, Edge.Output, *InputNodeHandle, Edge.Input, StepResult);
+				TSet<FNamedEdge> Edges;
+				Edges.Add(FNamedEdge { OutputNodeHandle->NodeID, Edge.Output, InputNodeHandle->NodeID, Edge.Input });
+				StepResult = Builder->GetBuilder().AddNamedEdges(Edges, nullptr, true)
+					? EMetaSoundBuilderResult::Succeeded
+					: EMetaSoundBuilderResult::Failed;
 			}
 			else
 			{
@@ -964,12 +1223,18 @@ namespace Metasound::Engine
 	{
 		EMetaSoundBuilderResult StepResult = EMetaSoundBuilderResult::Failed;
 
-		if (const FMetaSoundNodeHandle* OutputNodeHandle = NodeNameToHandle.Find(Edge.Node))
+		if (const FMetaSoundNodeHandle* DestNodeHandle = NodeNameToHandle.Find(Edge.Node))
 		{
-			Builder->ConnectGraphInputToNode(Edge.GraphInput, *OutputNodeHandle, Edge.NodeInput, StepResult);
-			if (Result == EMetaSoundBuilderResult::Succeeded)
+			const FMetaSoundFrontendDocumentBuilder& ConstBuilder = Builder->GetConstBuilder();
+			if (const FMetasoundFrontendNode* GraphInputNode = ConstBuilder.FindGraphInputNode(Edge.GraphInput))
 			{
-				Result = StepResult;
+				using namespace Frontend;
+				const FMetasoundFrontendVertex& OutputVertex = GraphInputNode->Interface.Outputs.Last();
+				TSet<FNamedEdge> Edges;
+				Edges.Add(FNamedEdge { GraphInputNode->GetID(), OutputVertex.Name, DestNodeHandle->NodeID, Edge.NodeInput });
+				StepResult = Builder->GetBuilder().AddNamedEdges(Edges, nullptr, true)
+					? EMetaSoundBuilderResult::Succeeded
+					: EMetaSoundBuilderResult::Failed;
 			}
 		}
 		else
@@ -988,19 +1253,26 @@ namespace Metasound::Engine
 
 	FDocumentConfigurator& FDocumentConfigurator::Add(FDocumentConfigurator::FEdgeToOutput Edge)
 	{
+		using namespace Frontend;
+
 		EMetaSoundBuilderResult StepResult = EMetaSoundBuilderResult::Failed;
 
-		if (const FMetaSoundNodeHandle* OutputNodeHandle = NodeNameToHandle.Find(Edge.Node))
+		if (const FMetaSoundNodeHandle* SourceNodeHandle = NodeNameToHandle.Find(Edge.Node))
 		{
-			Builder->ConnectNodeToGraphOutput(*OutputNodeHandle, Edge.NodeOutput, Edge.GraphOutput, StepResult);
-			if (Result == EMetaSoundBuilderResult::Succeeded)
+			const FMetaSoundFrontendDocumentBuilder& ConstBuilder = Builder->GetConstBuilder();
+			if (const FMetasoundFrontendNode* GraphOutputNode = ConstBuilder.FindGraphOutputNode(Edge.GraphOutput))
 			{
-				Result = StepResult;
+				const FMetasoundFrontendVertex& InputVertex = GraphOutputNode->Interface.Inputs.Last();
+				TSet<FNamedEdge> Edges;
+				Edges.Add(FNamedEdge { SourceNodeHandle->NodeID, Edge.NodeOutput, GraphOutputNode->GetID(), InputVertex.Name });
+				StepResult = Builder->GetBuilder().AddNamedEdges(Edges, nullptr, true)
+					? EMetaSoundBuilderResult::Succeeded
+					: EMetaSoundBuilderResult::Failed;
 			}
 		}
 		else
 		{
-			UE_LOG(LogMetaSound, Error, TEXT("Failed to find node with name '%s': Graph output to node intput connection failed"), *Edge.Node.ToString());
+			UE_LOG(LogMetaSound, Error, TEXT("Failed to find node with name '%s': Graph output to node input connection failed"), *Edge.Node.ToString());
 			StepResult = EMetaSoundBuilderResult::Failed;
 		}
 
@@ -1056,6 +1328,21 @@ namespace Metasound::Engine
 		return *this;
 	}
 
+	FDocumentConfigurator& FDocumentConfigurator::SetBuildPage(FName PageName)
+	{
+		if (const FGuid* PageID = UMetaSoundSettings::GetPageID(PageName))
+		{
+			Builder->GetBuilder().SetBuildPageID(*PageID);
+		}
+		else
+		{
+			UE_LOG(LogMetaSound, Error, TEXT("FDocumentConfigurator::SetBuildPage: Page '%s' not found in MetaSound settings"), *PageName.ToString());
+			Result = EMetaSoundBuilderResult::Failed;
+		}
+
+		return *this;
+	}
+
 	FDocumentConfigurator& FDocumentConfigurator::Set(FDocumentConfigurator::FInputLocation Location)
 	{
 		EMetaSoundBuilderResult StepResult = EMetaSoundBuilderResult::Failed;
diff --git a/Engine/Plugins/Runtime/Metasound/Source/MetasoundEngine/Private/MetasoundBuilderBase.cpp b/Engine/Plugins/Runtime/Metasound/Source/MetasoundEngine/Private/MetasoundBuilderBase.cpp
--- a/Engine/Plugins/Runtime/Metasound/Source/MetasoundEngine/Private/MetasoundBuilderBase.cpp
+++ b/Engine/Plugins/Runtime/Metasound/Source/MetasoundEngine/Private/MetasoundBuilderBase.cpp
@@ -164,9 +164,12 @@ void UMetaSoundBuilderBase::AddGraphPage(FName PageName, bool bDuplicateLastGrap
 	{
 		if (const FMetaSoundPageSettings* PageSettings = Settings->FindPageSettings(PageName))
 		{
-			Builder.AddGraphPage(PageSettings->UniqueId, bDuplicateLastGraph, bSetAsBuildGraph);
-			OutResult = EMetaSoundBuilderResult::Succeeded;
-			return;
+			if (Builder.FindGraphPageIndex(PageSettings->UniqueId) == INDEX_NONE)
+			{
+				Builder.AddGraphPage(PageSettings->UniqueId, bDuplicateLastGraph, bSetAsBuildGraph);
+				OutResult = EMetaSoundBuilderResult::Succeeded;
+				return;
+			}
 		}
 	}
 
diff --git a/Engine/Plugins/Runtime/Metasound/Source/MetasoundEngine/Public/DocumentTemplates/MetasoundDocumentConfigurator.h b/Engine/Plugins/Runtime/Metasound/Source/MetasoundEngine/Public/DocumentTemplates/MetasoundDocumentConfigurator.h
--- a/Engine/Plugins/Runtime/Metasound/Source/MetasoundEngine/Public/DocumentTemplates/MetasoundDocumentConfigurator.h
+++ b/Engine/Plugins/Runtime/Metasound/Source/MetasoundEngine/Public/DocumentTemplates/MetasoundDocumentConfigurator.h
@@ -30,6 +30,12 @@ namespace Metasound::Engine
 			bool bResetDoc = true;
 		};
 
+#if WITH_EDITORONLY_DATA
+		// Constructor exposed to enable testing configurator that wraps a
+		// transient builder. Not for use with runtime or editor behavior.
+		UE_INTERNAL UE_API FDocumentConfigurator(UMetaSoundBuilderBase& InBuilder);
+#endif // WITH_EDITORONLY_DATA
+
 		UE_API FDocumentConfigurator(const FArgs& Args);
 		UE_API ~FDocumentConfigurator();
 
@@ -79,6 +85,42 @@ namespace Metasound::Engine
 		UE_EXPERIMENTAL(5.8, "Configurator is in active development and subject to change")
 		UE_API FDocumentConfigurator& Add(FOutput Output);
 
+		struct FVariable
+		{
+			FName Name = "FloatVariable";
+			FName DataType = GetMetasoundDataTypeName<float>();
+			FMetasoundFrontendLiteral DefaultValue = Frontend::MakeFrontendLiteral(0.f);
+		};
+		UE_EXPERIMENTAL(5.8, "Configurator is in active development and subject to change")
+		UE_API FDocumentConfigurator& Add(FVariable Variable);
+
+		struct FVariableSetNode
+		{
+			FName Name;
+			FName VariableName;
+			FVector2D Location = FVector2D::Zero();
+		};
+		UE_EXPERIMENTAL(5.8, "Configurator is in active development and subject to change")
+		UE_API FDocumentConfigurator& Add(FVariableSetNode Node);
+
+		struct FVariableGetNode
+		{
+			FName Name;
+			FName VariableName;
+			FVector2D Location = FVector2D::Zero();
+		};
+		UE_EXPERIMENTAL(5.8, "Configurator is in active development and subject to change")
+		UE_API FDocumentConfigurator& Add(FVariableGetNode Node);
+
+		struct FVariableGetDelayedNode
+		{
+			FName Name;
+			FName VariableName;
+			FVector2D Location = FVector2D::Zero();
+		};
+		UE_EXPERIMENTAL(5.8, "Configurator is in active development and subject to change")
+		UE_API FDocumentConfigurator& Add(FVariableGetDelayedNode Node);
+
 		struct FNode
 		{
 			FName Name;
@@ -140,6 +182,9 @@ namespace Metasound::Engine
 		UE_EXPERIMENTAL(5.8, "Configurator is in active development and subject to change")
 		UE_API FDocumentConfigurator& ResetDocument();
 
+		UE_EXPERIMENTAL(5.8, "Configurator is in active development and subject to change")
+		UE_API FDocumentConfigurator& SetBuildPage(FName PageName);
+
 		struct FInputLocation
 		{
 			FName Name;
diff --git a/Engine/Plugins/Runtime/Metasound/Source/MetasoundEngineTest/Private/EngineTestDocumentConfiguratorReflect.cpp b/Engine/Plugins/Runtime/Metasound/Source/MetasoundEngineTest/Private/EngineTestDocumentConfiguratorReflect.cpp
new file mode 100644
--- /dev/null
+++ b/Engine/Plugins/Runtime/Metasound/Source/MetasoundEngineTest/Private/EngineTestDocumentConfiguratorReflect.cpp
@@ -0,0 +1,461 @@
+// Copyright Epic Games, Inc. All Rights Reserved.
+
+#include "DocumentTemplates/MetasoundDocumentConfigurator.h"
+#include "EngineTestMetaSoundBuilder.h"
+#include "MetasoundBuilderSubsystem.h"
+#include "MetasoundFrontendDocument.h"
+#include "MetasoundFrontendDocumentBuilder.h"
+#include "MetasoundFrontendGraphLayer.h"
+#include "MetasoundFrontendLiteral.h"
+#include "MetasoundSettings.h"
+#include "Misc/AutomationTest.h"
+
+#if WITH_DEV_AUTOMATION_TESTS && WITH_EDITORONLY_DATA
+
+namespace Metasound::Test::ConfiguratorReflect
+{
+	// Builds a FMetasoundFrontendGraphLayer from a builder's current document state.
+	// This mirrors what the editor's CopySelectionToLayer does but without requiring editor graph objects.
+	FMetasoundFrontendGraphLayer BuildLayerFromBuilder(const UMetaSoundBuilderBase& Builder)
+	{
+		const FMetaSoundFrontendDocumentBuilder& DocBuilder = Builder.GetConstBuilder();
+		const FMetasoundFrontendDocument& Document = DocBuilder.GetConstDocumentChecked();
+		const FGuid& PageID = DocBuilder.GetBuildPageID();
+
+		FMetasoundFrontendGraphLayer Layer;
+
+		const FMetasoundFrontendGraph* Graph = Document.RootGraph.FindConstGraph(PageID);
+		if (!Graph)
+		{
+			return Layer;
+		}
+
+		Layer.Graph = *Graph;
+
+		// Copy inputs and outputs from the root graph class interface
+		Layer.Inputs = Document.RootGraph.GetDefaultInterface().Inputs;
+		Layer.Outputs = Document.RootGraph.GetDefaultInterface().Outputs;
+
+		// Gather dependencies for all nodes in the graph
+		TMap<FGuid, FMetasoundFrontendClass> DependencyMap;
+		for (const FMetasoundFrontendNode& Node : Layer.Graph.Nodes)
+		{
+			if (const FMetasoundFrontendClass* Dep = DocBuilder.FindDependency(Node.ClassID))
+			{
+				DependencyMap.Add(Dep->ID, *Dep);
+			}
+		}
+		DependencyMap.GenerateValueArray(Layer.Dependencies);
+
+		return Layer;
+	}
+
+	// Validates that reflected code contains required boilerplate and all expected fragments.
+	bool ValidateReflectedCode(FAutomationTestBase& Test, const FString& Code, const TArray<FString>& ExpectedFragments, const TArray<FString>& UnexpectedFragments = { })
+	{
+		bool bPassed = true;
+
+		// Boilerplate checks
+		bPassed &= Test.TestTrue(TEXT("Code contains FDC alias"), Code.Contains(TEXT("using FDC = Metasound::Engine::FDocumentConfigurator")));
+		bPassed &= Test.TestTrue(TEXT("Code contains Configurator construction"), Code.Contains(TEXT("FDC Configurator(FDC::FArgs")));
+		bPassed &= Test.TestTrue(TEXT("Code contains .Succeeded()"), Code.Contains(TEXT(".Succeeded()")));
+
+		for (const FString& Fragment : ExpectedFragments)
+		{
+			bPassed &= Test.TestTrue(*FString::Printf(TEXT("Code contains: %s"), *Fragment), Code.Contains(Fragment));
+		}
+
+		for (const FString& Fragment : UnexpectedFragments)
+		{
+			bPassed &= Test.TestFalse(*FString::Printf(TEXT("Code should NOT contain: %s"), *Fragment), Code.Contains(Fragment));
+		}
+
+		return bPassed;
+	}
+
+	// ============================================================================
+	// Test 1: Reflect empty layer
+	// ============================================================================
+	IMPLEMENT_SIMPLE_AUTOMATION_TEST(
+		FReflectEmptyLayerTest,
+		"Audio.Metasound.Configurator.Reflect.EmptyLayer",
+		EAutomationTestFlags::EditorContext | EAutomationTestFlags::EngineFilter)
+	bool FReflectEmptyLayerTest::RunTest(const FString&)
+	{
+		FMetasoundFrontendGraphLayer Layer;
+		Layer.Graph.PageID = Frontend::DefaultPageID;
+
+		FString Code;
+		Engine::FDocumentConfigurator::Reflect(MoveTemp(Layer), Code);
+
+		UTEST_FALSE("Code is not empty", Code.IsEmpty());
+
+		ValidateReflectedCode(*this, Code, { },
+		{
+			TEXT(".Add<FDC::FInput>"),
+			TEXT(".Add<FDC::FOutput>"),
+			TEXT(".Add<FDC::FNode>"),
+			TEXT(".Add<FDC::FEdge>"),
+			TEXT(".Add<FDC::FVariable>"),
+			TEXT(".SetBuildPage(")
+		});
+
+		return true;
+	}
+
+	// ============================================================================
+	// Test 2: Reflect inputs and outputs
+	// ============================================================================
+	IMPLEMENT_SIMPLE_AUTOMATION_TEST(
+		FReflectInputsOutputsTest,
+		"Audio.Metasound.Configurator.Reflect.InputsAndOutputs",
+		EAutomationTestFlags::EditorContext | EAutomationTestFlags::EngineFilter)
+	bool FReflectInputsOutputsTest::RunTest(const FString&)
+	{
+		using namespace EngineTestMetaSoundPatchBuilderPrivate;
+
+		EMetaSoundBuilderResult Result = EMetaSoundBuilderResult::Failed;
+		UMetaSoundPatchBuilder& Builder = CreatePatchBuilderChecked(*this, "ReflectInputsOutputs_Test", { });
+
+		FMetasoundFrontendLiteral GainDefault;
+		GainDefault.Set(1.0f);
+		Builder.AddGraphInputNode("Gain", GetMetasoundDataTypeName<float>(), GainDefault, Result);
+		UTEST_EQUAL("Add input succeeded", Result, EMetaSoundBuilderResult::Succeeded);
+
+		FMetasoundFrontendLiteral OutputDefault;
+		OutputDefault.Set(0.0f);
+		Builder.AddGraphOutputNode("Result", GetMetasoundDataTypeName<float>(), OutputDefault, Result);
+		UTEST_EQUAL("Add output succeeded", Result, EMetaSoundBuilderResult::Succeeded);
+
+		FMetasoundFrontendGraphLayer Layer = BuildLayerFromBuilder(Builder);
+
+		FString Code;
+		Engine::FDocumentConfigurator::Reflect(MoveTemp(Layer), Code);
+
+		ValidateReflectedCode(*this, Code,
+		{
+			TEXT(".Add<FDC::FInput>"),
+			TEXT(".Name = \"Gain\""),
+			TEXT(".bIsConstructorInput = false"),
+			TEXT(".Add<FDC::FOutput>"),
+			TEXT(".Name = \"Result\""),
+		},
+		{
+			TEXT(".Add<FDC::FVariable>"),
+			TEXT(".SetBuildPage(") // Default page should not emit PageName
+		});
+
+		return true;
+	}
+
+	// ============================================================================
+	// Test 3: Reflect nodes and edges
+	// ============================================================================
+	IMPLEMENT_SIMPLE_AUTOMATION_TEST(
+		FReflectNodesEdgesTest,
+		"Audio.Metasound.Configurator.Reflect.NodesAndEdges",
+		EAutomationTestFlags::EditorContext | EAutomationTestFlags::EngineFilter)
+	bool FReflectNodesEdgesTest::RunTest(const FString&)
+	{
+		using namespace EngineTestMetaSoundPatchBuilderPrivate;
+
+		EMetaSoundBuilderResult Result = EMetaSoundBuilderResult::Failed;
+		UMetaSoundPatchBuilder& Builder = CreatePatchBuilderChecked(*this, "ReflectNodesEdges_Test", { });
+
+		// Add Sine oscillator
+		FMetaSoundNodeHandle SineNode = Builder.AddNodeByClassName({ "UE", "Sine", "Audio" }, Result, 1);
+		UTEST_EQUAL("Add Sine node succeeded", Result, EMetaSoundBuilderResult::Succeeded);
+		UTEST_TRUE("Sine node handle valid", SineNode.IsSet());
+
+		// Add float input connected to Frequency
+		FMetasoundFrontendLiteral FreqDefault;
+		FreqDefault.Set(440.0f);
+		FMetaSoundBuilderNodeOutputHandle FreqOutput = Builder.AddGraphInputNode("Frequency", GetMetasoundDataTypeName<float>(), FreqDefault, Result);
+		UTEST_EQUAL("Add Frequency input succeeded", Result, EMetaSoundBuilderResult::Succeeded);
+
+		FMetaSoundBuilderNodeInputHandle SineFreqInput = Builder.FindNodeInputByName(SineNode, "Frequency", Result);
+		Builder.ConnectNodes(FreqOutput, SineFreqInput, Result);
+		UTEST_EQUAL("Connect Frequency succeeded", Result, EMetaSoundBuilderResult::Succeeded);
+
+		FMetasoundFrontendGraphLayer Layer = BuildLayerFromBuilder(Builder);
+
+		FString Code;
+		Engine::FDocumentConfigurator::Reflect(MoveTemp(Layer), Code);
+
+		ValidateReflectedCode(*this, Code,
+		{
+			TEXT(".Add<FDC::FNode>"),
+			TEXT(".ClassName = { \"UE\", \"Sine\", \"Audio\" }"),
+			TEXT(".Add<FDC::FInput>"),
+			TEXT(".Name = \"Frequency\""),
+		});
+
+		return true;
+	}
+
+	// ============================================================================
+	// Test 4: Reflect variables
+	// ============================================================================
+	IMPLEMENT_SIMPLE_AUTOMATION_TEST(
+		FReflectVariablesTest,
+		"Audio.Metasound.Configurator.Reflect.Variables",
+		EAutomationTestFlags::EditorContext | EAutomationTestFlags::EngineFilter)
+	bool FReflectVariablesTest::RunTest(const FString&)
+	{
+		using namespace EngineTestMetaSoundPatchBuilderPrivate;
+
+		EMetaSoundBuilderResult Result = EMetaSoundBuilderResult::Failed;
+		UMetaSoundPatchBuilder& Builder = CreatePatchBuilderChecked(*this, "ReflectVariables_Test", { });
+
+		// Add a float variable
+		FMetasoundFrontendLiteral VarDefault;
+		VarDefault.Set(0.5f);
+		Builder.AddGraphVariable("MyVar", GetMetasoundDataTypeName<float>(), VarDefault, Result);
+		UTEST_EQUAL("Add variable succeeded", Result, EMetaSoundBuilderResult::Succeeded);
+
+		// Add variable nodes
+		FMetaSoundNodeHandle SetNode = Builder.AddGraphVariableSetNode("MyVar", Result);
+		UTEST_EQUAL("Add variable set node succeeded", Result, EMetaSoundBuilderResult::Succeeded);
+		UTEST_TRUE("Set node handle valid", SetNode.IsSet());
+
+		FMetaSoundNodeHandle GetNode = Builder.AddGraphVariableGetNode("MyVar", Result);
+		UTEST_EQUAL("Add variable get node succeeded", Result, EMetaSoundBuilderResult::Succeeded);
+		UTEST_TRUE("Get node handle valid", GetNode.IsSet());
+
+		FMetaSoundNodeHandle GetDelayedNode = Builder.AddGraphVariableGetDelayedNode("MyVar", Result);
+		UTEST_EQUAL("Add variable get delayed node succeeded", Result, EMetaSoundBuilderResult::Succeeded);
+		UTEST_TRUE("Get delayed node handle valid", GetDelayedNode.IsSet());
+
+		FMetasoundFrontendGraphLayer Layer = BuildLayerFromBuilder(Builder);
+
+		FString Code;
+		Engine::FDocumentConfigurator::Reflect(MoveTemp(Layer), Code);
+
+		ValidateReflectedCode(*this, Code,
+		{
+			TEXT(".Add<FDC::FVariable>"),
+			TEXT(".Name = \"MyVar\""),
+			TEXT(".Add<FDC::FVariableSetNode>"),
+			TEXT(".VariableName = \"MyVar\""),
+			TEXT(".Add<FDC::FVariableGetNode>"),
+			TEXT(".Add<FDC::FVariableGetDelayedNode>"),
+		},
+		{
+			TEXT(".SetBuildPage(") // Default page should not emit PageName
+		});
+
+		return true;
+	}
+
+	// ============================================================================
+	// Test 5: Reflect page name (non-default page emits PageName, default does not)
+	// ============================================================================
+	IMPLEMENT_SIMPLE_AUTOMATION_TEST(
+		FReflectPageNameTest,
+		"Audio.Metasound.Configurator.Reflect.PageName",
+		EAutomationTestFlags::EditorContext | EAutomationTestFlags::EngineFilter)
+	bool FReflectPageNameTest::RunTest(const FString&)
+	{
+		using namespace EngineTestMetaSoundPatchBuilderPrivate;
+
+		// Test default page: no .PageName should appear
+		{
+			EMetaSoundBuilderResult Result = EMetaSoundBuilderResult::Failed;
+			UMetaSoundPatchBuilder& Builder = CreatePatchBuilderChecked(*this, "ReflectPageName_Default_Test", { });
+
+			FMetasoundFrontendLiteral Literal;
+			Literal.Set(1.0f);
+			Builder.AddGraphInputNode("TestInput", GetMetasoundDataTypeName<float>(), Literal, Result);
+
+			FMetasoundFrontendGraphLayer Layer = BuildLayerFromBuilder(Builder);
+			// Ensure PageID is the default
+			Layer.Graph.PageID = Frontend::DefaultPageID;
+
+			FString Code;
+			Engine::FDocumentConfigurator::Reflect(MoveTemp(Layer), Code);
+
+			UTEST_FALSE("Default page code not empty", Code.IsEmpty());
+			UTEST_FALSE("Default page code should NOT contain .SetBuildPage", Code.Contains(TEXT(".SetBuildPage(")));
+		}
+
+		// Test non-default page: .SetBuildPage should appear
+		// Only run if there are non-default pages registered in settings
+		{
+			const TArray<FName> PageNames = UMetaSoundSettings::GetPageNames();
+			FName NonDefaultPageName;
+			for (const FName& Name : PageNames)
+			{
+				if (Name != Frontend::DefaultPageName)
+				{
+					NonDefaultPageName = Name;
+					break;
+				}
+			}
+
+			if (!NonDefaultPageName.IsNone())
+			{
+				const FGuid* NonDefaultPageID = UMetaSoundSettings::GetPageID(NonDefaultPageName);
+				if (NonDefaultPageID)
+				{
+					EMetaSoundBuilderResult Result = EMetaSoundBuilderResult::Failed;
+					UMetaSoundPatchBuilder& Builder = CreatePatchBuilderChecked(*this, "ReflectPageName_NonDefault_Test", { });
+
+					FMetasoundFrontendLiteral Literal;
+					Literal.Set(1.0f);
+					Builder.AddGraphInputNode("TestInput", GetMetasoundDataTypeName<float>(), Literal, Result);
+
+					Builder.AddNodeByClassName({ "UE", "Sine", "Audio" }, Result, 1);
+
+					FMetasoundFrontendGraphLayer Layer = BuildLayerFromBuilder(Builder);
+					// Override PageID to non-default
+					Layer.Graph.PageID = *NonDefaultPageID;
+
+					FString Code;
+					Engine::FDocumentConfigurator::Reflect(MoveTemp(Layer), Code);
+
+					const FString ExpectedPageNameFragment = FString::Printf(TEXT(".SetBuildPage(\"%s\")"), *NonDefaultPageName.ToString());
+					UTEST_TRUE("Non-default page code contains .SetBuildPage", Code.Contains(ExpectedPageNameFragment));
+				}
+			}
+			else
+			{
+				AddInfo(TEXT("Skipping non-default page test: no non-default pages registered in MetaSound settings"));
+			}
+		}
+
+		return true;
+	}
+
+	// ============================================================================
+	// Test 6: Round-trip validation
+	// Creates a graph via builder, reflects it, then builds a new graph via
+	// FDocumentConfigurator and compares the two documents structurally.
+	// ============================================================================
+	IMPLEMENT_SIMPLE_AUTOMATION_TEST(
+		FReflectRoundTripTest,
+		"Audio.Metasound.Configurator.Reflect.RoundTrip",
+		EAutomationTestFlags::EditorContext | EAutomationTestFlags::EngineFilter)
+	bool FReflectRoundTripTest::RunTest(const FString&)
+	{
+		using namespace EngineTestMetaSoundPatchBuilderPrivate;
+		using namespace Frontend;
+
+		EMetaSoundBuilderResult Result = EMetaSoundBuilderResult::Failed;
+
+		// === Build source graph ===
+		UMetaSoundPatchBuilder& PatchBuilder = CreatePatchBuilderChecked(*this, "ReflectRoundTrip_Source", { });
+
+		// Add input
+		FMetasoundFrontendLiteral FreqDefault;
+		FreqDefault.Set(440.0f);
+		FMetaSoundBuilderNodeOutputHandle FreqOutput = PatchBuilder.AddGraphInputNode("Frequency", GetMetasoundDataTypeName<float>(), FreqDefault, Result);
+		UTEST_EQUAL("Source: Add input succeeded", Result, EMetaSoundBuilderResult::Succeeded);
+
+		// Add Sine node
+		FMetaSoundNodeHandle SineNode = PatchBuilder.AddNodeByClassName({ "UE", "Sine", "Audio" }, Result, 1);
+		UTEST_EQUAL("Source: Add Sine node succeeded", Result, EMetaSoundBuilderResult::Succeeded);
+
+		// Connect input to Sine frequency
+		FMetaSoundBuilderNodeInputHandle SineFreqInput = PatchBuilder.FindNodeInputByName(SineNode, "Frequency", Result);
+		PatchBuilder.ConnectNodes(FreqOutput, SineFreqInput, Result);
+
+		// Add variable
+		FMetasoundFrontendLiteral VarDefault;
+		VarDefault.Set(1.0f);
+		PatchBuilder.AddGraphVariable("Gain", GetMetasoundDataTypeName<float>(), VarDefault, Result);
+		UTEST_EQUAL("Source: Add variable succeeded", Result, EMetaSoundBuilderResult::Succeeded);
+
+		FMetaSoundNodeHandle VarGetNode = PatchBuilder.AddGraphVariableGetNode("Gain", Result);
+		UTEST_EQUAL("Source: Add variable get node succeeded", Result, EMetaSoundBuilderResult::Succeeded);
+
+		// === Build layer and reflect ===
+		FMetasoundFrontendGraphLayer SourceLayer = BuildLayerFromBuilder(PatchBuilder);
+
+		// Verify source layer has expected content
+		UTEST_TRUE("Source layer has nodes", SourceLayer.Graph.Nodes.Num() > 0);
+		UTEST_TRUE("Source layer has inputs", SourceLayer.Inputs.Num() > 0);
+		UTEST_TRUE("Source layer has variables", SourceLayer.Graph.Variables.Num() > 0);
+
+		FString Code;
+		Engine::FDocumentConfigurator::Reflect(MoveTemp(SourceLayer), Code);
+
+		UTEST_FALSE("Reflected code not empty", Code.IsEmpty());
+
+		// === Apply same configuration to a new builder via FDocumentConfigurator ===
+		UMetaSoundPatchBuilder& DestBuilder = CreatePatchBuilderChecked(*this, "ReflectRoundTrip_Dest", { });
+
+		FMetaSoundFrontendDocumentBuilder& DestDocBuilder = DestBuilder.GetBuilder();
+
+		{
+			Engine::FDocumentConfigurator Configurator(DestBuilder);
+
+			// Replicate what the reflected code would do
+			Configurator
+				.Add(Engine::FDocumentConfigurator::FInput
+				{
+					.Name = "Frequency",
+					.DataType = GetMetasoundDataTypeName<float>(),
+					.DefaultValue = Frontend::MakeFrontendLiteral(440.0f)
+				})
+				.Add(Engine::FDocumentConfigurator::FNode
+				{
+					.Name = "Sine",
+					.ClassName = { "UE", "Sine", "Audio" }
+				})
+				.Add(Engine::FDocumentConfigurator::FVariable
+				{
+					.Name = "Gain",
+					.DataType = GetMetasoundDataTypeName<float>(),
+					.DefaultValue = Frontend::MakeFrontendLiteral(1.0f)
+				})
+				.Add(Engine::FDocumentConfigurator::FVariableGetNode
+				{
+					.Name = "Get_Gain",
+					.VariableName = "Gain"
+				});
+
+			UTEST_TRUE("Configurator succeeded", Configurator.Succeeded());
+		}
+
+		// === Compare documents ===
+		const FMetaSoundFrontendDocumentBuilder& DestConstBuilder = DestBuilder.GetConstBuilder();
+		const FMetasoundFrontendDocument& DestDoc = DestConstBuilder.GetConstDocumentChecked();
+
+		// Verify destination has the same structural elements
+		UTEST_NOT_NULL("Dest has Frequency input", DestConstBuilder.FindGraphInput("Frequency"));
+		UTEST_NOT_NULL("Dest has Gain variable", DestConstBuilder.FindGraphVariable("Gain"));
+
+		const FMetasoundFrontendGraph* DestGraph = DestDoc.RootGraph.FindConstGraph(DestConstBuilder.GetBuildPageID());
+		UTEST_NOT_NULL("Dest graph exists", DestGraph);
+
+		if (DestGraph)
+		{
+			// Count external nodes (Sine) and variable nodes
+			bool bFoundSineNode = false;
+			bool bFoundVariableGetNode = false;
+			for (const FMetasoundFrontendNode& Node : DestGraph->Nodes)
+			{
+				if (const FMetasoundFrontendClass* Dep = DestConstBuilder.FindDependency(Node.ClassID))
+				{
+					if (Dep->Metadata.GetClassName() == FMetasoundFrontendClassName("UE", "Sine", "Audio"))
+					{
+						bFoundSineNode = true;
+					}
+					if (Dep->Metadata.GetType() == EMetasoundFrontendClassType::VariableAccessor)
+					{
+						bFoundVariableGetNode = true;
+					}
+				}
+			}
+			UTEST_TRUE("Dest has Sine node", bFoundSineNode);
+			UTEST_TRUE("Dest has variable get node", bFoundVariableGetNode);
+
+			UTEST_TRUE("Dest has variables", DestGraph->Variables.Num() > 0);
+		}
+
+		return true;
+	}
+
+} // namespace Metasound::Test::ConfiguratorReflect
+
+#endif // WITH_DEV_AUTOMATION_TESTS && WITH_EDITORONLY_DATA
diff --git a/Engine/Plugins/Runtime/Metasound/Source/MetasoundEngineTest/Private/EngineTestMetaSoundBuilder.cpp b/Engine/Plugins/Runtime/Metasound/Source/MetasoundEngineTest/Private/EngineTestMetaSoundBuilder.cpp
new file mode 100644
--- /dev/null
+++ b/Engine/Plugins/Runtime/Metasound/Source/MetasoundEngineTest/Private/EngineTestMetaSoundBuilder.cpp
@@ -0,0 +1,282 @@
+// Copyright Epic Games, Inc. All Rights Reserved.
+
+#include "EngineTestMetaSoundBuilder.h"
+
+#if WITH_DEV_AUTOMATION_TESTS
+
+bool FMetaSoundSourceBuilderDisconnectInputLatentCommand::Update()
+{
+	EMetaSoundBuilderResult Result = EMetaSoundBuilderResult::Failed;
+	if (Builder)
+	{
+		Builder->DisconnectNodeInput(InputToDisconnect, Result);
+		Test.AddErrorIfFalse(Result == EMetaSoundBuilderResult::Succeeded, TEXT("Failed to disconnect MetaSound node's input"));
+	}
+
+	return Result == EMetaSoundBuilderResult::Succeeded;
+}
+
+bool FMetaSoundSourceBuilderSetLiteralLatentCommand::Update()
+{
+	EMetaSoundBuilderResult Result = EMetaSoundBuilderResult::Failed;
+	if (Builder)
+	{
+		Builder->SetNodeInputDefault(NodeInput, NewValue, Result);
+		Test.AddErrorIfFalse(Result == EMetaSoundBuilderResult::Succeeded, TEXT("Failed to set MetaSound node's input default"));
+	}
+
+	return Result == EMetaSoundBuilderResult::Succeeded;
+}
+
+bool FMetaSoundSourceBuilderRemoveNodeDefaultLiteralLatentCommand::Update()
+{
+	EMetaSoundBuilderResult Result = EMetaSoundBuilderResult::Failed;
+	if (Builder)
+	{
+		Builder->RemoveNodeInputDefault(NodeInput, Result);
+		Test.AddErrorIfFalse(Result == EMetaSoundBuilderResult::Succeeded, TEXT("Failed to remove MetaSound node's input default"));
+	}
+
+	return Result == EMetaSoundBuilderResult::Succeeded;
+}
+
+bool FMetaSoundSourceBuilderCreateAndConnectTriGeneratorNodeLatentCommand::Update()
+{
+	EMetaSoundBuilderResult Result = EMetaSoundBuilderResult::Failed;
+	if (Builder)
+	{
+		// Tri Oscillator Node
+		FMetaSoundNodeHandle TriNode = Builder->AddNodeByClassName({ "UE", "Triangle", "Audio" }, Result, 1);
+		Test.AddErrorIfFalse(Result == EMetaSoundBuilderResult::Succeeded && TriNode.IsSet(), TEXT("Failed to create node by class name 'UE:Triangle:Audio"));
+
+		FMetaSoundBuilderNodeOutputHandle TriNodeAudioOutput = Builder->FindNodeOutputByName(TriNode, "Audio", Result);
+		Test.AddErrorIfFalse(Result == EMetaSoundBuilderResult::Succeeded, TEXT("Failed to find Triangle Oscillator node output 'Audio'"));
+
+		Builder->ConnectNodes(TriNodeAudioOutput, AudioOutNodeInput, Result);
+		Test.AddErrorIfFalse(Result == EMetaSoundBuilderResult::Succeeded, TEXT("Failed to connect 'Audio' Triangle Oscillator output to MetaSound graph's 'Mono Output'"));
+	}
+
+	return Result == EMetaSoundBuilderResult::Succeeded;
+}
+
+bool FBuilderRemoveFromRootLatentCommand::Update()
+{
+	if (Builder)
+	{
+		Builder->RemoveFromRoot();
+		return true;
+	}
+	return false;
+}
+
+// Creates a collection of MetaSound patches from builders and exercises bindings all input and output interfaces by connecting and disconnecting using various builder API calls.
+
+IMPLEMENT_SIMPLE_AUTOMATION_TEST(FAudioMetaSoundBuilderInterfaceBindingConnectAndDisconnect, "Audio.Metasound.Builder.InterfaceBindingConnectAndDisconnect", EAutomationTestFlags::EditorContext | EAutomationTestFlags::EngineFilter)
+REGISTER_SIMPLE_AUTOMATION_TEST_TAGS(FAudioMetaSoundBuilderInterfaceBindingConnectAndDisconnect, "Audio.Metasound.Builder.InterfaceBindingConnectAndDisconnect", "[MustPass]")
+bool FAudioMetaSoundBuilderInterfaceBindingConnectAndDisconnect::RunTest(const FString& Parameters)
+{
+	using namespace EngineTestMetaSoundPatchBuilderPrivate;
+	using namespace Metasound::Engine;
+
+	const TArray<FName> InputInterfaces =
+	{
+		InputFormatMonoInterface::GetVersion().Name,
+		InputFormatStereoInterface::GetVersion().Name,
+		InputFormatQuadInterface::GetVersion().Name,
+		InputFormatFiveDotOneInterface::GetVersion().Name,
+		InputFormatSevenDotOneInterface::GetVersion().Name
+	};
+	const TArray<FName> OutputInterfaces =
+	{
+		OutputFormatMonoInterface::GetVersion().Name,
+		OutputFormatStereoInterface::GetVersion().Name,
+		OutputFormatQuadInterface::GetVersion().Name,
+		OutputFormatFiveDotOneInterface::GetVersion().Name,
+		OutputFormatSevenDotOneInterface::GetVersion().Name
+	};
+
+	auto MakePatchName = [](FName InterfaceName) { return FString::Printf(TEXT("%s Patch"), *InterfaceName.ToString()); };
+	EMetaSoundBuilderResult Result = EMetaSoundBuilderResult::Failed;
+	for (const FName& InputInterface : InputInterfaces)
+	{
+		if (UMetaSoundPatch* InputBuilder = CreatePatchFromBuilder(*this, MakePatchName(InputInterface), { InputInterface }))
+		{
+			for (const FName& OutputInterface : OutputInterfaces)
+			{
+				if (UMetaSoundPatch* OutputBuilder = CreatePatchFromBuilder(*this, MakePatchName(OutputInterface), { OutputInterface }))
+				{
+					const FString BindBuilderName = FString::Printf(TEXT("%s to %s Binding"), *OutputInterface.ToString(), *InputInterface.ToString());
+					UMetaSoundPatchBuilder& BindingBuilder = CreatePatchBuilderChecked(*this, *BindBuilderName, { InputInterface, OutputInterface });
+					FMetaSoundNodeHandle InputInterfaceNode = BindingBuilder.AddNode(TScriptInterface<IMetaSoundDocumentInterface>(InputBuilder), Result);
+					FMetaSoundNodeHandle OutputInterfaceNode = BindingBuilder.AddNode(TScriptInterface<IMetaSoundDocumentInterface>(OutputBuilder), Result);
+					BindingBuilder.ConnectNodesByInterfaceBindings(OutputInterfaceNode, InputInterfaceNode, Result);
+					AddErrorIfFalse(
+						Result == EMetaSoundBuilderResult::Succeeded,
+						FString::Printf(TEXT("Failed to connect nodes via binding with interface '%s' to '%s'"), *OutputInterface.ToString(), *InputInterface.ToString())
+					);
+
+					BindingBuilder.DisconnectNodesByInterfaceBindings(OutputInterfaceNode, InputInterfaceNode, Result);
+					AddErrorIfFalse(
+						Result == EMetaSoundBuilderResult::Succeeded,
+						FString::Printf(TEXT("Failed to disconnect nodes via binding with interface '%s' to '%s'"), *OutputInterface.ToString(), *InputInterface.ToString())
+					);
+
+					BindingBuilder.ConnectNodeInputsToMatchingGraphInterfaceInputs(InputInterfaceNode, Result);
+					AddErrorIfFalse(
+						Result == EMetaSoundBuilderResult::Succeeded,
+						FString::Printf(TEXT("Failed to connect input node with matching graph interface inputs: interface '%s'"), *InputInterface.ToString())
+					);
+
+					BindingBuilder.ConnectNodeOutputsToMatchingGraphInterfaceOutputs(OutputInterfaceNode, Result);
+					AddErrorIfFalse(
+						Result == EMetaSoundBuilderResult::Succeeded,
+						FString::Printf(TEXT("Failed to connect output node with matching graph interface outputs: interface '%s'"), *OutputInterface.ToString())
+					);
+				}
+			}
+		}
+	}
+
+	return true;
+}
+
+// This test creates a MetaSound patch, then adds and connects sin oscillator, attempts to retrieve a default input set, then clears it, and finally retrieves the class default
+IMPLEMENT_SIMPLE_AUTOMATION_TEST(FAudioMetaSoundNodeClassQueryFunctions, "Audio.Metasound.Builder.NodeClassQueryFunctions", EAutomationTestFlags::EditorContext | EAutomationTestFlags::EngineFilter)
+REGISTER_SIMPLE_AUTOMATION_TEST_TAGS(FAudioMetaSoundNodeClassQueryFunctions, "Audio.Metasound.Builder.NodeClassQueryFunctions", "[MustPass]")
+bool FAudioMetaSoundNodeClassQueryFunctions::RunTest(const FString& Parameters)
+{
+	using namespace EngineTestMetaSoundPatchBuilderPrivate;
+	using namespace Metasound;
+	using namespace Metasound::Frontend;
+
+	constexpr float NodeDefaultFreq = 100.0f;
+
+	EMetaSoundBuilderResult Result = EMetaSoundBuilderResult::Failed;
+	UMetaSoundPatchBuilder& Builder = CreatePatchBuilderChecked(*this, "DefaultLiteralAssignment_Test", { });
+
+	// Sine Oscillator Node
+	const FMetaSoundNodeHandle OscNode = Builder.AddNodeByClassName({ "UE", "Sine", "Audio" }, Result, 1);
+	AddErrorIfFalse(Result == EMetaSoundBuilderResult::Succeeded && OscNode.IsSet(), TEXT("Failed to create new MetaSound node by class name"));
+
+	// Make connections:
+	const FMetaSoundBuilderNodeInputHandle OscNodeFrequencyInput = Builder.FindNodeInputByName(OscNode, "Frequency", Result);
+	AddErrorIfFalse(Result == EMetaSoundBuilderResult::Succeeded && OscNodeFrequencyInput.IsSet(), TEXT("Failed to find Sine Oscillator node input 'Frequency'"));
+
+	FMetasoundFrontendLiteral Literal;
+	Literal.Set(100.0f);
+
+	Builder.SetNodeInputDefault(OscNodeFrequencyInput, Literal, Result);
+	AddErrorIfFalse(Result == EMetaSoundBuilderResult::Succeeded, TEXT("Failed to 'SetNodeInputDefault'"));
+
+	Literal = Builder.GetNodeInputDefault(OscNodeFrequencyInput, Result);
+	AddErrorIfFalse(Result == EMetaSoundBuilderResult::Succeeded, TEXT("Failed to 'GetNodeInputDefault'"));
+
+	float Freq = 0.0f;
+	bool bLiteralIsFloat = Literal.TryGet(Freq);
+	AddErrorIfFalse(bLiteralIsFloat, TEXT("Failed to retrieve node literal as 'float'"));
+	AddErrorIfFalse(FMath::IsNearlyEqual(NodeDefaultFreq, Freq), TEXT("'Freq' node default not set to provided default"));
+
+	Builder.RemoveNodeInputDefault(OscNodeFrequencyInput, Result);
+	AddErrorIfFalse(Result == EMetaSoundBuilderResult::Succeeded, TEXT("Failed to 'RemoveNodeInputDefault'"));
+
+	Literal = Builder.GetNodeInputClassDefault(OscNodeFrequencyInput, Result);
+	AddErrorIfFalse(Result == EMetaSoundBuilderResult::Succeeded, TEXT("Failed to 'GetNodeInputClassDefault' on c++ Sin Osc node"));
+
+	bLiteralIsFloat = Literal.TryGet(Freq);
+	AddErrorIfFalse(bLiteralIsFloat, TEXT("Failed to retrieve class input literal as 'float'"));
+	AddErrorIfFalse(FMath::IsNearlyEqual(440.0f, Freq), TEXT("'Freq' node default not set to expected sin freq default of 440.0f Hz"));
+
+	// Output Default Test
+	{
+		Literal.Set(123.0f);
+		const FMetaSoundBuilderNodeInputHandle InputTest = Builder.AddGraphOutputNode("Output", GetMetasoundDataTypeName<float>(), Literal, Result);
+		AddErrorIfFalse(Result == EMetaSoundBuilderResult::Succeeded && InputTest.IsSet(), TEXT("Failed to create new MetaSound graph input"));
+
+		bLiteralIsFloat = Literal.TryGet(Freq);
+		AddErrorIfFalse(bLiteralIsFloat, TEXT("Failed to retrieve class input literal as 'float'"));
+		AddErrorIfFalse(FMath::IsNearlyEqual(123.0f, Freq), TEXT("'Freq' node default not set to expected sin freq default of 440.0f Hz"));
+
+		const bool bIsConstructorPin = Builder.GetNodeInputIsConstructorPin(InputTest);
+		AddErrorIfFalse(!bIsConstructorPin, TEXT("Pin is not constructor pin but GetNodeInputIsConstructorPin is returning true"));
+	}
+
+	// Output Ctor tests
+	{
+		const FMetaSoundBuilderNodeInputHandle CtorTest = Builder.AddGraphOutputNode("CtorOutputTest", GetMetasoundDataTypeName<float>(), Literal, Result, true);
+		AddErrorIfFalse(Result == EMetaSoundBuilderResult::Succeeded && CtorTest.IsSet(), TEXT("Failed to create new MetaSound graph input"));
+
+		const bool bIsConstructorPin = Builder.GetNodeInputIsConstructorPin(CtorTest);
+		AddErrorIfFalse(bIsConstructorPin, TEXT("Output handle is constructor pin but GetNodeInputIsConstructorPin is returning false"));
+	}
+
+	// Input Ctor tests
+	{
+		const FMetaSoundBuilderNodeOutputHandle CtorTest = Builder.AddGraphInputNode("CtorOutput", GetMetasoundDataTypeName<float>(), Literal, Result, true);
+		AddErrorIfFalse(Result == EMetaSoundBuilderResult::Succeeded && CtorTest.IsSet(), TEXT("Failed to create new MetaSound graph input"));
+
+		const bool bIsConstructorPin = Builder.GetNodeOutputIsConstructorPin(CtorTest);
+		AddErrorIfFalse(bIsConstructorPin, TEXT("Input handle is constructor pin but GetNodeInputIsConstructorPin is returning false"));
+	}
+
+	return true;
+}
+
+// This test creates a MetaSound patch, adds a frequency input, connects a sin oscillator, and if the editor is loaded, and then attempts to inject template input nodes.
+IMPLEMENT_SIMPLE_AUTOMATION_TEST(FAudioMetaSoundInjectTemplateNodes, "Audio.Metasound.Builder.InjectInputTemplateNodes", EAutomationTestFlags::EditorContext | EAutomationTestFlags::EngineFilter)
+bool FAudioMetaSoundInjectTemplateNodes::RunTest(const FString& Parameters)
+{
+	using namespace EngineTestMetaSoundPatchBuilderPrivate;
+	using namespace Metasound;
+	using namespace Metasound::Frontend;
+
+	constexpr float NodeDefaultFreq = 100.0f;
+
+	EMetaSoundBuilderResult Result = EMetaSoundBuilderResult::Failed;
+	UMetaSoundPatchBuilder& Builder = CreatePatchBuilderChecked(*this, "TemplateInputInjection_Test", { });
+
+	// Sine Oscillator Node
+	const FMetaSoundNodeHandle OscNode = Builder.AddNodeByClassName({ "UE", "Sine", "Audio" }, Result, 1);
+	AddErrorIfFalse(Result == EMetaSoundBuilderResult::Succeeded && OscNode.IsSet(), TEXT("Failed to create new MetaSound node by class name"));
+
+	auto MakeAndConnectSinInputToGraphInput = [&](FName InputName, FName TypeName, const FVector2D* Location, const FMetasoundFrontendLiteral& InLiteral, bool bIsCtorPin)
+	{
+		const FMetaSoundBuilderNodeOutputHandle InputNodeOutputHandle = Builder.AddGraphInputNode(InputName, TypeName, InLiteral, Result, bIsCtorPin);
+		AddErrorIfFalse(Result == EMetaSoundBuilderResult::Succeeded && InputNodeOutputHandle.IsSet(), TEXT("Failed to create new MetaSound graph input"));
+
+		const FMetaSoundBuilderNodeInputHandle OscNodeInput = Builder.FindNodeInputByName(OscNode, InputName, Result);
+		AddErrorIfFalse(Result == EMetaSoundBuilderResult::Succeeded && OscNodeInput.IsSet(), TEXT("Failed to find Sine Oscillator node input"));
+
+		Builder.ConnectNodes(InputNodeOutputHandle, OscNodeInput, Result);
+		AddErrorIfFalse(Result == EMetaSoundBuilderResult::Succeeded, TEXT("Failed to connect input node to node input"));
+
+#if WITH_EDITOR
+		if (Location)
+		{
+			Builder.SetNodeLocation(FMetaSoundNodeHandle{ InputNodeOutputHandle.NodeID }, *Location, Result);
+			AddErrorIfFalse(Result == EMetaSoundBuilderResult::Succeeded, TEXT("Failed to set input node location"));
+		}
+#endif // WITH_EDITOR
+	};
+	FMetasoundFrontendLiteral Literal;
+
+	Literal.Set(true);
+	MakeAndConnectSinInputToGraphInput("Enabled", GetMetasoundDataTypeName<bool>(), nullptr, Literal, true);
+
+	Literal.Set(100.0f);
+	FVector2D FreqLoc(100, 100);
+	MakeAndConnectSinInputToGraphInput("Frequency", GetMetasoundDataTypeName<float>(), nullptr, Literal, false);
+
+	Literal.Set(FMetasoundFrontendLiteral::FDefault { });
+	FVector2D SyncLoc(300, 300);
+	MakeAndConnectSinInputToGraphInput("Sync", GetMetasoundDataTypeName<FTrigger>(), nullptr, Literal, false);
+
+#if WITH_EDITOR
+	constexpr bool bForceNodeCreation = false;
+	Builder.InjectInputTemplateNodes(bForceNodeCreation, Result);
+	AddErrorIfFalse(Result == EMetaSoundBuilderResult::Succeeded, TEXT("Failed to inject input template node"));
+#endif // WITH_EDITOR
+
+	return true;
+}
+
+#endif // WITH_DEV_AUTOMATION_TESTS
diff --git a/Engine/Plugins/Runtime/Metasound/Source/MetasoundEngineTest/Private/EngineTestMetaSoundBuilder.h b/Engine/Plugins/Runtime/Metasound/Source/MetasoundEngineTest/Private/EngineTestMetaSoundBuilder.h
--- a/Engine/Plugins/Runtime/Metasound/Source/MetasoundEngineTest/Private/EngineTestMetaSoundBuilder.h
+++ b/Engine/Plugins/Runtime/Metasound/Source/MetasoundEngineTest/Private/EngineTestMetaSoundBuilder.h
@@ -52,286 +52,12 @@ namespace EngineTestMetaSoundPatchBuilderPrivate
 
 DEFINE_LATENT_AUTOMATION_COMMAND_THREE_PARAMETER(FMetaSoundSourceBuilderDisconnectInputLatentCommand, FAutomationTestBase&, Test, UMetaSoundBuilderBase*, Builder, FMetaSoundBuilderNodeInputHandle, InputToDisconnect);
 
-inline bool FMetaSoundSourceBuilderDisconnectInputLatentCommand::Update()
-{
-	EMetaSoundBuilderResult Result = EMetaSoundBuilderResult::Failed;
-	if (Builder)
-	{
-		Builder->DisconnectNodeInput(InputToDisconnect, Result);
-		Test.AddErrorIfFalse(Result == EMetaSoundBuilderResult::Succeeded, TEXT("Failed to disconnect MetaSound node's input"));
-	}
-
-	return Result == EMetaSoundBuilderResult::Succeeded;
-}
-
 DEFINE_LATENT_AUTOMATION_COMMAND_FOUR_PARAMETER(FMetaSoundSourceBuilderSetLiteralLatentCommand, FAutomationTestBase&, Test, UMetaSoundBuilderBase*, Builder, FMetaSoundBuilderNodeInputHandle, NodeInput, FMetasoundFrontendLiteral, NewValue);
 
-inline bool FMetaSoundSourceBuilderSetLiteralLatentCommand::Update()
-{
-	EMetaSoundBuilderResult Result = EMetaSoundBuilderResult::Failed;
-	if (Builder)
-	{
-		Builder->SetNodeInputDefault(NodeInput, NewValue, Result);
-		Test.AddErrorIfFalse(Result == EMetaSoundBuilderResult::Succeeded, TEXT("Failed to disconnect MetaSound node's input"));
-	}
-
-	return Result == EMetaSoundBuilderResult::Succeeded;
-}
-
 DEFINE_LATENT_AUTOMATION_COMMAND_THREE_PARAMETER(FMetaSoundSourceBuilderRemoveNodeDefaultLiteralLatentCommand, FAutomationTestBase&, Test, UMetaSoundBuilderBase*, Builder, FMetaSoundBuilderNodeInputHandle, NodeInput);
 
-inline bool FMetaSoundSourceBuilderRemoveNodeDefaultLiteralLatentCommand::Update()
-{
-	EMetaSoundBuilderResult Result = EMetaSoundBuilderResult::Failed;
-	if (Builder)
-	{
-		Builder->RemoveNodeInputDefault(NodeInput, Result);
-		Test.AddErrorIfFalse(Result == EMetaSoundBuilderResult::Succeeded, TEXT("Failed to disconnect MetaSound node's input"));
-	}
-
-	return Result == EMetaSoundBuilderResult::Succeeded;
-}
-
 DEFINE_LATENT_AUTOMATION_COMMAND_THREE_PARAMETER(FMetaSoundSourceBuilderCreateAndConnectTriGeneratorNodeLatentCommand, FAutomationTestBase&, Test, UMetaSoundSourceBuilder*, Builder, FMetaSoundBuilderNodeInputHandle, AudioOutNodeInput);
 
-inline bool FMetaSoundSourceBuilderCreateAndConnectTriGeneratorNodeLatentCommand::Update()
-{
-	EMetaSoundBuilderResult Result = EMetaSoundBuilderResult::Failed;
-	if (Builder)
-	{
-		// Tri Oscillator Node
-		FMetaSoundNodeHandle TriNode = Builder->AddNodeByClassName({ "UE", "Triangle", "Audio" }, Result, 1);
-		Test.AddErrorIfFalse(Result == EMetaSoundBuilderResult::Succeeded && TriNode.IsSet(), TEXT("Failed to create node by class name 'UE:Triangle:Audio"));
-
-		FMetaSoundBuilderNodeOutputHandle TriNodeAudioOutput = Builder->FindNodeOutputByName(TriNode, "Audio", Result);
-		Test.AddErrorIfFalse(Result == EMetaSoundBuilderResult::Succeeded, TEXT("Failed to find Triangle Oscillator node output 'Audio'"));
-
-		Builder->ConnectNodes(TriNodeAudioOutput, AudioOutNodeInput, Result);
-		Test.AddErrorIfFalse(Result == EMetaSoundBuilderResult::Succeeded, TEXT("Failed to connect 'Audio' Triangle Oscillator output to MetaSound graph's 'Mono Output'"));
-	}
-
-	return Result == EMetaSoundBuilderResult::Succeeded;
-}
-
 DEFINE_LATENT_AUTOMATION_COMMAND_ONE_PARAMETER(FBuilderRemoveFromRootLatentCommand, UMetaSoundBuilderBase*, Builder);
 
-inline bool FBuilderRemoveFromRootLatentCommand::Update()
-{
-	if (Builder)
-	{
-		Builder->RemoveFromRoot();
-		return true;
-	}
-	return false;
-}
-
-// Creates a collection of MetaSound patches from builders and exercises bindings all input and output interfaces by connecting and disconnecting using various builder API calls.
-
-IMPLEMENT_SIMPLE_AUTOMATION_TEST(FAudioMetaSoundBuilderInterfaceBindingConnectAndDisconnect, "Audio.Metasound.Builder.InterfaceBindingConnectAndDisconnect", EAutomationTestFlags::EditorContext | EAutomationTestFlags::EngineFilter)
-REGISTER_SIMPLE_AUTOMATION_TEST_TAGS(FAudioMetaSoundBuilderInterfaceBindingConnectAndDisconnect, "Audio.Metasound.Builder.InterfaceBindingConnectAndDisconnect", "[MustPass]")
-inline bool FAudioMetaSoundBuilderInterfaceBindingConnectAndDisconnect::RunTest(const FString& Parameters)
-{
-	using namespace EngineTestMetaSoundPatchBuilderPrivate;
-	using namespace Metasound::Engine;
-
-	const TArray<FName> InputInterfaces =
-	{
-		InputFormatMonoInterface::GetVersion().Name,
-		InputFormatStereoInterface::GetVersion().Name,
-		InputFormatQuadInterface::GetVersion().Name,
-		InputFormatFiveDotOneInterface::GetVersion().Name,
-		InputFormatSevenDotOneInterface::GetVersion().Name
-	};
-	const TArray<FName> OutputInterfaces =
-	{
-		OutputFormatMonoInterface::GetVersion().Name,
-		OutputFormatStereoInterface::GetVersion().Name,
-		OutputFormatQuadInterface::GetVersion().Name,
-		OutputFormatFiveDotOneInterface::GetVersion().Name,
-		OutputFormatSevenDotOneInterface::GetVersion().Name
-	};
-
-	auto MakePatchName = [](FName InterfaceName) { return FString::Printf(TEXT("%s Patch"), *InterfaceName.ToString()); };
-	EMetaSoundBuilderResult Result = EMetaSoundBuilderResult::Failed;
-	for (const FName& InputInterface : InputInterfaces)
-	{
-		if (UMetaSoundPatch* InputBuilder = CreatePatchFromBuilder(*this, MakePatchName(InputInterface), { InputInterface }))
-		{
-			for (const FName& OutputInterface : OutputInterfaces)
-			{
-				if (UMetaSoundPatch* OutputBuilder = CreatePatchFromBuilder(*this, MakePatchName(OutputInterface), { OutputInterface }))
-				{
-					const FString BindBuilderName = FString::Printf(TEXT("%s to %s Binding"), *OutputInterface.ToString(), *InputInterface.ToString());
-					UMetaSoundPatchBuilder& BindingBuilder = CreatePatchBuilderChecked(*this, *BindBuilderName, { InputInterface, OutputInterface });
-					FMetaSoundNodeHandle InputInterfaceNode = BindingBuilder.AddNode(TScriptInterface<IMetaSoundDocumentInterface>(InputBuilder), Result);
-					FMetaSoundNodeHandle OutputInterfaceNode = BindingBuilder.AddNode(TScriptInterface<IMetaSoundDocumentInterface>(OutputBuilder), Result);
-					BindingBuilder.ConnectNodesByInterfaceBindings(OutputInterfaceNode, InputInterfaceNode, Result);
-					AddErrorIfFalse(
-						Result == EMetaSoundBuilderResult::Succeeded,
-						FString::Printf(TEXT("Failed to connect nodes via binding with interface '%s' to '%s'"), *OutputInterface.ToString(), *InputInterface.ToString())
-					);
-
-					BindingBuilder.DisconnectNodesByInterfaceBindings(OutputInterfaceNode, InputInterfaceNode, Result);
-					AddErrorIfFalse(
-						Result == EMetaSoundBuilderResult::Succeeded,
-						FString::Printf(TEXT("Failed to disconnect nodes via binding with interface '%s' to '%s'"), *OutputInterface.ToString(), *InputInterface.ToString())
-					);
-
-					BindingBuilder.ConnectNodeInputsToMatchingGraphInterfaceInputs(InputInterfaceNode, Result);
-					AddErrorIfFalse(
-						Result == EMetaSoundBuilderResult::Succeeded,
-						FString::Printf(TEXT("Failed to connect input node with matching graph interface inputs: interface '%s'"), *InputInterface.ToString())
-					);
-
-					BindingBuilder.ConnectNodeOutputsToMatchingGraphInterfaceOutputs(OutputInterfaceNode, Result);
-					AddErrorIfFalse(
-						Result == EMetaSoundBuilderResult::Succeeded,
-						FString::Printf(TEXT("Failed to connect output node with matching graph interface outputs: interface '%s'"), *OutputInterface.ToString())
-					);
-				}
-			}
-		}
-	}
-
-	return true;
-}
-
-// This test creates a MetaSound patch, then adds and connects sin oscillator, attempts to retrieve a default input set, then clears it, and finally retrieves the class default
-IMPLEMENT_SIMPLE_AUTOMATION_TEST(FAudioMetaSoundNodeClassQueryFunctions, "Audio.Metasound.Builder.NodeClassQueryFunctions", EAutomationTestFlags::EditorContext | EAutomationTestFlags::EngineFilter)
-REGISTER_SIMPLE_AUTOMATION_TEST_TAGS(FAudioMetaSoundNodeClassQueryFunctions, "Audio.Metasound.Builder.NodeClassQueryFunctions", "[MustPass]")
-inline bool FAudioMetaSoundNodeClassQueryFunctions::RunTest(const FString& Parameters)
-{
-	using namespace EngineTestMetaSoundPatchBuilderPrivate;
-	using namespace Metasound;
-	using namespace Metasound::Frontend;
-
-	constexpr float NodeDefaultFreq = 100.0f;
-
-	EMetaSoundBuilderResult Result = EMetaSoundBuilderResult::Failed;
-	UMetaSoundPatchBuilder& Builder = CreatePatchBuilderChecked(*this, "DefaultLiteralAssignment_Test", { });
-
-	// Sine Oscillator Node
-	const FMetaSoundNodeHandle OscNode = Builder.AddNodeByClassName({ "UE", "Sine", "Audio" }, Result, 1);
-	AddErrorIfFalse(Result == EMetaSoundBuilderResult::Succeeded && OscNode.IsSet(), TEXT("Failed to create new MetaSound node by class name"));
-
-	// Make connections:
-	const FMetaSoundBuilderNodeInputHandle OscNodeFrequencyInput = Builder.FindNodeInputByName(OscNode, "Frequency", Result);
-	AddErrorIfFalse(Result == EMetaSoundBuilderResult::Succeeded && OscNodeFrequencyInput.IsSet(), TEXT("Failed to find Sine Oscillator node input 'Frequency'"));
-
-	FMetasoundFrontendLiteral Literal;
-	Literal.Set(100.0f);
-
-	Builder.SetNodeInputDefault(OscNodeFrequencyInput, Literal, Result);
-	AddErrorIfFalse(Result == EMetaSoundBuilderResult::Succeeded, TEXT("Failed to 'SetNodeInputDefault'"));
-
-	Literal = Builder.GetNodeInputDefault(OscNodeFrequencyInput, Result);
-	AddErrorIfFalse(Result == EMetaSoundBuilderResult::Succeeded, TEXT("Failed to 'GetNodeInputDefault'"));
-
-	float Freq = 0.0f;
-	bool bLiteralIsFloat = Literal.TryGet(Freq);
-	AddErrorIfFalse(bLiteralIsFloat, TEXT("Failed to retrieve node literal as 'float'"));
-	AddErrorIfFalse(FMath::IsNearlyEqual(NodeDefaultFreq, Freq), TEXT("'Freq' node default not set to provided default"));
-
-	Builder.RemoveNodeInputDefault(OscNodeFrequencyInput, Result);
-	AddErrorIfFalse(Result == EMetaSoundBuilderResult::Succeeded, TEXT("Failed to 'RemoveNodeInputDefault'"));
-
-	Literal = Builder.GetNodeInputClassDefault(OscNodeFrequencyInput, Result);
-	AddErrorIfFalse(Result == EMetaSoundBuilderResult::Succeeded, TEXT("Failed to 'GetNodeInputClassDefault' on c++ Sin Osc node"));
-
-	bLiteralIsFloat = Literal.TryGet(Freq);
-	AddErrorIfFalse(bLiteralIsFloat, TEXT("Failed to retrieve class input literal as 'float'"));
-	AddErrorIfFalse(FMath::IsNearlyEqual(440.0f, Freq), TEXT("'Freq' node default not set to expected sin freq default of 440.0f Hz"));
-
-	// Output Default Test
-	{
-		Literal.Set(123.0f);
-		const FMetaSoundBuilderNodeInputHandle InputTest = Builder.AddGraphOutputNode("Output", GetMetasoundDataTypeName<float>(), Literal, Result);
-		AddErrorIfFalse(Result == EMetaSoundBuilderResult::Succeeded && InputTest.IsSet(), TEXT("Failed to create new MetaSound graph input"));
-
-		bLiteralIsFloat = Literal.TryGet(Freq);
-		AddErrorIfFalse(bLiteralIsFloat, TEXT("Failed to retrieve class input literal as 'float'"));
-		AddErrorIfFalse(FMath::IsNearlyEqual(123.0f, Freq), TEXT("'Freq' node default not set to expected sin freq default of 440.0f Hz"));
-
-		const bool bIsConstructorPin = Builder.GetNodeInputIsConstructorPin(InputTest);
-		AddErrorIfFalse(!bIsConstructorPin, TEXT("Pin is not constructor pin but GetNodeInputIsConstructorPin is returning true"));
-	}
-
-	// Output Ctor tests
-	{
-		const FMetaSoundBuilderNodeInputHandle CtorTest = Builder.AddGraphOutputNode("CtorOutputTest", GetMetasoundDataTypeName<float>(), Literal, Result, true);
-		AddErrorIfFalse(Result == EMetaSoundBuilderResult::Succeeded && CtorTest.IsSet(), TEXT("Failed to create new MetaSound graph input"));
-
-		const bool bIsConstructorPin = Builder.GetNodeInputIsConstructorPin(CtorTest);
-		AddErrorIfFalse(bIsConstructorPin, TEXT("Output handle is constructor pin but GetNodeInputIsConstructorPin is returning false"));
-	}
-
-	// Input Ctor tests
-	{
-		const FMetaSoundBuilderNodeOutputHandle CtorTest = Builder.AddGraphInputNode("CtorOutput", GetMetasoundDataTypeName<float>(), Literal, Result, true);
-		AddErrorIfFalse(Result == EMetaSoundBuilderResult::Succeeded && CtorTest.IsSet(), TEXT("Failed to create new MetaSound graph input"));
-
-		const bool bIsConstructorPin = Builder.GetNodeOutputIsConstructorPin(CtorTest);
-		AddErrorIfFalse(bIsConstructorPin, TEXT("Input handle is constructor pin but GetNodeInputIsConstructorPin is returning false"));
-	}
-
-	return true;
-}
-
-// This test creates a MetaSound patch, adds a frequency input, connects a sin oscillator, and if the editor is loaded, and then attempts to inject template input nodes.
-IMPLEMENT_SIMPLE_AUTOMATION_TEST(FAudioMetaSoundInjectTemplateNodes, "Audio.Metasound.Builder.InjectInputTemplateNodes", EAutomationTestFlags::EditorContext | EAutomationTestFlags::EngineFilter)
-inline bool FAudioMetaSoundInjectTemplateNodes::RunTest(const FString& Parameters)
-{
-	using namespace EngineTestMetaSoundPatchBuilderPrivate;
-	using namespace Metasound;
-	using namespace Metasound::Frontend;
-
-	constexpr float NodeDefaultFreq = 100.0f;
-
-	EMetaSoundBuilderResult Result = EMetaSoundBuilderResult::Failed;
-	UMetaSoundPatchBuilder& Builder = CreatePatchBuilderChecked(*this, "TemplateInputInjection_Test", { });
-
-	// Sine Oscillator Node
-	const FMetaSoundNodeHandle OscNode = Builder.AddNodeByClassName({ "UE", "Sine", "Audio" }, Result, 1);
-	AddErrorIfFalse(Result == EMetaSoundBuilderResult::Succeeded && OscNode.IsSet(), TEXT("Failed to create new MetaSound node by class name"));
-
-	auto MakeAndConnectSinInputToGraphInput = [&](FName InputName, FName TypeName, const FVector2D* Location, const FMetasoundFrontendLiteral& InLiteral, bool bIsCtorPin)
-	{
-		const FMetaSoundBuilderNodeOutputHandle InputNodeOutputHandle = Builder.AddGraphInputNode(InputName, TypeName, InLiteral, Result, bIsCtorPin);
-		AddErrorIfFalse(Result == EMetaSoundBuilderResult::Succeeded && InputNodeOutputHandle.IsSet(), TEXT("Failed to create new MetaSound graph input"));
-
-		const FMetaSoundBuilderNodeInputHandle OscNodeInput = Builder.FindNodeInputByName(OscNode, InputName, Result);
-		AddErrorIfFalse(Result == EMetaSoundBuilderResult::Succeeded && OscNodeInput.IsSet(), TEXT("Failed to find Sine Oscillator node input"));
-
-		Builder.ConnectNodes(InputNodeOutputHandle, OscNodeInput, Result);
-		AddErrorIfFalse(Result == EMetaSoundBuilderResult::Succeeded, TEXT("Failed to connect input node to node input"));
-
-#if WITH_EDITOR
-		if (Location)
-		{
-			Builder.SetNodeLocation(FMetaSoundNodeHandle{ InputNodeOutputHandle.NodeID }, *Location, Result);
-			AddErrorIfFalse(Result == EMetaSoundBuilderResult::Succeeded, TEXT("Failed to set input node location"));
-		}
-#endif // WITH_EDITOR
-	};
-	FMetasoundFrontendLiteral Literal;
-
-	Literal.Set(true);
-	MakeAndConnectSinInputToGraphInput("Enabled", GetMetasoundDataTypeName<bool>(), nullptr, Literal, true);
-
-	Literal.Set(100.0f);
-	FVector2D FreqLoc(100, 100);
-	MakeAndConnectSinInputToGraphInput("Frequency", GetMetasoundDataTypeName<float>(), nullptr, Literal, false);
-
-	Literal.Set(FMetasoundFrontendLiteral::FDefault { });
-	FVector2D SyncLoc(300, 300);
-	MakeAndConnectSinInputToGraphInput("Sync", GetMetasoundDataTypeName<FTrigger>(), nullptr, Literal, false);
-
-#if WITH_EDITOR
-	constexpr bool bForceNodeCreation = false;
-	Builder.InjectInputTemplateNodes(bForceNodeCreation, Result);
-	AddErrorIfFalse(Result == EMetaSoundBuilderResult::Succeeded, TEXT("Failed to inject input template node"));
-#endif // WITH_EDITOR
-
-	return true;
-}
 #endif // WITH_DEV_AUTOMATION_TESTS
diff --git a/Engine/Plugins/Runtime/Metasound/Source/MetasoundFrontend/Private/MetasoundFrontendDocumentBuilder.cpp b/Engine/Plugins/Runtime/Metasound/Source/MetasoundFrontend/Private/MetasoundFrontendDocumentBuilder.cpp
--- a/Engine/Plugins/Runtime/Metasound/Source/MetasoundFrontend/Private/MetasoundFrontendDocumentBuilder.cpp
+++ b/Engine/Plugins/Runtime/Metasound/Source/MetasoundFrontend/Private/MetasoundFrontendDocumentBuilder.cpp
@@ -3605,7 +3605,7 @@ const FMetasoundFrontendVertex* FMetaSoundFrontendDocumentBuilder::FindNodeOutpu
 	return nullptr;
 }
 
-int32 FMetaSoundFrontendDocumentBuilder::FindPageIndex(const FGuid& InPageID) const
+int32 FMetaSoundFrontendDocumentBuilder::FindGraphPageIndex(const FGuid& InPageID) const
 {
 	const FMetasoundFrontendDocument& Document = GetDocumentChecked();
 	const TArray<FMetasoundFrontendGraph>& GraphPages = Document.RootGraph.GetConstGraphPages();
diff --git a/Engine/Plugins/Runtime/Metasound/Source/MetasoundFrontend/Public/MetasoundFrontendDocumentBuilder.h b/Engine/Plugins/Runtime/Metasound/Source/MetasoundFrontend/Public/MetasoundFrontendDocumentBuilder.h
--- a/Engine/Plugins/Runtime/Metasound/Source/MetasoundFrontend/Public/MetasoundFrontendDocumentBuilder.h
+++ b/Engine/Plugins/Runtime/Metasound/Source/MetasoundFrontend/Public/MetasoundFrontendDocumentBuilder.h
@@ -396,8 +396,12 @@ PRAGMA_ENABLE_EXPERIMENTAL_WARNINGS
 	UE_API TArray<const FMetasoundFrontendVertex*> FindNodeOutputs(const FGuid& InNodeID, FName TypeName = FName(), const FGuid* InPageID = nullptr) const;
 	UE_API const FMetasoundFrontendVertex* FindNodeOutputConnectedToNodeInput(const FGuid& InInputNodeID, const FGuid& InInputVertexID, const FMetasoundFrontendNode** ConnectedOutputNode = nullptr, const FGuid* InPageID = nullptr) const;
 
-	// Return the index of the given page in the document's PagedGraphs array 
-	UE_API int32 FindPageIndex(const FGuid& InPageID) const;
+	// Return the index of the given page in the document's PagedGraphs array. Returns
+	// INDEX_NONE in document where the given graph with the provided PageID is not found.
+	UE_API int32 FindGraphPageIndex(const FGuid& InPageID) const;
+
+	UE_DEPRECATED(5.8, "Use 'FindGraphPageIndex' instead")
+	int32 FindPageIndex(const FGuid& InPageID) const { return FindGraphPageIndex(InPageID); }
 
 	UE_API const FMetasoundFrontendDocument& GetConstDocumentChecked() const;
 
