[Sequencer] Improve UE.Source.StartTime interface warning for MetaSoundSources that don't implement it.

#Synced-CL 52394978

diff --git a/Engine/Plugins/Runtime/Metasound/Source/MetasoundEngine/Private/MetasoundSource.cpp b/Engine/Plugins/Runtime/Metasound/Source/MetasoundEngine/Private/MetasoundSource.cpp
--- a/Engine/Plugins/Runtime/Metasound/Source/MetasoundEngine/Private/MetasoundSource.cpp
+++ b/Engine/Plugins/Runtime/Metasound/Source/MetasoundEngine/Private/MetasoundSource.cpp
@@ -1271,10 +1271,11 @@ ISoundGeneratorPtr UMetaSoundSource::CreateSoundGenerator(const FSoundGeneratorI
 		{
 			InDefaultParameters.Add({ SourceStartTimeInterface::Inputs::StartTime, InParams.StartTime });
 		}
-		else
+		else if(!bHasWarnedAboutStartTime)
 		{
 			const FString AssetName = GetName();
 			UE_LOG(LogMetaSound, Warning, TEXT("A non-zero StartTime (%.2f) was used in asset '%s' that doesn't implement the Start Time interface. StartTime will be ignored."), InParams.StartTime, *AssetName);
+			bHasWarnedAboutStartTime = true;
 		}
 	}
 
diff --git a/Engine/Plugins/Runtime/Metasound/Source/MetasoundEngine/Public/MetasoundSource.h b/Engine/Plugins/Runtime/Metasound/Source/MetasoundEngine/Public/MetasoundSource.h
--- a/Engine/Plugins/Runtime/Metasound/Source/MetasoundEngine/Public/MetasoundSource.h
+++ b/Engine/Plugins/Runtime/Metasound/Source/MetasoundEngine/Public/MetasoundSource.h
@@ -473,6 +473,7 @@ private:
 	mutable Metasound::FSampleRate CachedAudioDeviceSampleRate = 0;
 	
 	bool bIsBuilderActive = false;
+	bool bHasWarnedAboutStartTime = false;
 
 	// Preset graph inflation is a performance optimization intended for use with the MetaSoundOperatorPool. If multiple presets 
 	// utilize the same base MetaSound, they may be able to share their operators in the operator pool. This makes for a more
