MetaSound BaseNodeConfig: Core type changes

Make FClassInterface copyable, add FVertexInterface conversion,
remove UE_EXPERIMENTAL annotations from FClassInterface and related types.
Change FNodeClassMetadata::DefaultInterface from FVertexInterface to FClassInterface.
Add separate FVertexInterface member to FNode and FGraph for mutable operations.
Remove deprecated SetVertexInterface/IsVertexInterfaceSupported.

Removed from IDataTypeRegistryEntry:
- CreateInputNode(FInputNodeConstructorParams&&)
- CreateConstructorInputNode(FInputNodeConstructorParams&&)
- CreateOutputNode(FOutputNodeConstructorParams&&)
- CreateConstructorOutputNode(FOutputNodeConstructorParams&&)
- CreateVariableNode(FVariableNodeConstructorParams&&)
- CreateVariableMutatorNode(const FNodeInitData&)
- CreateVariableAccessorNode(const FNodeInitData&)
- CreateVariableDeferredAccessorNode(const FNodeInitData&)

Kept (still have callers):
- CreateLiteralNode(FLiteralNodeConstructorParams&&) ??? called by graph builder (CL3) and a test
- CreateReceiveNode(const FNodeInitData&) ??? required by interface implementations

Removed from IDataTypeRegistry:
- Same deprecated input/output/constructor/variable overloads taking old param types
- CreateOutputNode(FOutputNodeConstructorParams&&), CreateConstructorOutputNode(FOutputNodeConstructorParams&&)
- CreateInputNode(FInputNodeConstructorParams&&), CreateConstructorInputNode(FInputNodeConstructorParams&&)
- CreateVariableNode(FVariableNodeConstructorParams&&)
- CreateVariableMutatorNode(FNodeInitData), CreateVariableAccessorNode(FNodeInitData), CreateVariableDeferredAccessorNode(FNodeInitData)
#rnx
#Synced-CL 52453762

diff --git a/Engine/Plugins/Runtime/Metasound/Source/MetasoundFrontend/Private/MetasoundAssetBase.cpp b/Engine/Plugins/Runtime/Metasound/Source/MetasoundFrontend/Private/MetasoundAssetBase.cpp
--- a/Engine/Plugins/Runtime/Metasound/Source/MetasoundFrontend/Private/MetasoundAssetBase.cpp
+++ b/Engine/Plugins/Runtime/Metasound/Source/MetasoundFrontend/Private/MetasoundAssetBase.cpp
@@ -88,13 +88,13 @@ namespace Metasound
 				{
 					// Copy FrontendClass to preserve original document.
 					FrontendClass.Metadata.SetType(EMetasoundFrontendClassType::External);
+					CachedClassInterface = FClassInterface(CreateDefaultVertexInterfaceFromClassNoProxy(FrontendClass));
 				}
 
 				FDocumentNodeRegistryEntryForSerialization(const FDocumentNodeRegistryEntryForSerialization& InOther) = default;
 
 				virtual ~FDocumentNodeRegistryEntryForSerialization() = default;
 
-				virtual TUniquePtr<INode> CreateNode(const FNodeInitData&) const override { return nullptr; }
 				virtual TUniquePtr<INode> CreateNode(FNodeData) const override { return nullptr; }
 
 				virtual const FMetasoundFrontendClass& GetFrontendClass() const override
@@ -113,6 +113,11 @@ namespace Metasound
 					return CreateDefaultVertexInterfaceFromClassNoProxy(FrontendClass);
 				}
 
+				virtual const FClassInterface& GetClassInterface() const override
+				{
+					return CachedClassInterface;
+				}
+
 				virtual TInstancedStruct<FMetaSoundFrontendNodeConfiguration> CreateFrontendNodeConfiguration() const override
 				{
 					// Document nodes do not currently support node configuration.
@@ -129,6 +134,7 @@ namespace Metasound
 				TSet<FMetasoundFrontendVersion> Interfaces;
 				FMetasoundFrontendClass FrontendClass;
 				FTopLevelAssetPath AssetPath;
+				FClassInterface CachedClassInterface;
 			};
 		} // namespace AssetBasePrivate
 
diff --git a/Engine/Plugins/Runtime/Metasound/Source/MetasoundFrontend/Private/MetasoundChannelAgnosticDataSchemas.cpp b/Engine/Plugins/Runtime/Metasound/Source/MetasoundFrontend/Private/MetasoundChannelAgnosticDataSchemas.cpp
--- a/Engine/Plugins/Runtime/Metasound/Source/MetasoundFrontend/Private/MetasoundChannelAgnosticDataSchemas.cpp
+++ b/Engine/Plugins/Runtime/Metasound/Source/MetasoundFrontend/Private/MetasoundChannelAgnosticDataSchemas.cpp
@@ -660,18 +660,6 @@ namespace Metasound
 				return TUniquePtr<INode>(nullptr);
 			}
 
-			virtual TUniquePtr<INode> CreateVariableAccessorNode(const FNodeInitData& InParams) const override
-			{
-				if (Info.bIsParsable)
-				{
-					return MakeUnique<TVariableAccessorNode<TDataType>>(InParams);
-				}
-				else
-				{
-					return TUniquePtr<INode>(nullptr);
-				}
-			}
-
 			virtual TUniquePtr<INode> CreateVariableAccessorNode(FNodeData InNodeData) const override
 			{
 				if (Info.bIsParsable)
@@ -688,9 +676,9 @@ namespace Metasound
 			{
 				if (Info.bIsParsable)
 				{
-					if (VariableAccessorClassMetadata)
+					if (VariableDeferredAccessorClassMetadata)
 					{
-						return MakeUnique<TVariableDeferredAccessorNode<TDataType>>(MoveTemp(InNodeData), VariableAccessorClassMetadata.ToSharedRef());
+						return MakeUnique<TVariableDeferredAccessorNode<TDataType>>(MoveTemp(InNodeData), VariableDeferredAccessorClassMetadata.ToSharedRef());
 					}
 				}
 				return TUniquePtr<INode>(nullptr);
@@ -715,51 +703,9 @@ namespace Metasound
 			virtual TSharedPtr<Audio::IProxyData> CreateProxy(UObject* InObject) const override { return {}; }
 			virtual TSharedPtr<IDataChannel, ESPMode::ThreadSafe> CreateDataChannel(const Metasound::FOperatorSettings&) const override { return {}; }
 
-			// stubbed.
-			virtual TUniquePtr<INode> CreateVariableNode(Metasound::FVariableNodeConstructorParams&& InInitParams) const override
-			{
-				return {};
-			}
-
-			virtual TUniquePtr<INode> CreateLiteralNode(Metasound::FLiteralNodeConstructorParams&&) const override
-			{
-				return {};
-			}
-
-			virtual TUniquePtr<INode> CreateReceiveNode(const Metasound::FNodeInitData&) const override
-			{
-				return {};
-			}
-
-			virtual TUniquePtr<INode> CreateConstructorOutputNode(Metasound::FOutputNodeConstructorParams&&) const override
-			{
-				return {};
-			}
-
-			virtual TUniquePtr<INode> CreateOutputNode(Metasound::FOutputNodeConstructorParams&&) const override
-			{
-				return {};
-			}
-
-			virtual TUniquePtr<INode> CreateVariableMutatorNode(const Metasound::FNodeInitData&) const override
-			{
-				return {};
-			}
-
-			virtual TUniquePtr<INode> CreateVariableDeferredAccessorNode(const Metasound::FNodeInitData&) const override
-			{
-				return {};
-			}
-
-			virtual TUniquePtr<INode> CreateInputNode(Metasound::FInputNodeConstructorParams&&) const override
-			{
-				return {};
-			}
-
-			virtual TUniquePtr<INode> CreateConstructorInputNode(Metasound::FInputNodeConstructorParams&&) const override
-			{
-				return {};
-			}
+			// Stubbed deprecated methods still required by interface
+			virtual TUniquePtr<INode> CreateLiteralNode(Metasound::FLiteralNodeConstructorParams&&) const override { return {}; }
+			virtual TUniquePtr<INode> CreateReceiveNode(const Metasound::FNodeInitData&) const override { return {}; }
 
 			FDataTypeRegistryInfo Info;
 			FMetasoundFrontendClass InputClass;
diff --git a/Engine/Plugins/Runtime/Metasound/Source/MetasoundFrontend/Private/MetasoundFrontend.cpp b/Engine/Plugins/Runtime/Metasound/Source/MetasoundFrontend/Private/MetasoundFrontend.cpp
--- a/Engine/Plugins/Runtime/Metasound/Source/MetasoundFrontend/Private/MetasoundFrontend.cpp
+++ b/Engine/Plugins/Runtime/Metasound/Source/MetasoundFrontend/Private/MetasoundFrontend.cpp
@@ -42,7 +42,7 @@ namespace Metasound
 			FMetasoundFrontendClass ClassDescription;
 
 			ClassDescription.Metadata = FMetasoundFrontendClassMetadata::GenerateClassMetadata(InNodeMetadata, ClassType);
-			ClassDescription.SetDefaultInterface(FMetasoundFrontendClassInterface::GenerateClassInterface(InNodeMetadata.DefaultInterface));
+			ClassDescription.SetDefaultInterface(FMetasoundFrontendClassInterface::GenerateClassInterface(InNodeMetadata.DefaultInterface.CreateVertexInterface({}, {})));
 #if WITH_EDITORONLY_DATA
 			ClassDescription.Style = FMetasoundFrontendClassStyle::GenerateClassStyle(InNodeMetadata.DisplayStyle);
 #endif // WITH_EDITORONLY_DATA
diff --git a/Engine/Plugins/Runtime/Metasound/Source/MetasoundFrontend/Private/MetasoundFrontendDataTypeRegistry.cpp b/Engine/Plugins/Runtime/Metasound/Source/MetasoundFrontend/Private/MetasoundFrontendDataTypeRegistry.cpp
--- a/Engine/Plugins/Runtime/Metasound/Source/MetasoundFrontend/Private/MetasoundFrontendDataTypeRegistry.cpp
+++ b/Engine/Plugins/Runtime/Metasound/Source/MetasoundFrontend/Private/MetasoundFrontendDataTypeRegistry.cpp
@@ -113,6 +113,8 @@ namespace Metasound
 			}
 
 			// Base class for INodeRegistryEntrys that come from an IDataTypeRegistryEntry
+			static const FClassInterface EmptyClassInterface;
+
 			class FDataTypeNodeRegistryEntry : public INodeClassRegistryEntry
 			{
 			public:
@@ -120,11 +122,6 @@ namespace Metasound
 
 				virtual ~FDataTypeNodeRegistryEntry() = default;
 
-				virtual TUniquePtr<INode> CreateNode(const FNodeInitData&) const override
-				{
-					return nullptr;
-				}
-		
 				virtual const TSet<FMetasoundFrontendVersion>* GetImplementedInterfaces() const override
 				{
 					return nullptr;
@@ -135,6 +132,11 @@ namespace Metasound
 					return FVertexInterface();
 				}
 
+				virtual const FClassInterface& GetClassInterface() const override
+				{
+					return EmptyClassInterface;
+				}
+
 				virtual TInstancedStruct<FMetaSoundFrontendNodeConfiguration> CreateFrontendNodeConfiguration() const override
 				{
 					// By default, data type related nodes do not offer any extensions. 
@@ -166,35 +168,29 @@ namespace Metasound
 					return DataTypeEntry->GetFrontendInputClass();
 				}
 
-				using FDataTypeNodeRegistryEntry::CreateNode;
-
 				virtual TUniquePtr<INode> CreateNode(FNodeData InNodeData) const override
 				{
 					return DataTypeEntry->CreateInputNode(MoveTemp(InNodeData));
 				}
 
-				UE_DEPRECATED(5.6, "Node classes should be constructed with FNodeData")
-				virtual TUniquePtr<INode> CreateNode(FDefaultNamedVertexWithLiteralNodeConstructorParams&& InParams) const override
+				virtual FVertexInterface GetDefaultVertexInterface() const override
 				{
-					FInputNodeConstructorParams InputParams;
-					InputParams.InitParam = MoveTemp(InParams.InitParam);
-					InputParams.InstanceID = InParams.InstanceID;
-					InputParams.NodeName = InParams.NodeName;
-					InputParams.VertexName = InParams.VertexName;
-
-PRAGMA_DISABLE_DEPRECATION_WARNINGS
-					return DataTypeEntry->CreateInputNode(MoveTemp(InputParams));
-PRAGMA_ENABLE_DEPRECATION_WARNINGS
+					TSharedPtr<const FNodeClassMetadata> NodeClassMetadata = DataTypeEntry->GetInputClassMetadata();
+					if (ensure(NodeClassMetadata))
+					{
+						return FVertexInterface(NodeClassMetadata->DefaultInterface);
+					}
+					return FVertexInterface();
 				}
 
-				virtual FVertexInterface GetDefaultVertexInterface() const override
+				virtual const FClassInterface& GetClassInterface() const override
 				{
 					TSharedPtr<const FNodeClassMetadata> NodeClassMetadata = DataTypeEntry->GetInputClassMetadata();
 					if (ensure(NodeClassMetadata))
 					{
 						return NodeClassMetadata->DefaultInterface;
 					}
-					return FVertexInterface();
+					return EmptyClassInterface;
 				}
 
 			private:
@@ -220,35 +216,29 @@ PRAGMA_ENABLE_DEPRECATION_WARNINGS
 					return DataTypeEntry->GetFrontendConstructorInputClass();
 				}
 
-				using FDataTypeNodeRegistryEntry::CreateNode;
-
 				virtual TUniquePtr<INode> CreateNode(FNodeData InNodeData) const override
 				{
 					return DataTypeEntry->CreateConstructorInputNode(MoveTemp(InNodeData));
 				}
 
-				UE_DEPRECATED(5.6, "Node classes should be constructed with FNodeData")
-				virtual TUniquePtr<INode> CreateNode(FDefaultNamedVertexWithLiteralNodeConstructorParams&& InParams) const override
+				virtual FVertexInterface GetDefaultVertexInterface() const override
 				{
-					FInputNodeConstructorParams InputParams;
-					InputParams.InitParam = MoveTemp(InParams.InitParam);
-					InputParams.InstanceID = InParams.InstanceID;
-					InputParams.NodeName = InParams.NodeName;
-					InputParams.VertexName = InParams.VertexName;
-
-PRAGMA_DISABLE_DEPRECATION_WARNINGS
-					return DataTypeEntry->CreateConstructorInputNode(MoveTemp(InputParams));
-PRAGMA_ENABLE_DEPRECATION_WARNINGS
+					TSharedPtr<const FNodeClassMetadata> NodeClassMetadata = DataTypeEntry->GetConstructorInputClassMetadata();
+					if (ensure(NodeClassMetadata))
+					{
+						return FVertexInterface(NodeClassMetadata->DefaultInterface);
+					}
+					return FVertexInterface();
 				}
 
-				virtual FVertexInterface GetDefaultVertexInterface() const override
+				virtual const FClassInterface& GetClassInterface() const override
 				{
 					TSharedPtr<const FNodeClassMetadata> NodeClassMetadata = DataTypeEntry->GetConstructorInputClassMetadata();
 					if (ensure(NodeClassMetadata))
 					{
 						return NodeClassMetadata->DefaultInterface;
 					}
-					return FVertexInterface();
+					return EmptyClassInterface;
 				}
 
 			private:
@@ -274,8 +264,6 @@ PRAGMA_ENABLE_DEPRECATION_WARNINGS
 					return DataTypeEntry->GetFrontendOutputClass();
 				}
 
-				using FDataTypeNodeRegistryEntry::CreateNode;
-
 				virtual TUniquePtr<INode> CreateNode(FNodeData InNodeData) const override
 				{
 					return DataTypeEntry->CreateOutputNode(MoveTemp(InNodeData));
@@ -286,11 +274,21 @@ PRAGMA_ENABLE_DEPRECATION_WARNINGS
 					TSharedPtr<const FNodeClassMetadata> NodeClassMetadata = DataTypeEntry->GetOutputClassMetadata();
 					if (ensure(NodeClassMetadata))
 					{
-						return NodeClassMetadata->DefaultInterface;
+						return FVertexInterface(NodeClassMetadata->DefaultInterface);
 					}
 					return FVertexInterface();
 				}
 
+				virtual const FClassInterface& GetClassInterface() const override
+				{
+					TSharedPtr<const FNodeClassMetadata> NodeClassMetadata = DataTypeEntry->GetOutputClassMetadata();
+					if (ensure(NodeClassMetadata))
+					{
+						return NodeClassMetadata->DefaultInterface;
+					}
+					return EmptyClassInterface;
+				}
+
 			private:
 				
 				TSharedRef<IDataTypeRegistryEntry, ESPMode::ThreadSafe> DataTypeEntry;
@@ -314,8 +312,6 @@ PRAGMA_ENABLE_DEPRECATION_WARNINGS
 					return DataTypeEntry->GetFrontendConstructorOutputClass();
 				}
 
-				using FDataTypeNodeRegistryEntry::CreateNode;
-
 				virtual TUniquePtr<INode> CreateNode(FNodeData InNodeData) const override
 				{
 					return DataTypeEntry->CreateConstructorOutputNode(MoveTemp(InNodeData));
@@ -326,11 +322,21 @@ PRAGMA_ENABLE_DEPRECATION_WARNINGS
 					TSharedPtr<const FNodeClassMetadata> NodeClassMetadata = DataTypeEntry->GetConstructorOutputClassMetadata();
 					if (ensure(NodeClassMetadata))
 					{
-						return NodeClassMetadata->DefaultInterface;
+						return FVertexInterface(NodeClassMetadata->DefaultInterface);
 					}
 					return FVertexInterface();
 				}
 
+				virtual const FClassInterface& GetClassInterface() const override
+				{
+					TSharedPtr<const FNodeClassMetadata> NodeClassMetadata = DataTypeEntry->GetConstructorOutputClassMetadata();
+					if (ensure(NodeClassMetadata))
+					{
+						return NodeClassMetadata->DefaultInterface;
+					}
+					return EmptyClassInterface;
+				}
+
 			private:
 				
 				TSharedRef<IDataTypeRegistryEntry, ESPMode::ThreadSafe> DataTypeEntry;
@@ -365,14 +371,6 @@ PRAGMA_ENABLE_DEPRECATION_WARNINGS
 					return nullptr;
 				}
 
-				UE_DEPRECATED(5.6, "Literal nodes should no longer be created from the registry")
-				virtual TUniquePtr<INode> CreateNode(FDefaultLiteralNodeConstructorParams&& InParams) const override
-				{
-PRAGMA_DISABLE_DEPRECATION_WARNINGS	
-					return DataTypeEntry->CreateLiteralNode(MoveTemp(InParams));
-PRAGMA_ENABLE_DEPRECATION_WARNINGS	
-				}
-
 			private:
 				
 				TSharedRef<IDataTypeRegistryEntry, ESPMode::ThreadSafe> DataTypeEntry;
@@ -406,22 +404,24 @@ PRAGMA_ENABLE_DEPRECATION_WARNINGS
 					return nullptr;
 				}
 
-				UE_DEPRECATED(5.6, "Node classes should be constructed from the data type registry") 
-				virtual TUniquePtr<INode> CreateNode(FDefaultLiteralNodeConstructorParams&& InParams) const override
+				virtual FVertexInterface GetDefaultVertexInterface() const override
 				{
-PRAGMA_DISABLE_DEPRECATION_WARNINGS
-					return DataTypeEntry->CreateVariableNode(MoveTemp(InParams));
-PRAGMA_ENABLE_DEPRECATION_WARNINGS
+					TSharedPtr<const FNodeClassMetadata> NodeClassMetadata = DataTypeEntry->GetVariableClassMetadata();
+					if (ensure(NodeClassMetadata))
+					{
+						return FVertexInterface(NodeClassMetadata->DefaultInterface);
+					}
+					return FVertexInterface();
 				}
 
-				virtual FVertexInterface GetDefaultVertexInterface() const override
+				virtual const FClassInterface& GetClassInterface() const override
 				{
 					TSharedPtr<const FNodeClassMetadata> NodeClassMetadata = DataTypeEntry->GetVariableClassMetadata();
 					if (ensure(NodeClassMetadata))
 					{
 						return NodeClassMetadata->DefaultInterface;
 					}
-					return FVertexInterface();
+					return EmptyClassInterface;
 				}
 
 			private:
@@ -447,13 +447,6 @@ PRAGMA_ENABLE_DEPRECATION_WARNINGS
 					return DataTypeEntry->GetFrontendVariableMutatorClass();
 				}
 
-				virtual TUniquePtr<INode> CreateNode(const FNodeInitData& InParams) const override
-				{
-PRAGMA_DISABLE_DEPRECATION_WARNINGS
-					return DataTypeEntry->CreateVariableMutatorNode(InParams);
-PRAGMA_ENABLE_DEPRECATION_WARNINGS
-				}
-
 				virtual TUniquePtr<INode> CreateNode(FNodeData InNodeData) const override
 				{
 					return DataTypeEntry->CreateVariableMutatorNode(MoveTemp(InNodeData));
@@ -464,11 +457,21 @@ PRAGMA_ENABLE_DEPRECATION_WARNINGS
 					TSharedPtr<const FNodeClassMetadata> NodeClassMetadata = DataTypeEntry->GetVariableMutatorClassMetadata();
 					if (ensure(NodeClassMetadata))
 					{
-						return NodeClassMetadata->DefaultInterface;
+						return FVertexInterface(NodeClassMetadata->DefaultInterface);
 					}
 					return FVertexInterface();
 				}
 
+				virtual const FClassInterface& GetClassInterface() const override
+				{
+					TSharedPtr<const FNodeClassMetadata> NodeClassMetadata = DataTypeEntry->GetVariableMutatorClassMetadata();
+					if (ensure(NodeClassMetadata))
+					{
+						return NodeClassMetadata->DefaultInterface;
+					}
+					return EmptyClassInterface;
+				}
+
 			private:
 				
 				TSharedRef<IDataTypeRegistryEntry, ESPMode::ThreadSafe> DataTypeEntry;
@@ -492,13 +495,6 @@ PRAGMA_ENABLE_DEPRECATION_WARNINGS
 					return DataTypeEntry->GetFrontendVariableAccessorClass();
 				}
 
-				virtual TUniquePtr<INode> CreateNode(const FNodeInitData& InParams) const override
-				{
-PRAGMA_DISABLE_DEPRECATION_WARNINGS
-					return DataTypeEntry->CreateVariableAccessorNode(InParams);
-PRAGMA_ENABLE_DEPRECATION_WARNINGS
-				}
-
 				virtual TUniquePtr<INode> CreateNode(FNodeData InNodeData) const override
 				{
 					return DataTypeEntry->CreateVariableAccessorNode(MoveTemp(InNodeData));
@@ -509,10 +505,20 @@ PRAGMA_ENABLE_DEPRECATION_WARNINGS
 					TSharedPtr<const FNodeClassMetadata> NodeClassMetadata = DataTypeEntry->GetVariableAccessorClassMetadata();
 					if (ensure(NodeClassMetadata))
 					{
-						return NodeClassMetadata->DefaultInterface;
+						return FVertexInterface(NodeClassMetadata->DefaultInterface);
 					}
 					return FVertexInterface();
 				}
+				
+				virtual const FClassInterface& GetClassInterface() const override
+				{
+					TSharedPtr<const FNodeClassMetadata> NodeClassMetadata = DataTypeEntry->GetVariableAccessorClassMetadata();
+					if (ensure(NodeClassMetadata))
+					{
+						return NodeClassMetadata->DefaultInterface;
+					}
+					return EmptyClassInterface;
+				}
 
 			private:
 				
@@ -537,13 +543,6 @@ PRAGMA_ENABLE_DEPRECATION_WARNINGS
 					return DataTypeEntry->GetFrontendVariableDeferredAccessorClass();
 				}
 
-				virtual TUniquePtr<INode> CreateNode(const FNodeInitData& InParams) const override
-				{
-PRAGMA_DISABLE_DEPRECATION_WARNINGS
-					return DataTypeEntry->CreateVariableDeferredAccessorNode(InParams);
-PRAGMA_ENABLE_DEPRECATION_WARNINGS
-				}
-
 				virtual TUniquePtr<INode> CreateNode(FNodeData InNodeData) const override
 				{
 					return DataTypeEntry->CreateVariableDeferredAccessorNode(MoveTemp(InNodeData));
@@ -554,11 +553,21 @@ PRAGMA_ENABLE_DEPRECATION_WARNINGS
 					TSharedPtr<const FNodeClassMetadata> NodeClassMetadata = DataTypeEntry->GetVariableDeferredAccessorClassMetadata();
 					if (ensure(NodeClassMetadata))
 					{
-						return NodeClassMetadata->DefaultInterface;
+						return FVertexInterface(NodeClassMetadata->DefaultInterface);
 					}
 					return FVertexInterface();
 				}
 
+				virtual const FClassInterface& GetClassInterface() const override
+				{
+					TSharedPtr<const FNodeClassMetadata> NodeClassMetadata = DataTypeEntry->GetVariableDeferredAccessorClassMetadata();
+					if (ensure(NodeClassMetadata))
+					{
+						return NodeClassMetadata->DefaultInterface;
+					}
+					return EmptyClassInterface;
+				}
+
 			private:
 				
 				TSharedRef<IDataTypeRegistryEntry, ESPMode::ThreadSafe> DataTypeEntry;
@@ -628,42 +637,20 @@ PRAGMA_ENABLE_DEPRECATION_WARNINGS
 				virtual TSharedPtr<const FNodeClassMetadata> GetVariableDeferredAccessorClassMetadata(const FName& InDataType) const override;
 
 				// Create a new instance of a C++ implemented node from the registry.
-				UE_DEPRECATED(5.6, "Create input nodes with FNodeData")
-				virtual TUniquePtr<INode> CreateInputNode(const FName& InInputType, FInputNodeConstructorParams&& InParams) const override;
 				virtual TUniquePtr<INode> CreateInputNode(const FName& InInputType, FNodeData InNodeData) const override;
-
-				UE_DEPRECATED(5.6, "Create constructor input nodes with FNodeData")
-				virtual TUniquePtr<INode> CreateConstructorInputNode(const FName& InInputType, FInputNodeConstructorParams&& InParams) const override;
 				virtual TUniquePtr<INode> CreateConstructorInputNode(const FName& InInputType, FNodeData InNodeData) const override;
 
 				UE_DEPRECATED(5.6, "Direct creation of literal nodes will no longer be supported")
 				virtual TUniquePtr<INode> CreateLiteralNode(const FName& InLiteralType, FLiteralNodeConstructorParams&& InParams) const override;
 
-				UE_DEPRECATED(5.6, "Create output nodes with FNodeData")
-				virtual TUniquePtr<INode> CreateOutputNode(const FName& InDataTypeName, FOutputNodeConstructorParams&& InParams) const override;
-				virtual TUniquePtr<INode> CreateOutputNode(const FName& InDataTypeName, FNodeData InNodeData) const override;
-
-				UE_DEPRECATED(5.6, "Create output nodes with FNodeData")
-				virtual TUniquePtr<INode> CreateConstructorOutputNode(const FName& InDataTypeName, FOutputNodeConstructorParams&& InParams) const override;
-				virtual TUniquePtr<INode> CreateConstructorOutputNode(const FName& InDataTypeName, FNodeData InNodeData) const override;
-
 				UE_DEPRECATED(5.6, "Direct creation of receive nodes will no longer be supported")
 				virtual TUniquePtr<INode> CreateReceiveNode(const FName& InDataTypeName, const FNodeInitData& InParams) const override;
 
-				UE_DEPRECATED(5.6, "Create Variable Node with FNodeData")
-				virtual TUniquePtr<INode> CreateVariableNode(const FName& InDataType, FVariableNodeConstructorParams&&) const override;
+				virtual TUniquePtr<INode> CreateOutputNode(const FName& InDataTypeName, FNodeData InNodeData) const override;
+				virtual TUniquePtr<INode> CreateConstructorOutputNode(const FName& InDataTypeName, FNodeData InNodeData) const override;
 				virtual TUniquePtr<INode> CreateVariableNode(const FName& InDataType, FLiteral InLiteral, FNodeData InNodeData) const override;
-
-				UE_DEPRECATED(5.6, "Create variable mutator nodes with FNodeData")
-				virtual TUniquePtr<INode> CreateVariableMutatorNode(const FName& InDataType, const FNodeInitData&) const override;
 				virtual TUniquePtr<INode> CreateVariableMutatorNode(const FName& InDataType, FNodeData InNodeData) const override;
-
-				UE_DEPRECATED(5.6, "Create variable accessor nodes with FNodeData")
-				virtual TUniquePtr<INode> CreateVariableAccessorNode(const FName& InDataType, const FNodeInitData&) const override;
 				virtual TUniquePtr<INode> CreateVariableAccessorNode(const FName& InDataType, FNodeData InNodeData) const override;
-
-				UE_DEPRECATED(5.6, "Create variable deferred accessor nodes with FNodeData")
-				virtual TUniquePtr<INode> CreateVariableDeferredAccessorNode(const FName& InDataType, const FNodeInitData&) const override;
 				virtual TUniquePtr<INode> CreateVariableDeferredAccessorNode(const FName& InDataType, FNodeData InNodeData) const override;
 
 			private:
@@ -1376,17 +1363,6 @@ PRAGMA_ENABLE_DEPRECATION_WARNINGS
 			}
 
 
-			TUniquePtr<INode> FDataTypeRegistry::CreateInputNode(const FName& InDataType, FInputNodeConstructorParams&& InParams) const
-			{
-				if (const IDataTypeRegistryEntry* Entry = FindDataTypeEntry(InDataType))
-				{
-PRAGMA_DISABLE_DEPRECATION_WARNINGS
-					return Entry->CreateInputNode(MoveTemp(InParams));
-PRAGMA_ENABLE_DEPRECATION_WARNINGS
-				}
-				return nullptr;
-			}
-
 			TUniquePtr<INode> FDataTypeRegistry::CreateInputNode(const FName& InDataType, FNodeData InNodeData) const
 			{
 				if (const IDataTypeRegistryEntry* Entry = FindDataTypeEntry(InDataType))
@@ -1396,17 +1372,6 @@ PRAGMA_ENABLE_DEPRECATION_WARNINGS
 				return nullptr;
 			}
 
-			TUniquePtr<INode> FDataTypeRegistry::CreateConstructorInputNode(const FName& InDataType, FInputNodeConstructorParams&& InParams) const
-			{
-				if (const IDataTypeRegistryEntry* Entry = FindDataTypeEntry(InDataType))
-				{
-PRAGMA_DISABLE_DEPRECATION_WARNINGS
-					return Entry->CreateConstructorInputNode(MoveTemp(InParams));
-PRAGMA_ENABLE_DEPRECATION_WARNINGS
-				}
-				return nullptr;
-			}
-
 			TUniquePtr<INode> FDataTypeRegistry::CreateConstructorInputNode(const FName& InDataType, FNodeData InNodeData) const
 			{
 				if (const IDataTypeRegistryEntry* Entry = FindDataTypeEntry(InDataType))
@@ -1416,7 +1381,7 @@ PRAGMA_ENABLE_DEPRECATION_WARNINGS
 				return nullptr;
 			}
 
-			TUniquePtr<INode> FDataTypeRegistry::CreateLiteralNode(const FName& InDataType, FLiteralNodeConstructorParams&& InParams) const 
+			TUniquePtr<INode> FDataTypeRegistry::CreateLiteralNode(const FName& InDataType, FLiteralNodeConstructorParams&& InParams) const
 			{
 				if (const IDataTypeRegistryEntry* Entry = FindDataTypeEntry(InDataType))
 				{
@@ -1427,12 +1392,12 @@ PRAGMA_ENABLE_DEPRECATION_WARNINGS
 				return nullptr;
 			}
 
-			TUniquePtr<INode> FDataTypeRegistry::CreateOutputNode(const FName& InDataType, FOutputNodeConstructorParams&& InParams) const
+			TUniquePtr<INode> FDataTypeRegistry::CreateReceiveNode(const FName& InDataType, const FNodeInitData& InParams) const
 			{
 				if (const IDataTypeRegistryEntry* Entry = FindDataTypeEntry(InDataType))
 				{
 PRAGMA_DISABLE_DEPRECATION_WARNINGS
-					return Entry->CreateOutputNode(MoveTemp(InParams));
+					return Entry->CreateReceiveNode(InParams);
 PRAGMA_ENABLE_DEPRECATION_WARNINGS
 				}
 				return nullptr;
@@ -1447,17 +1412,6 @@ PRAGMA_ENABLE_DEPRECATION_WARNINGS
 				return nullptr;
 			}
 
-			TUniquePtr<INode> FDataTypeRegistry::CreateConstructorOutputNode(const FName& InDataType, FOutputNodeConstructorParams&& InParams) const
-			{
-				if (const IDataTypeRegistryEntry* Entry = FindDataTypeEntry(InDataType))
-				{
-PRAGMA_DISABLE_DEPRECATION_WARNINGS
-					return Entry->CreateConstructorOutputNode(MoveTemp(InParams));
-PRAGMA_ENABLE_DEPRECATION_WARNINGS
-				}
-				return nullptr;
-			}
-
 			TUniquePtr<INode> FDataTypeRegistry::CreateConstructorOutputNode(const FName& InDataTypeName, FNodeData InNodeData) const
 			{
 				if (const IDataTypeRegistryEntry* Entry = FindDataTypeEntry(InDataTypeName))
@@ -1467,28 +1421,6 @@ PRAGMA_ENABLE_DEPRECATION_WARNINGS
 				return nullptr;
 			}
 
-			TUniquePtr<INode> FDataTypeRegistry::CreateReceiveNode(const FName& InDataType, const FNodeInitData& InParams) const
-			{
-				if (const IDataTypeRegistryEntry* Entry = FindDataTypeEntry(InDataType))
-				{
-PRAGMA_DISABLE_DEPRECATION_WARNINGS
-					return Entry->CreateReceiveNode(InParams);
-PRAGMA_ENABLE_DEPRECATION_WARNINGS
-				}
-				return nullptr;
-			}
-
-			TUniquePtr<INode> FDataTypeRegistry::CreateVariableNode(const FName& InDataType, FVariableNodeConstructorParams&& InParams) const
-			{
-				if (const IDataTypeRegistryEntry* Entry = FindDataTypeEntry(InDataType))
-				{
-PRAGMA_DISABLE_DEPRECATION_WARNINGS
-					return Entry->CreateVariableNode(MoveTemp(InParams));
-PRAGMA_ENABLE_DEPRECATION_WARNINGS
-				}
-				return nullptr;
-			}
-
 			TUniquePtr<INode> FDataTypeRegistry::CreateVariableNode(const FName& InDataType, FLiteral InLiteral, FNodeData InNodeData) const
 			{
 				if (const IDataTypeRegistryEntry* Entry = FindDataTypeEntry(InDataType))
@@ -1498,17 +1430,6 @@ PRAGMA_ENABLE_DEPRECATION_WARNINGS
 				return nullptr;
 			}
 
-			TUniquePtr<INode> FDataTypeRegistry::CreateVariableMutatorNode(const FName& InDataType, const FNodeInitData& InParams) const
-			{
-				if (const IDataTypeRegistryEntry* Entry = FindDataTypeEntry(InDataType))
-				{
-PRAGMA_DISABLE_DEPRECATION_WARNINGS
-					return Entry->CreateVariableMutatorNode(InParams);
-PRAGMA_ENABLE_DEPRECATION_WARNINGS
-				}
-				return nullptr;
-			}
-
 			TUniquePtr<INode> FDataTypeRegistry::CreateVariableMutatorNode(const FName& InDataType, FNodeData InNodeData) const 
 			{
 				if (const IDataTypeRegistryEntry* Entry = FindDataTypeEntry(InDataType))
@@ -1518,17 +1439,6 @@ PRAGMA_ENABLE_DEPRECATION_WARNINGS
 				return nullptr;
 			}
 
-			TUniquePtr<INode> FDataTypeRegistry::CreateVariableAccessorNode(const FName& InDataType, const FNodeInitData& InParams) const
-			{
-				if (const IDataTypeRegistryEntry* Entry = FindDataTypeEntry(InDataType))
-				{
-PRAGMA_DISABLE_DEPRECATION_WARNINGS
-					return Entry->CreateVariableAccessorNode(InParams);
-PRAGMA_ENABLE_DEPRECATION_WARNINGS
-				}
-				return nullptr;
-			}
-
 			TUniquePtr<INode> FDataTypeRegistry::CreateVariableAccessorNode(const FName& InDataType, FNodeData InNodeData) const 
 			{
 				if (const IDataTypeRegistryEntry* Entry = FindDataTypeEntry(InDataType))
@@ -1538,17 +1448,6 @@ PRAGMA_ENABLE_DEPRECATION_WARNINGS
 				return nullptr;
 			}
 
-			TUniquePtr<INode> FDataTypeRegistry::CreateVariableDeferredAccessorNode(const FName& InDataType, const FNodeInitData& InParams) const
-			{
-				if (const IDataTypeRegistryEntry* Entry = FindDataTypeEntry(InDataType))
-				{
-PRAGMA_DISABLE_DEPRECATION_WARNINGS
-					return Entry->CreateVariableDeferredAccessorNode(InParams);
-PRAGMA_ENABLE_DEPRECATION_WARNINGS
-				}
-				return nullptr;
-			}
-
 			TUniquePtr<INode> FDataTypeRegistry::CreateVariableDeferredAccessorNode(const FName& InDataType, FNodeData InNodeData) const 
 			{
 				if (const IDataTypeRegistryEntry* Entry = FindDataTypeEntry(InDataType))
diff --git a/Engine/Plugins/Runtime/Metasound/Source/MetasoundFrontend/Private/MetasoundFrontendNodeClassRegistry.cpp b/Engine/Plugins/Runtime/Metasound/Source/MetasoundFrontend/Private/MetasoundFrontendNodeClassRegistry.cpp
--- a/Engine/Plugins/Runtime/Metasound/Source/MetasoundFrontend/Private/MetasoundFrontendNodeClassRegistry.cpp
+++ b/Engine/Plugins/Runtime/Metasound/Source/MetasoundFrontend/Private/MetasoundFrontendNodeClassRegistry.cpp
@@ -129,21 +129,6 @@ namespace Metasound::Frontend
 
 			virtual ~FDocumentNodeClassRegistryEntry() = default;
 
-			virtual TUniquePtr<INode> CreateNode(const FNodeInitData& InNodeInitData) const override
-			{
-				if (Graph.IsValid())
-				{
-					return MakeUnique<FGraphNode>(InNodeInitData, Graph.ToSharedRef());
-				}
-				else
-				{
-					UE_LOG(LogMetaSound, Error,
-						TEXT("Cannot create MetaSound node from class %s due to prior failure to build graph"),
-						*FrontendClass.Metadata.GetClassName().ToString());
-					return TUniquePtr<INode>();
-				}
-			}
-
 			virtual TUniquePtr<INode> CreateNode(FNodeData InNodeData) const override
 			{
 				if (Graph.IsValid())
@@ -171,7 +156,7 @@ namespace Metasound::Frontend
 			{
 				if (ensure(Graph.IsValid()))
 				{
-					return Graph->GetMetadata().DefaultInterface;
+					return FVertexInterface(Graph->GetMetadata().DefaultInterface);
 				}
 				else
 				{
@@ -179,6 +164,19 @@ namespace Metasound::Frontend
 				}
 			}
 
+			virtual const FClassInterface& GetClassInterface() const override
+			{
+				if (ensure(Graph.IsValid()))
+				{
+					return Graph->GetMetadata().DefaultInterface;
+				}
+				else
+				{
+					static const FClassInterface EmptyClassInterface;
+					return EmptyClassInterface;
+				}
+			}
+
 			virtual TInstancedStruct<FMetaSoundFrontendNodeConfiguration> CreateFrontendNodeConfiguration() const override final
 			{
 				// Document based nodes do not support node configuration because
@@ -207,26 +205,6 @@ namespace Metasound::Frontend
 		return FNodeClassInfo::GetInvalid();
 	}
 
-	TUniquePtr<INode> INodeClassRegistryEntry::CreateNode(FNodeData InNodeData) const
-	{
-#if !UE_METASOUND_DISABLE_5_6_NODE_REGISTRATION_DEPRECATION_WARNINGS
-		static bool bDidLogError = false;
-
-		if (!bDidLogError)
-		{
-			bDidLogError = true;
-			UE_LOG(LogMetaSound, Warning,
-				TEXT("Use of deprecated code path for node registration. First occurrence on node %s. "
-					"Please implement INodeClassRegistryEntry::CreateNode(FNodeData) const"),
-				*GetFrontendClass().Metadata.GetClassName().ToString());
-		}
-#endif
-PRAGMA_DISABLE_DEPRECATION_WARNINGS
-		// Use old path for node creation. 
-		return CreateNode(FNodeInitData{InNodeData.Name, InNodeData.ID});
-PRAGMA_ENABLE_DEPRECATION_WARNINGS
-	}
-
 #if !UE_METASOUND_PURE_VIRTUAL_CREATE_FRONTEND_NODE_EXTENSION
 	TInstancedStruct<FMetaSoundFrontendNodeConfiguration> INodeClassRegistryEntry::CreateFrontendNodeConfiguration() const
 	{
@@ -487,33 +465,6 @@ PRAGMA_ENABLE_DEPRECATION_WARNINGS
 		ObjectReferencer = MoveTemp(InReferencer);
 	}
 
-	TUniquePtr<Metasound::INode> FNodeClassRegistry::CreateNode(const FNodeClassRegistryKey& InKey, const Metasound::FNodeInitData& InInitData) const
-	{
-		TUniquePtr<INode> Node;
-
-		auto CreateNodeLambda = [&Node, &InInitData](const INodeClassRegistryEntry& Entry) mutable
-		{ 
-			const FMetasoundFrontendClass& FrontendClass = Entry.GetFrontendClass();
-
-			FNodeData NodeData
-			{
-				InInitData.InstanceName,
-				InInitData.InstanceID,
-				Entry.GetDefaultVertexInterface()
-			};
-			Node = Entry.CreateNode(MoveTemp(NodeData)); 
-		};
-
-		if (!AccessNodeEntryThreadSafe(InKey, CreateNodeLambda))
-		{
-			// Creation of external nodes can rely on assets being unavailable due to errors in loading order, asset(s)
-			// missing, etc. 
-			UE_LOG(LogMetaSound, Error, TEXT("Could not find node [RegistryKey:%s]"), *InKey.ToString());
-		}
-
-		return MoveTemp(Node);
-	}
-
 	TUniquePtr<Metasound::INode> FNodeClassRegistry::CreateNode(const FNodeClassRegistryKey& InKey, Metasound::FNodeData InNodeData) const
 	{
 		TUniquePtr<INode> Node;
@@ -1287,6 +1238,21 @@ PRAGMA_ENABLE_DEPRECATION_WARNINGS
 		return false;
 	}
 
+	bool FNodeClassRegistry::FindClassInterface(const FNodeClassRegistryKey& InKey, FClassInterface& OutClassInterface) const
+	{
+		auto GetClassInterface = [&OutClassInterface](const INodeClassRegistryEntry& Entry)
+		{
+			OutClassInterface = Entry.GetClassInterface();
+		};
+
+		if (AccessNodeEntryThreadSafe(InKey, GetClassInterface))
+		{
+			return true;
+		}
+
+		return false;
+	}
+
 	bool FNodeClassRegistry::FindFrontendClassFromRegistered(const FNodeClassRegistryKey& InKey, FMetasoundFrontendClass& OutClass) const
 	{
 		auto SetFrontendClass = [&OutClass](const INodeClassRegistryEntry& Entry)
diff --git a/Engine/Plugins/Runtime/Metasound/Source/MetasoundFrontend/Private/MetasoundFrontendNodeClassRegistryPrivate.h b/Engine/Plugins/Runtime/Metasound/Source/MetasoundFrontend/Private/MetasoundFrontendNodeClassRegistryPrivate.h
--- a/Engine/Plugins/Runtime/Metasound/Source/MetasoundFrontend/Private/MetasoundFrontendNodeClassRegistryPrivate.h
+++ b/Engine/Plugins/Runtime/Metasound/Source/MetasoundFrontend/Private/MetasoundFrontendNodeClassRegistryPrivate.h
@@ -110,6 +110,7 @@ namespace Metasound::Frontend
 
 		// Find Frontend class data.
 		virtual bool FindDefaultVertexInterface(const FNodeClassRegistryKey& InKey, FVertexInterface& OutVertexInterface) const override;
+		virtual bool FindClassInterface(const FNodeClassRegistryKey& InKey, FClassInterface& OutClassInterface) const override;
 		
 		virtual bool FindFrontendClassFromRegistered(const FNodeClassRegistryKey& InKey, FMetasoundFrontendClass& OutClass) const override;
 
@@ -122,10 +123,6 @@ namespace Metasound::Frontend
 		virtual bool FindVariableNodeRegistryKeyForDataType(const FName& InDataTypeName, FNodeClassRegistryKey& OutKey) override;
 		virtual bool FindOutputNodeRegistryKeyForDataType(const FName& InDataTypeName, const EMetasoundFrontendVertexAccessType InAccessType, FNodeClassRegistryKey& OutKey) override;
 
-		// Create a new instance of a C++ implemented node from the registry.
-		UE_DEPRECATED(5.6, "Node classes should be constructed with FNodeData")
-		virtual TUniquePtr<Metasound::INode> CreateNode(const FNodeClassRegistryKey& InKey, const Metasound::FNodeInitData& InInitData) const override;
-
 		// Create a new instance of a C++ implemented node from the registry.
 		virtual TUniquePtr<Metasound::INode> CreateNode(const FNodeClassRegistryKey& InKey, Metasound::FNodeData InNodeData) const override;
 
diff --git a/Engine/Plugins/Runtime/Metasound/Source/MetasoundFrontend/Private/MetasoundGraphNode.cpp b/Engine/Plugins/Runtime/Metasound/Source/MetasoundFrontend/Private/MetasoundGraphNode.cpp
--- a/Engine/Plugins/Runtime/Metasound/Source/MetasoundFrontend/Private/MetasoundGraphNode.cpp
+++ b/Engine/Plugins/Runtime/Metasound/Source/MetasoundFrontend/Private/MetasoundGraphNode.cpp
@@ -38,7 +38,7 @@ namespace Metasound::Frontend
 	}
 
 	FGraphNode::FGraphNode(const FNodeInitData& InNodeInitData, TSharedRef<const IGraph> InGraphToWrap)
-	: FGraphNode(FNodeData(InNodeInitData.InstanceName, InNodeInitData.InstanceID, InGraphToWrap->GetMetadata().DefaultInterface), MoveTemp(InGraphToWrap))
+	: FGraphNode(FNodeData(InNodeInitData.InstanceName, InNodeInitData.InstanceID, InGraphToWrap->GetVertexInterface()), MoveTemp(InGraphToWrap))
 	{
 	}
 
diff --git a/Engine/Plugins/Runtime/Metasound/Source/MetasoundFrontend/Private/MetasoundLiteralNode.cpp b/Engine/Plugins/Runtime/Metasound/Source/MetasoundFrontend/Private/MetasoundLiteralNode.cpp
--- a/Engine/Plugins/Runtime/Metasound/Source/MetasoundFrontend/Private/MetasoundLiteralNode.cpp
+++ b/Engine/Plugins/Runtime/Metasound/Source/MetasoundFrontend/Private/MetasoundLiteralNode.cpp
@@ -57,16 +57,6 @@ namespace Metasound
 		return Interface;
 	}
 
-	bool FLiteralNode::SetVertexInterface(const FVertexInterface& InInterface)
-	{
-		return Interface == InInterface;
-	}
-
-	bool FLiteralNode::IsVertexInterfaceSupported(const FVertexInterface& InInterface) const
-	{
-		return Interface == InInterface;
-	}
-
 	TSharedRef<IOperatorFactory, ESPMode::ThreadSafe> FLiteralNode::GetDefaultOperatorFactory() const
 	{
 		return Factory;
diff --git a/Engine/Plugins/Runtime/Metasound/Source/MetasoundFrontend/Private/MetasoundNodeRegistrationMacro.cpp b/Engine/Plugins/Runtime/Metasound/Source/MetasoundFrontend/Private/MetasoundNodeRegistrationMacro.cpp
--- a/Engine/Plugins/Runtime/Metasound/Source/MetasoundFrontend/Private/MetasoundNodeRegistrationMacro.cpp
+++ b/Engine/Plugins/Runtime/Metasound/Source/MetasoundFrontend/Private/MetasoundNodeRegistrationMacro.cpp
@@ -25,17 +25,17 @@ namespace Metasound::Frontend
 			return FrontendClass;
 		}
 
-		TUniquePtr<INode> FNodeRegistryEntryBase::CreateNode(const FNodeInitData& InInitData) const 
-		{
-			return CreateNode(FNodeData{InInitData.InstanceName, InInitData.InstanceID, ClassMetadata->DefaultInterface});
-		}
-
 		const TSet<FMetasoundFrontendVersion>* FNodeRegistryEntryBase::GetImplementedInterfaces() const
 		{
 			return nullptr;
 		}
 
 		FVertexInterface FNodeRegistryEntryBase::GetDefaultVertexInterface() const
+		{
+			return FVertexInterface(ClassMetadata->DefaultInterface);
+		}
+
+		const FClassInterface& FNodeRegistryEntryBase::GetClassInterface() const
 		{
 			return ClassMetadata->DefaultInterface;
 		}
@@ -67,50 +67,6 @@ namespace Metasound::Frontend
 		}
 #endif
 
-		// This adapter class forwards the correct FBuilderOperatorParams
-		// to the node's operator creation method. Many operator creation
-		// methods downcast the supplied INode in `FBuilderOperatorParams`
-		// and so it is required that it point to the correct runtime instance
-		// when calling CreateOperator(...)
-		class FDeprecatedNodeAPIFactory : public IOperatorFactory 
-		{
-		public:
-			FDeprecatedNodeAPIFactory(TUniquePtr<INode> InNode)
-			: Node(MoveTemp(InNode))
-			{
-				check(Node);
-			}
-
-			virtual ~FDeprecatedNodeAPIFactory() = default;
-
-			virtual TUniquePtr<IOperator> CreateOperator(const FBuildOperatorParams& InParams, FBuildResults& OutResults) override
-			{
-				FBuildOperatorParams ForwardParams 
-				{
-					*Node,  // Point to correct INode instance
-					InParams.OperatorSettings,
-					InParams.InputData,
-					InParams.Environment,
-					InParams.Builder,
-					InParams.GraphRenderCost
-				};
-				return Node->GetDefaultOperatorFactory()->CreateOperator(ForwardParams, OutResults);
-			}
-		private:
-			TUniquePtr<INode> Node;
-		};
-
-		FDeprecatedNodeAPIAdapterBase::FDeprecatedNodeAPIAdapterBase(FNodeData InNodeData, TSharedRef<const FNodeClassMetadata> InClassMetadata, TUniquePtr<INode> InNode)
-		: FBasicNode(MoveTemp(InNodeData), MoveTemp(InClassMetadata))
-		, Factory(MakeShared<FDeprecatedNodeAPIFactory>(MoveTemp(InNode)))
-		{
-		}
-
-		FOperatorFactorySharedRef FDeprecatedNodeAPIAdapterBase::GetDefaultOperatorFactory() const 
-		{
-			return Factory;
-		}
-
 		bool RegisterNodeInternal(TUniquePtr<INodeClassRegistryEntry>&& InEntry)
 		{
 			Frontend::FNodeRegistryKey Key = INodeClassRegistry::Get()->RegisterNode(MoveTemp(InEntry));
diff --git a/Engine/Plugins/Runtime/Metasound/Source/MetasoundFrontend/Public/MetasoundDataTypeRegistrationMacro.h b/Engine/Plugins/Runtime/Metasound/Source/MetasoundFrontend/Public/MetasoundDataTypeRegistrationMacro.h
--- a/Engine/Plugins/Runtime/Metasound/Source/MetasoundFrontend/Public/MetasoundDataTypeRegistrationMacro.h
+++ b/Engine/Plugins/Runtime/Metasound/Source/MetasoundFrontend/Public/MetasoundDataTypeRegistrationMacro.h
@@ -495,19 +495,6 @@ namespace Metasound::Frontend
 					PRAGMA_ENABLE_EXPERIMENTAL_WARNINGS
 				}
 
-				UE_DEPRECATED(5.6, "Create input nodes using FNodeData")
-				virtual TUniquePtr<INode> CreateInputNode(FInputNodeConstructorParams&& InParams) const override
-				{
-					if constexpr (bIsParsable)
-					{
-						return MakeUnique<TInputNode<TDataType, EVertexAccessType::Reference>>(MoveTemp(InParams));
-					}
-					else
-					{
-						return TUniquePtr<INode>(nullptr);
-					}
-				}
-
 				virtual TUniquePtr<INode> CreateInputNode(FNodeData InNodeData) const override
 				{
 					if constexpr (bIsParsable)
@@ -525,19 +512,6 @@ namespace Metasound::Frontend
 					return TUniquePtr<INode>(nullptr);
 				}
 
-				UE_DEPRECATED(5.6, "Create constructor input nodes using FNodeData")
-				virtual TUniquePtr<INode> CreateConstructorInputNode(FInputNodeConstructorParams&& InParams) const override
-				{
-					if constexpr (bIsParsable && bIsConstructorType)
-					{
-						return MakeUnique<TInputNode<TDataType, EVertexAccessType::Value>>(MoveTemp(InParams));
-					}
-					else
-					{
-						return TUniquePtr<INode>(nullptr);
-					}
-				}
-
 				virtual TUniquePtr<INode> CreateConstructorInputNode(FNodeData InNodeData) const override
 				{
 					if constexpr (bIsParsable && bIsConstructorType)
@@ -555,37 +529,6 @@ namespace Metasound::Frontend
 					return TUniquePtr<INode>(nullptr);
 				}
 
-				using FDataTypeRegistryEntryBase::CreateOutputNode; // Unhide override
-
-				UE_DEPRECATED(5.6, "Create output nodes using FNodeData")
-				virtual TUniquePtr<INode> CreateOutputNode(FOutputNodeConstructorParams&& InParams) const override
-				{
-					if constexpr (bIsParsable)
-					{
-						return MakeUnique<TOutputNode<TDataType, EVertexAccessType::Reference>>(InParams.NodeName, InParams.InstanceID, InParams.VertexName);
-					}
-					else
-					{
-						return TUniquePtr<INode>(nullptr);
-					}
-				}
-
-				using FDataTypeRegistryEntryBase::CreateConstructorOutputNode; // Unhide override
-																			   
-				UE_DEPRECATED(5.6, "Create constructor output nodes using FNodeData")
-				virtual TUniquePtr<INode> CreateConstructorOutputNode(FOutputNodeConstructorParams&& InParams) const override
-				{
-					if constexpr (bIsParsable && bIsConstructorType)
-					{
-						return MakeUnique<TOutputNode<TDataType, EVertexAccessType::Value>>(InParams.NodeName, InParams.InstanceID, InParams.VertexName);
-					}
-					else
-					{
-						return TUniquePtr<INode>(nullptr);
-					}
-				}
-
-
 				UE_DEPRECATED(5.6, "Direct creation of literal nodes will no longer be supported")
 				virtual TUniquePtr<INode> CreateLiteralNode(FLiteralNodeConstructorParams&& InParams) const override
 				{
@@ -604,20 +547,9 @@ namespace Metasound::Frontend
 				{
 					if constexpr (bIsParsable)
 					{
+PRAGMA_DISABLE_DEPRECATION_WARNINGS
 						return MakeUnique<TReceiveNode<TDataType>>(InParams);
-					}
-					else
-					{
-						return TUniquePtr<INode>(nullptr);
-					}
-				}
-
-				UE_DEPRECATED(5.6, "Create variable nodes using FNodeData")
-				virtual TUniquePtr<INode> CreateVariableNode(FVariableNodeConstructorParams&& InParams) const override
-				{
-					if constexpr (bIsParsable)
-					{
-						return MakeUnique<TVariableNode<TDataType>>(InParams.NodeName, InParams.InstanceID, MoveTemp(InParams.Literal));
+PRAGMA_ENABLE_DEPRECATION_WARNINGS
 					}
 					else
 					{
@@ -637,19 +569,6 @@ namespace Metasound::Frontend
 					return TUniquePtr<INode>(nullptr);
 				}
  
-				UE_DEPRECATED(5.6, "Create variable mutator nodes using FNodeData")
-				virtual TUniquePtr<INode> CreateVariableMutatorNode(const FNodeInitData& InParams) const override
-				{
-					if constexpr (bIsParsable)
-					{
-						return MakeUnique<TVariableMutatorNode<TDataType>>(InParams);
-					}
-					else
-					{
-						return TUniquePtr<INode>(nullptr);
-					}
-				}
-
 				virtual TUniquePtr<INode> CreateVariableMutatorNode(FNodeData InNodeData) const override
 				{
 					if constexpr (bIsParsable)
@@ -662,19 +581,6 @@ namespace Metasound::Frontend
 					return TUniquePtr<INode>(nullptr);
 				}
 
-				UE_DEPRECATED(5.6, "Create variable accessor nodes using FNodeData")
-				virtual TUniquePtr<INode> CreateVariableAccessorNode(const FNodeInitData& InParams) const override
-				{
-					if constexpr (bIsParsable)
-					{
-						return MakeUnique<TVariableAccessorNode<TDataType>>(InParams);
-					}
-					else
-					{
-						return TUniquePtr<INode>(nullptr);
-					}
-				}
-
 				virtual TUniquePtr<INode> CreateVariableAccessorNode(FNodeData InNodeData) const override
 				{
 					if constexpr (bIsParsable)
@@ -687,26 +593,13 @@ namespace Metasound::Frontend
 					return TUniquePtr<INode>(nullptr);
 				}
 
-				UE_DEPRECATED(5.6, "Create variable deferred accessor nodes using FNodeData")
-				virtual TUniquePtr<INode> CreateVariableDeferredAccessorNode(const FNodeInitData& InParams) const override
-				{
-					if constexpr (bIsParsable)
-					{
-						return MakeUnique<TVariableDeferredAccessorNode<TDataType>>(InParams);
-					}
-					else
-					{
-						return TUniquePtr<INode>(nullptr);
-					}
-				}
-
 				virtual TUniquePtr<INode> CreateVariableDeferredAccessorNode(FNodeData InNodeData) const override
 				{
 					if constexpr (bIsParsable)
 					{
-						if (VariableAccessorClassMetadata)
+						if (VariableDeferredAccessorClassMetadata)
 						{
-							return MakeUnique<TVariableDeferredAccessorNode<TDataType>>(MoveTemp(InNodeData), VariableAccessorClassMetadata.ToSharedRef());
+							return MakeUnique<TVariableDeferredAccessorNode<TDataType>>(MoveTemp(InNodeData), VariableDeferredAccessorClassMetadata.ToSharedRef());
 						}
 					}
 					return TUniquePtr<INode>(nullptr);
diff --git a/Engine/Plugins/Runtime/Metasound/Source/MetasoundFrontend/Public/MetasoundFrontendDataTypeRegistry.h b/Engine/Plugins/Runtime/Metasound/Source/MetasoundFrontend/Public/MetasoundFrontendDataTypeRegistry.h
--- a/Engine/Plugins/Runtime/Metasound/Source/MetasoundFrontend/Public/MetasoundFrontendDataTypeRegistry.h
+++ b/Engine/Plugins/Runtime/Metasound/Source/MetasoundFrontend/Public/MetasoundFrontendDataTypeRegistry.h
@@ -211,35 +211,19 @@ namespace Metasound
 			/** Return an FNodeClassMetadata representing an get delayed variable node of the data type. */
 			virtual TSharedPtr<const FNodeClassMetadata> GetVariableDeferredAccessorClassMetadata() const = 0;
 
-			/** Create an input node */
-			UE_DEPRECATED(5.6, "Create input nodes with FNodeData")
-			virtual TUniquePtr<INode> CreateInputNode(FInputNodeConstructorParams&&) const = 0;
-
 			/** Create an input node */
 			virtual TUniquePtr<INode> CreateInputNode(FNodeData InNodeData) const = 0;
 
-			/** Create an constructor input node */
-			UE_DEPRECATED(5.6, "Create constructor input nodes with FNodeData")
-			virtual TUniquePtr<INode> CreateConstructorInputNode(FInputNodeConstructorParams&&) const = 0;
-
 			/** Create a constructor input node */
 			virtual TUniquePtr<INode> CreateConstructorInputNode(FNodeData InNodeData) const = 0;
 
-			/** Create an output node */
-			UE_DEPRECATED(5.6, "Create output nodes with FNodeData")
-			virtual TUniquePtr<INode> CreateOutputNode(FOutputNodeConstructorParams&&) const = 0;
-
 			/** Create an output node */
 			virtual TUniquePtr<INode> CreateOutputNode(FNodeData InNodeData) const = 0;
 
-			/** Create an output node */
-			UE_DEPRECATED(5.6, "Create constructor output nodes with FNodeData")
-			virtual TUniquePtr<INode> CreateConstructorOutputNode(FOutputNodeConstructorParams&&) const = 0;
-
 			/** Create a constructor output node */
 			virtual TUniquePtr<INode> CreateConstructorOutputNode(FNodeData InNodeData) const = 0;
 
-			/** Create a variable node */
+			/** Create a literal node (deprecated — literal values should be set on the connected node directly). */
 			UE_DEPRECATED(5.6, "Direct creation of literal nodes will no longer be supported")
 			virtual TUniquePtr<INode> CreateLiteralNode(FLiteralNodeConstructorParams&&) const = 0;
 
@@ -247,43 +231,24 @@ namespace Metasound
 			UE_DEPRECATED(5.6, "Direct creation of receive nodes will no longer be supported")
 			virtual TUniquePtr<INode> CreateReceiveNode(const FNodeInitData&) const = 0;
 
-			/* Get function that should be used when a parameter pack contains a setting this data type. */ 
+			/* Get function that should be used when a parameter pack contains a setting this data type. */
 			virtual const IParameterAssignmentFunction& GetRawAssignmentFunction() const = 0;
 
 			/* Get a function that can be used to assign a FLiteral to a FAnyDataReference. */
 			virtual FLiteralAssignmentFunction GetLiteralAssignmentFunction() const = 0;
 
-			/** Create a init variable node for this data type. 
-			 *
-			 *  @param InInitParams - Contains a literal used to create the variable.
-			 */
-			UE_DEPRECATED(5.6, "Create variable nodes with FNodeData")
-			virtual TUniquePtr<INode> CreateVariableNode(FVariableNodeConstructorParams&& InInitParams) const = 0;
-
-			/** Create a init variable node for this data type. 
+			/** Create a init variable node for this data type.
 			 *
 			 *  @param InInitValue - A literal used to create the variable.
 			 */
 			virtual TUniquePtr<INode> CreateVariableNode(FLiteral InInitValue, FNodeData) const = 0;
 
-			/** Create a set variable node for this data type. */
-			UE_DEPRECATED(5.6, "Create variable mutator nodes with FNodeData")
-			virtual TUniquePtr<INode> CreateVariableMutatorNode(const FNodeInitData&) const = 0;
-
 			/** Create a set variable node for this data type. */
 			virtual TUniquePtr<INode> CreateVariableMutatorNode(FNodeData) const = 0;
 
-			/** Create a get variable node for this data type. */
-			UE_DEPRECATED(5.6, "Create variable accessor nodes with FNodeData")
-			virtual TUniquePtr<INode> CreateVariableAccessorNode(const FNodeInitData&) const = 0;
-
 			/** Create a get variable node for this data type. */
 			virtual TUniquePtr<INode> CreateVariableAccessorNode(FNodeData) const = 0;
 
-			/** Create a get delayed variable node for this data type. */
-			UE_DEPRECATED(5.6, "Create variable deferred accessor nodes with FNodeData")
-			virtual TUniquePtr<INode> CreateVariableDeferredAccessorNode(const FNodeInitData&) const = 0;
-
 			/** Create a get delayed variable node for this data type. */
 			virtual TUniquePtr<INode> CreateVariableDeferredAccessorNode(FNodeData) const = 0;
 			
@@ -416,42 +381,20 @@ namespace Metasound
 			 * such as Inputs, Outputs, and Variables. 
 			 */
 
-			UE_DEPRECATED(5.6, "Create input nodes with FNodeData")
-			virtual TUniquePtr<INode> CreateInputNode(const FName& InInputType, FInputNodeConstructorParams&& InParams) const = 0;
 			virtual TUniquePtr<INode> CreateInputNode(const FName& InInputType, FNodeData InNodeData) const = 0;
-
-			UE_DEPRECATED(5.6, "Create constructor input nodes with FNodeData")
-			virtual TUniquePtr<INode> CreateConstructorInputNode(const FName& InInputType, FInputNodeConstructorParams&& InParams) const = 0;
 			virtual TUniquePtr<INode> CreateConstructorInputNode(const FName& InInputType, FNodeData InNodeData) const = 0;
 
 			UE_DEPRECATED(5.6, "Direct creation of literal nodes will no longer be supported")
 			virtual TUniquePtr<INode> CreateLiteralNode(const FName& InLiteralType, FLiteralNodeConstructorParams&& InParams) const = 0;
 
-			UE_DEPRECATED(5.6, "Create output nodes with FNodeData")
-			virtual TUniquePtr<INode> CreateOutputNode(const FName& InDataTypeName, FOutputNodeConstructorParams&& InParams) const = 0;
-			virtual TUniquePtr<INode> CreateOutputNode(const FName& InDataTypeName, FNodeData InNodeData) const = 0;
-
-			UE_DEPRECATED(5.6, "Create output nodes with FNodeData")
-			virtual TUniquePtr<INode> CreateConstructorOutputNode(const FName& InDataTypeName, FOutputNodeConstructorParams&& InParams) const = 0;
-			virtual TUniquePtr<INode> CreateConstructorOutputNode(const FName& InDataTypeName, FNodeData InNodeData) const = 0;
-
 			UE_DEPRECATED(5.6, "Direct creation of receive nodes will no longer be supported")
 			virtual TUniquePtr<INode> CreateReceiveNode(const FName& InDataType, const FNodeInitData&) const = 0;
 
-			UE_DEPRECATED(5.6, "Create Variable Node with FNodeData")
-			virtual TUniquePtr<INode> CreateVariableNode(const FName& InDataType, FVariableNodeConstructorParams&& InParams) const = 0;
+			virtual TUniquePtr<INode> CreateOutputNode(const FName& InDataTypeName, FNodeData InNodeData) const = 0;
+			virtual TUniquePtr<INode> CreateConstructorOutputNode(const FName& InDataTypeName, FNodeData InNodeData) const = 0;
 			virtual TUniquePtr<INode> CreateVariableNode(const FName& InDataType, FLiteral InLiteral, FNodeData InNodeData) const = 0;
-
-			UE_DEPRECATED(5.6, "Create variable mutator nodes with FNodeData")
-			virtual TUniquePtr<INode> CreateVariableMutatorNode(const FName& InDataType, const FNodeInitData& InParams) const = 0;
 			virtual TUniquePtr<INode> CreateVariableMutatorNode(const FName& InDataType, FNodeData InNodeData) const = 0;
-
-			UE_DEPRECATED(5.6, "Create variable accessor nodes with FNodeData")
-			virtual TUniquePtr<INode> CreateVariableAccessorNode(const FName& InDataType, const FNodeInitData& InParams) const = 0;
 			virtual TUniquePtr<INode> CreateVariableAccessorNode(const FName& InDataType, FNodeData InNodeData) const = 0;
-
-			UE_DEPRECATED(5.6, "Create variable deferred accessor nodes with FNodeData")
-			virtual TUniquePtr<INode> CreateVariableDeferredAccessorNode(const FName& InDataType, const FNodeInitData& InParams) const = 0;
 			virtual TUniquePtr<INode> CreateVariableDeferredAccessorNode(const FName& InDataType, FNodeData InNodeData) const = 0;
 		};
 
diff --git a/Engine/Plugins/Runtime/Metasound/Source/MetasoundFrontend/Public/MetasoundFrontendNodeClassRegistry.h b/Engine/Plugins/Runtime/Metasound/Source/MetasoundFrontend/Public/MetasoundFrontendNodeClassRegistry.h
--- a/Engine/Plugins/Runtime/Metasound/Source/MetasoundFrontend/Public/MetasoundFrontendNodeClassRegistry.h
+++ b/Engine/Plugins/Runtime/Metasound/Source/MetasoundFrontend/Public/MetasoundFrontendNodeClassRegistry.h
@@ -45,44 +45,8 @@ namespace Metasound::Frontend
 		UE_DEPRECATED(5.8, "NodeClassInfo is actively being deprecated. Use FNodeClassRegistryKey instead (which can be generated from the frontend class)")
 		virtual const FNodeClassInfo& GetClassInfo() const { return FNodeClassInfo::GetInvalid(); }
 
-		/** Create a node given FDefaultNamedVertexNodeConstructorParams.
-			*
-			* If a node can be created with FDefaultNamedVertexNodeConstructorParams, this function
-			* should return a valid node pointer.
-			*/
-		UE_DEPRECATED(5.6, "Node classes should be constructed with FNodeData")
-		virtual TUniquePtr<INode> CreateNode(FDefaultNamedVertexNodeConstructorParams&&) const { return nullptr; }
-
-		/** Create a node given FDefaultNamedVertexWithLiteralNodeConstructorParams.
-			*
-			* If a node can be created with FDefaultNamedVertexWithLiteralNodeConstructorParams, this function
-			* should return a valid node pointer.
-			*/
-		UE_DEPRECATED(5.6, "Node classes should be constructed with FNodeData")
-		virtual TUniquePtr<INode> CreateNode(FDefaultNamedVertexWithLiteralNodeConstructorParams&&) const { return nullptr; }
-
-		/** Create a node given FDefaultLiteralNodeConstructorParams.
-			*
-			* If a node can be created with FDefaultLiteralNodeConstructorParams, this function
-			* should return a valid node pointer.
-			*/
-		UE_DEPRECATED(5.6, "Node classes should be constructed with FNodeData")
-		virtual TUniquePtr<INode> CreateNode(FDefaultLiteralNodeConstructorParams&&) const { return nullptr; }
-
-		/** Create a node given FNodeInitData.
-			*
-			* If a node can be created with FNodeInitData, this function
-			* should return a valid node pointer.
-			*/
-		UE_DEPRECATED(5.6, "Node classes should be constructed with FNodeData")
-		virtual TUniquePtr<INode> CreateNode(const FNodeInitData&) const = 0;
-
-		/** Create a node given FNodeData.
-			*
-			* If a node can be created with FNodeData, this function
-			* should return a valid node pointer.
-			*/
-		UE_API virtual TUniquePtr<INode> CreateNode(FNodeData) const;
+		/** Create a node given FNodeData. */
+		virtual TUniquePtr<INode> CreateNode(FNodeData) const = 0;
 
 		/** Return a FMetasoundFrontendClass which describes the node. */
 		virtual const FMetasoundFrontendClass& GetFrontendClass() const = 0;
@@ -99,6 +63,10 @@ namespace Metasound::Frontend
 
 		virtual FVertexInterface GetDefaultVertexInterface() const = 0;
 
+		/** Return the FClassInterface for this node class, which includes
+		 * sub-interface and variant schema information. */
+		virtual const FClassInterface& GetClassInterface() const = 0;
+
 		UE_DEPRECATED(5.6, "Node class registry no longer tracks nature of implementation (i.e. if was registered "
 			"from an asset or c++. Use the MetaSoundAssetManager to determine if the class has been defined within an asset(s).")
 		virtual bool IsNative() const { return false; }
@@ -371,6 +339,10 @@ namespace Metasound::Frontend
 		// returning true on success, false otherwise 
 		virtual bool FindDefaultVertexInterface(const FNodeClassRegistryKey& InKey, FVertexInterface& OutVertexInterface) const = 0;
 
+		/** Get the FClassInterface (with sub-interface/variant schema) for the given key.
+		 * Returns true on success, false if the key is not found. */
+		virtual bool FindClassInterface(const FNodeClassRegistryKey& InKey, FClassInterface& OutClassInterface) const = 0;
+
 		virtual bool FindFrontendClassFromRegistered(const FNodeClassRegistryKey& InKey, FMetasoundFrontendClass& OutClass) const = 0;
 
 		/** Returns true if the provided node configuration is compatible (same derived class) with the registered configuration for the node represented by the registry key. */
@@ -396,18 +368,6 @@ namespace Metasound::Frontend
 		 * 
 		 * Returns an invalid TUniquePtr<> if the key is not in the registry. 
 		 */
-		UE_DEPRECATED(5.6, "Node classes should be constructed with FNodeData")
-		virtual TUniquePtr<Metasound::INode> CreateNode(const FNodeClassRegistryKey& InKey, const Metasound::FNodeInitData&) const = 0;
-
-		UE_DEPRECATED(5.6, "Node classes should be constructed with FNodeData")
-		virtual TUniquePtr<Metasound::INode> CreateNode(const FNodeClassRegistryKey& InKey, Metasound::FDefaultLiteralNodeConstructorParams&&) const { return nullptr; }
-
-		UE_DEPRECATED(5.6, "Node classes should be constructed with FNodeData")
-		virtual TUniquePtr<Metasound::INode> CreateNode(const FNodeClassRegistryKey& InKey, Metasound::FDefaultNamedVertexNodeConstructorParams&&) const { return nullptr; }
-
-		UE_DEPRECATED(5.6, "Node classes should be constructed with FNodeData")
-		virtual TUniquePtr<Metasound::INode> CreateNode(const FNodeClassRegistryKey& InKey, Metasound::FDefaultNamedVertexWithLiteralNodeConstructorParams&&) const { return nullptr; }
-
 		/** Registers a node as a converter node */
 		virtual bool RegisterConversionNode(const FConverterNodeClassRegistryKey& InConversionKey, const FConverterNodeClassInfo& InNodeInfo) = 0;
 		
diff --git a/Engine/Plugins/Runtime/Metasound/Source/MetasoundFrontend/Public/MetasoundLiteralNode.h b/Engine/Plugins/Runtime/Metasound/Source/MetasoundFrontend/Public/MetasoundLiteralNode.h
--- a/Engine/Plugins/Runtime/Metasound/Source/MetasoundFrontend/Public/MetasoundLiteralNode.h
+++ b/Engine/Plugins/Runtime/Metasound/Source/MetasoundFrontend/Public/MetasoundLiteralNode.h
@@ -160,10 +160,6 @@ namespace Metasound
 
 		UE_API virtual const FVertexInterface& GetVertexInterface() const override;
 
-		UE_API virtual bool SetVertexInterface(const FVertexInterface& InInterface) override;
-
-		UE_API virtual bool IsVertexInterfaceSupported(const FVertexInterface& InInterface) const override;
-
 		UE_API virtual TSharedRef<IOperatorFactory, ESPMode::ThreadSafe> GetDefaultOperatorFactory() const override;
 
 	private:
diff --git a/Engine/Plugins/Runtime/Metasound/Source/MetasoundFrontend/Public/MetasoundNodeRegistrationMacro.h b/Engine/Plugins/Runtime/Metasound/Source/MetasoundFrontend/Public/MetasoundNodeRegistrationMacro.h
--- a/Engine/Plugins/Runtime/Metasound/Source/MetasoundFrontend/Public/MetasoundNodeRegistrationMacro.h
+++ b/Engine/Plugins/Runtime/Metasound/Source/MetasoundFrontend/Public/MetasoundNodeRegistrationMacro.h
@@ -10,6 +10,7 @@
 #include "MetasoundLog.h"
 #include "MetasoundFrontendModuleRegistrationMacros.h"
 #include "MetasoundNodeInterface.h"
+#include "MetasoundVertex.h"
 
 #include "Traits/MetasoundNodeConstructorTraits.h"
 #include "Traits/MetasoundNodeStaticMemberTraits.h"
@@ -38,16 +39,12 @@ namespace Metasound::Frontend
 
 			UE_API virtual const FMetasoundFrontendClass& GetFrontendClass() const override;
 
-			// Unhide CreateNode overloads which exist on INodeClassRegistryEntry
-			using INodeClassRegistryEntry::CreateNode;
-
-			UE_DEPRECATED(5.6, "Node classes should be constructed with FNodeData")
-			UE_API virtual TUniquePtr<INode> CreateNode(const FNodeInitData& InInitData) const override;
-
 			UE_API virtual const TSet<FMetasoundFrontendVersion>* GetImplementedInterfaces() const override;
 
 			UE_API virtual FVertexInterface GetDefaultVertexInterface() const override;
 
+			UE_API virtual const FClassInterface& GetClassInterface() const override;
+
 			UE_API virtual TInstancedStruct<FMetaSoundFrontendNodeConfiguration> CreateFrontendNodeConfiguration() const override;
 			UE_API virtual bool IsCompatibleNodeConfiguration(TConstStructView<FMetaSoundFrontendNodeConfiguration> InNodeConfiguration) const override;
 
@@ -78,9 +75,6 @@ namespace Metasound::Frontend
 			// Expose FNodeRegistryEntryBase constructors. 
 			using FNodeRegistryEntryBase::FNodeRegistryEntryBase;
 
-			// Unhide CreateNode overloads which exist on INodeClassRegistryEntry
-			using FNodeRegistryEntryBase::CreateNode;
-
 			virtual TUniquePtr<INode> CreateNode(FNodeData InNodeData) const override
 			{
 				if constexpr(std::is_constructible_v<TNodeType, ::Metasound::FNodeData, TSharedRef<const FNodeClassMetadata>>)
@@ -102,8 +96,8 @@ namespace Metasound::Frontend
 			}
 		};
 
-		// A node registry entry which also provides a node extension. 
-		template<typename NodeType, typename ConfigurationType> 
+		// A node registry entry which also provides a node extension.
+		template<typename NodeType, typename ConfigurationType>
 		class TNodeRegistryEntry : public TNodeRegistryEntryBase<NodeType>
 		{
 			static_assert(std::is_base_of_v<FMetaSoundFrontendNodeConfiguration, ConfigurationType>, "Configurations must inherit from FMetaSoundFrontendNodeConfiguration");
@@ -122,71 +116,14 @@ namespace Metasound::Frontend
 			}
 		};
 
-		// A partial template specialization for scenario where no node extension is provided
+		// A partial template specialization for scenario where no node extension is provided.
 		template<typename NodeType>
 		class TNodeRegistryEntry<NodeType, void> : public TNodeRegistryEntryBase<NodeType>
 		{
 		public:
-			// Inherit constructor
 			using TNodeRegistryEntryBase<NodeType>::TNodeRegistryEntryBase;
 		};
 
-#if UE_METASOUND_DISABLE_5_6_NODE_REGISTRATION_DEPRECATION_WARNINGS
-
-		template<typename T>
-		void TriggerDeprecatedNodeConstructorWarning() {}
-
-		template<typename T>
-		void TriggerMissingCreateNodeClassMetadataWarning() {}
-#else		
-		// These only trigger deprecation if the macro UE_METASOUND_DISABLE_5_6_NODE_REGISTRATION_DEPRECATION_WARNINGS is false
-		template<typename T>
-		UE_DEPRECATED(5.6, "Update the node's constructor to be of the form FMyNode(FNodeData InNodeData, TSharedRef<const FNodeClassMetadata> InClassMetadata)")
-		void TriggerDeprecatedNodeConstructorWarning() {}
-
-		template<typename T>
-		UE_DEPRECATED(5.6, "Update the node class to include 'static FNodeClassMetadata CreateNodeClassMetadata()`)")
-		void TriggerMissingCreateNodeClassMetadataWarning() {}
-#endif
-
-		// Forward declare
-		class FDeprecatedNodeAPIFactory;
-
-		// FDeprecatedNodeAPIAdapterBase is used as a back compatible shim
-		// for nodes which have not been updated to support the new APIs for node
-		// registration.
-		class FDeprecatedNodeAPIAdapterBase : public FBasicNode
-		{
-		protected:
-			UE_API FDeprecatedNodeAPIAdapterBase(FNodeData InNodeData, TSharedRef<const FNodeClassMetadata> InClassMetadata, TUniquePtr<INode> InNode);
-
-		public:
-
-			virtual ~FDeprecatedNodeAPIAdapterBase() = default;
-			UE_API virtual FOperatorFactorySharedRef GetDefaultOperatorFactory() const override;
-
-		private:
-
-			TSharedRef<FDeprecatedNodeAPIFactory> Factory;
-		};
-
-		// TDeprecatedNodeAPIAdapter is used as a back compatible shim
-		// for nodes which have not been updated to support the new APIs for node
-		// registration.
-		template<typename TNodeType>
-		class TDeprecatedNodeAPIAdapter : public FDeprecatedNodeAPIAdapterBase
-		{
-		public:
-			TDeprecatedNodeAPIAdapter(FNodeData InNodeData, TSharedRef<const FNodeClassMetadata> InClassMetadata)
-			: FDeprecatedNodeAPIAdapterBase(InNodeData, MoveTemp(InClassMetadata), MakeUnique<TNodeType>(FNodeInitData{InNodeData.Name, InNodeData.ID}))
-			{
-				if constexpr(TIsOnlyDeprecatedNodeConstructorProvided<TNodeType>::Value)
-				{
-					TriggerDeprecatedNodeConstructorWarning<TNodeType>();
-				}
-			}
-		};
-
 		UE_API bool RegisterNodeInternal(TUniquePtr<INodeClassRegistryEntry>&& InEntry);
 	} // namespace NodeRegistrationPrivate
 
@@ -194,14 +131,8 @@ namespace Metasound::Frontend
 	bool RegisterNode(const FNodeClassMetadata& InMetadata, const FModuleInfo& InOwningModuleInfo)
 	{
 		using namespace NodeRegistrationPrivate;
-		if constexpr (TIsOnlyDeprecatedNodeConstructorProvided<TNodeType>::Value)
-		{
-			return RegisterNode<TDeprecatedNodeAPIAdapter<TNodeType>, ConfigurationType>(InMetadata, InOwningModuleInfo);
-		}
-		else
-		{
-			return NodeRegistrationPrivate::RegisterNodeInternal(MakeUnique<TNodeRegistryEntry<TNodeType, ConfigurationType>>(InMetadata, InOwningModuleInfo));
-		}
+		
+		return NodeRegistrationPrivate::RegisterNodeInternal(MakeUnique<TNodeRegistryEntry<TNodeType, ConfigurationType>>(InMetadata, InOwningModuleInfo));
 
 	}
  
@@ -216,20 +147,17 @@ namespace Metasound::Frontend
 	template <typename TNodeType, typename ConfigurationType=void>
 	bool RegisterNode(const FModuleInfo& InOwningModuleInfo)
 	{
-		static_assert(::Metasound::TIsNodeConstructorSupported<TNodeType>::Value, "In order to be registered as a MetaSound node, the node needs to implement the following public constructor: Constructor(Metasound::FNodeData InNodeData) or Construct(Metasound::FNodeData InNodeData, TSharedRef<const FNodeClassMetadata> InMetadata)"); 
+		static_assert(::Metasound::TIsNodeConstructorSupported<TNodeType>::Value, "In order to be registered as a MetaSound node, the node needs to implement the following public constructor: Constructor(Metasound::FNodeData InNodeData) or Construct(Metasound::FNodeData InNodeData, TSharedRef<const FNodeClassMetadata> InMetadata)");
 
+		using namespace NodeRegistrationPrivate;
 		if constexpr(TIsCreateNodeClassMetadataDeclared<TNodeType>::Value)
 		{
 			return RegisterNode<TNodeType, ConfigurationType>(TNodeType::CreateNodeClassMetadata(), InOwningModuleInfo);
 		}
 		else
 		{
-			using namespace NodeRegistrationPrivate;
-			TriggerMissingCreateNodeClassMetadataWarning<TNodeType>();
-			// Register a node using a prototype node.
-			FNodeInitData InitData;
-			TUniquePtr<Metasound::INode> Node = MakeUnique<TNodeType>(InitData);
-			return RegisterNode<TNodeType, ConfigurationType>(Node->GetMetadata(), InOwningModuleInfo);
+			static_assert(sizeof(TNodeType) == 0, "Node class must implement 'static FNodeClassMetadata CreateNodeClassMetadata()'");
+			return false;
 		}
 	}
 
@@ -257,12 +185,8 @@ namespace Metasound::Frontend
 		}
 		else
 		{
-			using namespace NodeRegistrationPrivate;
-			TriggerMissingCreateNodeClassMetadataWarning<TNodeType>();
-			// Register a node using a prototype node.
-			FNodeInitData InitData;
-			TUniquePtr<Metasound::INode> Node = MakeUnique<TNodeType>(InitData);
-			return UnregisterNode(Node->GetMetadata(), InOwningModuleInfo);
+			static_assert(sizeof(TNodeType) == 0, "Node class must implement 'static FNodeClassMetadata CreateNodeClassMetadata()'");
+			return false;
 		}
 	}
 
@@ -276,21 +200,6 @@ namespace Metasound::Frontend
 
 namespace Metasound
 {
-	// Utility class to ensure that a node class can use the constructor the frontend uses.
-	template <typename NodeClass>
-	struct UE_DEPRECATED(5.6, "Use std::is_constructible<T, FNodeInitData> instead") ConstructorTakesNodeInitData 
-	{
-
-		// Use SFINAE trick to see if we have a valid constructor:
-		template<typename T>
-		static uint16 TestForConstructor(decltype(T(::Metasound::FNodeInitData()))*);
-
-		template<typename T>
-		static uint8 TestForConstructor(...);
-
-		static const bool Value = sizeof(TestForConstructor<NodeClass>(nullptr)) == sizeof(uint16);
-	};
-
 	template<typename T, typename ConfigurationType=void>
 	UE_DEPRECATED(5.6, "Use Frontend::RegisterNode()")
 	bool RegisterNodeWithFrontend()
diff --git a/Engine/Plugins/Runtime/Metasound/Source/MetasoundFrontend/Public/Traits/MetasoundNodeConstructorTraits.h b/Engine/Plugins/Runtime/Metasound/Source/MetasoundFrontend/Public/Traits/MetasoundNodeConstructorTraits.h
--- a/Engine/Plugins/Runtime/Metasound/Source/MetasoundFrontend/Public/Traits/MetasoundNodeConstructorTraits.h
+++ b/Engine/Plugins/Runtime/Metasound/Source/MetasoundFrontend/Public/Traits/MetasoundNodeConstructorTraits.h
@@ -7,31 +7,17 @@
 
 namespace Metasound
 {
-	/** Determines whether the constructor provided by the node is usable in node 
+	/** Determines whether the constructor provided by the node is usable in node
 	 * registration. */
 	template<typename T>
 	struct TIsNodeConstructorSupported
 	{
 	private:
-		static constexpr bool bIsConstructibleWithFNodeInitData = std::is_constructible_v<T, ::Metasound::FNodeInitData>;
 		static constexpr bool bIsConstructibleWithFNodeData = std::is_constructible_v<T, ::Metasound::FNodeData>;
 		static constexpr bool bIsConstructibleWithFNodeDataAndClassMetadata = std::is_constructible_v<T, ::Metasound::FNodeData, TSharedRef<const FNodeClassMetadata>>;
 
 	public:
-		static constexpr bool Value = bIsConstructibleWithFNodeInitData || bIsConstructibleWithFNodeData || bIsConstructibleWithFNodeDataAndClassMetadata;
-	};
-
-	/** Determines a node only provides a deprecated constructor. */
-	template<typename T> 
-	struct TIsOnlyDeprecatedNodeConstructorProvided
-	{
-	private:
-		static constexpr bool bIsConstructibleWithFNodeInitData = std::is_constructible_v<T, ::Metasound::FNodeInitData>;
-		static constexpr bool bIsConstructibleWithFNodeData = std::is_constructible_v<T, ::Metasound::FNodeData>;
-		static constexpr bool bIsConstructibleWithFNodeDataAndClassMetadata = std::is_constructible_v<T, ::Metasound::FNodeData, TSharedRef<const FNodeClassMetadata>>;
-
-	public:
-		static constexpr bool Value = bIsConstructibleWithFNodeInitData && !(bIsConstructibleWithFNodeData || bIsConstructibleWithFNodeDataAndClassMetadata);
+		static constexpr bool Value = bIsConstructibleWithFNodeData || bIsConstructibleWithFNodeDataAndClassMetadata;
 	};
 }
 
diff --git a/Engine/Plugins/Runtime/Metasound/Source/MetasoundGraphCore/Private/MetasoundGraph.cpp b/Engine/Plugins/Runtime/Metasound/Source/MetasoundGraphCore/Private/MetasoundGraph.cpp
--- a/Engine/Plugins/Runtime/Metasound/Source/MetasoundGraphCore/Private/MetasoundGraph.cpp
+++ b/Engine/Plugins/Runtime/Metasound/Source/MetasoundGraphCore/Private/MetasoundGraph.cpp
@@ -28,6 +28,9 @@ namespace Metasound
 		, Metadata(FNodeClassMetadata::GetEmpty())
 		, OperatorData(MoveTemp(InOperatorData))
 	{
+		// FGraph is a graph boundary, not a configurable node — it must not
+		// have sub-interfaces or variants on its DefaultInterface.
+		ensure(!Metadata.DefaultInterface.ContainsSubInterfaces() && !Metadata.DefaultInterface.ContainsVariants());
 	}
 
 	FGraph::FGraph(const FName& InInstanceName, const FGuid& InInstanceID, TSharedPtr<const IOperatorData> InOperatorData)
@@ -37,6 +40,7 @@ namespace Metasound
 		, Metadata(FNodeClassMetadata::GetEmpty())
 		, OperatorData(MoveTemp(InOperatorData))
 	{
+		ensure(!Metadata.DefaultInterface.ContainsSubInterfaces() && !Metadata.DefaultInterface.ContainsVariants());
 	}
 
 	FGraph::FGraph(const FTopLevelAssetPath& InAssetPath, const FGuid& InInstanceID, TSharedPtr<const IOperatorData> InOperatorData)
@@ -46,6 +50,7 @@ namespace Metasound
 	, Metadata(FNodeClassMetadata::GetEmpty())
 	, OperatorData(MoveTemp(InOperatorData))
 	{
+		ensure(!Metadata.DefaultInterface.ContainsSubInterfaces() && !Metadata.DefaultInterface.ContainsVariants());
 	}
 
 	const FTopLevelAssetPath& FGraph::GetAssetPath() const
@@ -84,13 +89,15 @@ namespace Metasound
 
 	void FGraph::AddInputDataDestination(const FInputDataDestination& InDestination)
 	{
-		Metadata.DefaultInterface.GetInputInterface().Add(InDestination.Vertex);
+		GraphInterface.GetInputInterface().Add(InDestination.Vertex);
 		InputDestinations.Add(MakeDestinationDataVertexKey(InDestination), InDestination);
+		Metadata.DefaultInterface = GraphInterface;
 	}
 
 	bool FGraph::RemoveInputDataDestination(const FVertexName& InVertexName)
 	{
-		bool bSuccessRemoveFromInterface = Metadata.DefaultInterface.GetInputInterface().Remove(InVertexName);
+		bool bSuccessRemoveFromInterface = GraphInterface.GetInputInterface().Remove(InVertexName);
+		Metadata.DefaultInterface = GraphInterface;
 
 		FInputDataDestinationCollection::TIterator Iter = InputDestinations.CreateIterator();
 		bool bSuccessRemoveFromDestinations = false;
@@ -127,13 +134,15 @@ namespace Metasound
 
 	void FGraph::AddOutputDataSource(const FOutputDataSource& InSource)
 	{
-		Metadata.DefaultInterface.GetOutputInterface().Add(InSource.Vertex);
+		GraphInterface.GetOutputInterface().Add(InSource.Vertex);
 		OutputSources.Add(MakeSourceDataVertexKey(InSource), InSource);
+		Metadata.DefaultInterface = GraphInterface;
 	}
 
 	bool FGraph::RemoveOutputDataSource(const FVertexName& InVertexName)
 	{
-		bool bSuccessRemoveFromInterface = Metadata.DefaultInterface.GetOutputInterface().Remove(InVertexName);
+		bool bSuccessRemoveFromInterface = GraphInterface.GetOutputInterface().Remove(InVertexName);
+		Metadata.DefaultInterface = GraphInterface;
 
 		FOutputDataSourceCollection::TIterator Iter = OutputSources.CreateIterator();
 		bool bSuccessRemoveFromSources = false;
@@ -368,16 +377,16 @@ PRAGMA_ENABLE_DEPRECATION_WARNINGS
 		return Edges;
 	}
 
-	const FVertexInterface& FGraph::GetVertexInterface() const 
+	const FVertexInterface& FGraph::GetVertexInterface() const
 	{
-		return Metadata.DefaultInterface;
+		return GraphInterface;
 	}
 
 	void FGraph::SetDefaultInput(const FVertexName& InVertexName, const FLiteral& InLiteral)
 	{
 		// There's some significant redundancy when storing default literals on a
 		// graph's input. We update them all. 
-		if (FInputDataVertex* InputVertex = Metadata.DefaultInterface.GetInputInterface().Find(InVertexName))
+		if (FInputDataVertex* InputVertex = GraphInterface.GetInputInterface().Find(InVertexName))
 		{
 			// Update literal on exposed vertex
 			InputVertex->SetDefaultLiteral(InLiteral);
@@ -399,12 +408,14 @@ PRAGMA_ENABLE_DEPRECATION_WARNINGS
 			{
 				// This check for mutable nodes will not be necessary once we remove
 				// storage for TSharedPtr<const INodes>. At that point we can assume
-				// that all internal nodes can be mutated safely. 
+				// that all internal nodes can be mutated safely.
 				if (INode* MutableInputNode = FindMutableNode(InputNode->GetInstanceID()))
 				{
 					MutableInputNode->SetDefaultInput(InVertexName, InLiteral);
 				}
 			}
+
+			Metadata.DefaultInterface = GraphInterface;
 		}
 	}
 
@@ -413,16 +424,6 @@ PRAGMA_ENABLE_DEPRECATION_WARNINGS
 		return OperatorData;
 	}
 
-	bool FGraph::SetVertexInterface(const FVertexInterface& InInterface)
-	{
-		return InInterface == Metadata.DefaultInterface;
-	}
-
-	bool FGraph::IsVertexInterfaceSupported(const FVertexInterface& InInterface) const
-	{
-		return InInterface == Metadata.DefaultInterface;
-	}
-
 	FOperatorFactorySharedRef FGraph::GetDefaultOperatorFactory() const 
 	{
 		return MakeShared<FGraph::FFactory>();
diff --git a/Engine/Plugins/Runtime/Metasound/Source/MetasoundGraphCore/Private/MetasoundNode.cpp b/Engine/Plugins/Runtime/Metasound/Source/MetasoundGraphCore/Private/MetasoundNode.cpp
--- a/Engine/Plugins/Runtime/Metasound/Source/MetasoundGraphCore/Private/MetasoundNode.cpp
+++ b/Engine/Plugins/Runtime/Metasound/Source/MetasoundGraphCore/Private/MetasoundNode.cpp
@@ -10,6 +10,7 @@ namespace Metasound
 	: InstanceName(InInstanceName)
 	, InstanceID(InInstanceID)
 	, Info(InInfo)
+	, VertexInterface(InInfo.DefaultInterface)
 	, OperatorData(MoveTemp(InOperatorData))
 	{
 	}
@@ -34,12 +35,12 @@ namespace Metasound
 
 	const FVertexInterface& FNode::GetVertexInterface() const
 	{
-		return Info.DefaultInterface;
+		return VertexInterface;
 	}
 
 	void FNode::SetDefaultInput(const FVertexName& InVertexName, const FLiteral& InLiteral)
 	{
-		if (FInputDataVertex* Vertex = Info.DefaultInterface.GetInputInterface().Find(InVertexName))
+		if (FInputDataVertex* Vertex = VertexInterface.GetInputInterface().Find(InVertexName))
 		{
 			Vertex->SetDefaultLiteral(InLiteral);
 		}
diff --git a/Engine/Plugins/Runtime/Metasound/Source/MetasoundGraphCore/Private/MetasoundNodeInterface.cpp b/Engine/Plugins/Runtime/Metasound/Source/MetasoundGraphCore/Private/MetasoundNodeInterface.cpp
--- a/Engine/Plugins/Runtime/Metasound/Source/MetasoundGraphCore/Private/MetasoundNodeInterface.cpp
+++ b/Engine/Plugins/Runtime/Metasound/Source/MetasoundGraphCore/Private/MetasoundNodeInterface.cpp
@@ -161,7 +161,7 @@ namespace Metasound
 		FText InDescription,
 		FString InAuthor,
 		FText InPromptIfMissing,
-		FVertexInterface InDefaultInterface,
+		FClassInterface InDefaultInterface,
 		TArray<FText> InCategoryHierarchy,
 		TArray<FText> InKeywords,
 		FNodeDisplayStyle InDisplayStyle)
diff --git a/Engine/Plugins/Runtime/Metasound/Source/MetasoundGraphCore/Private/MetasoundVertex.cpp b/Engine/Plugins/Runtime/Metasound/Source/MetasoundGraphCore/Private/MetasoundVertex.cpp
--- a/Engine/Plugins/Runtime/Metasound/Source/MetasoundGraphCore/Private/MetasoundVertex.cpp
+++ b/Engine/Plugins/Runtime/Metasound/Source/MetasoundGraphCore/Private/MetasoundVertex.cpp
@@ -848,7 +848,49 @@ namespace Metasound
 
 
 
-PRAGMA_DISABLE_EXPERIMENTAL_WARNINGS
+	FClassInterface::FClassInterface(const FVertexInterface& InVertexInterface)
+	{
+		for (const FInputDataVertex& Vertex : InVertexInterface.GetInputInterface())
+		{
+			Inputs.Add(Vertex);
+		}
+		for (const FOutputDataVertex& Vertex : InVertexInterface.GetOutputInterface())
+		{
+			Outputs.Add(Vertex);
+		}
+		for (const FEnvironmentVertex& Vertex : InVertexInterface.GetEnvironmentInterface())
+		{
+			EnvironmentVariables.Add(Vertex);
+		}
+	}
+
+	FClassInterface::FClassInterface(const FClassInterface& Other)
+	: VariantSchema(Other.VariantSchema ? MakeUnique<VertexPrivate::FVariantSchema>(*Other.VariantSchema) : nullptr)
+	, SubInterfaceSchema(Other.SubInterfaceSchema ? MakeUnique<VertexPrivate::FSubInterfaceSchema>(*Other.SubInterfaceSchema) : nullptr)
+	, Inputs(Other.Inputs)
+	, Outputs(Other.Outputs)
+	, EnvironmentVariables(Other.EnvironmentVariables)
+	{
+	}
+
+	FClassInterface& FClassInterface::operator=(const FClassInterface& Other)
+	{
+		if (this != &Other)
+		{
+			VariantSchema = Other.VariantSchema ? MakeUnique<VertexPrivate::FVariantSchema>(*Other.VariantSchema) : nullptr;
+			SubInterfaceSchema = Other.SubInterfaceSchema ? MakeUnique<VertexPrivate::FSubInterfaceSchema>(*Other.SubInterfaceSchema) : nullptr;
+			Inputs = Other.Inputs;
+			Outputs = Other.Outputs;
+			EnvironmentVariables = Other.EnvironmentVariables;
+		}
+		return *this;
+	}
+
+	FClassInterface::operator FVertexInterface() const
+	{
+		return CreateVertexInterface({}, {});
+	}
+
 	FClassInterface::FClassInterface(const VertexPrivate::FClassInterfaceData& InClassInterfaceData)
 	: Inputs(InClassInterfaceData.Inputs)
 	, Outputs(InClassInterfaceData.Outputs)
@@ -858,7 +900,7 @@ PRAGMA_DISABLE_EXPERIMENTAL_WARNINGS
 
 		AssertSubInterfacesAreWellFormed(InClassInterfaceData);
 		AssertVariantsAreWellFormed(InClassInterfaceData);
-		
+
 		if (InClassInterfaceData.SubInterfaces.Num())
 		{
 			SubInterfaceSchema = MakeUnique<FSubInterfaceSchema>(FSubInterfaceSchema::FromClassInterfaceData(InClassInterfaceData));
@@ -956,8 +998,6 @@ PRAGMA_DISABLE_EXPERIMENTAL_WARNINGS
 		return nullptr;
 	}
 
-PRAGMA_ENABLE_EXPERIMENTAL_WARNINGS
-
 	namespace VertexPrivate
 	{
 		template<typename ElementType, typename ArrayType>
diff --git a/Engine/Plugins/Runtime/Metasound/Source/MetasoundGraphCore/Public/MetasoundFacade.h b/Engine/Plugins/Runtime/Metasound/Source/MetasoundGraphCore/Public/MetasoundFacade.h
--- a/Engine/Plugins/Runtime/Metasound/Source/MetasoundGraphCore/Public/MetasoundFacade.h
+++ b/Engine/Plugins/Runtime/Metasound/Source/MetasoundGraphCore/Public/MetasoundFacade.h
@@ -144,7 +144,7 @@ namespace Metasound
 
 	private:
 		FORCENOINLINE FNodeFacade(const FVertexName& InInstanceName, const FGuid& InInstanceID, const FNodeClassMetadata& InNodeInfo, FFactory::FCreateOperatorFunction InCreate)
-		: FBasicNode(FNodeData(InInstanceName, InInstanceID, InNodeInfo.DefaultInterface), MakeShared<FNodeClassMetadata>(InNodeInfo))
+		: FBasicNode(FNodeData(InInstanceName, InInstanceID, InNodeInfo.DefaultInterface.CreateVertexInterface({}, {})), MakeShared<FNodeClassMetadata>(InNodeInfo))
 		, Factory(MakeShared<FFactory, ESPMode::ThreadSafe>(InCreate))
 		{
 		}
diff --git a/Engine/Plugins/Runtime/Metasound/Source/MetasoundGraphCore/Public/MetasoundGraph.h b/Engine/Plugins/Runtime/Metasound/Source/MetasoundGraphCore/Public/MetasoundGraph.h
--- a/Engine/Plugins/Runtime/Metasound/Source/MetasoundGraphCore/Public/MetasoundGraph.h
+++ b/Engine/Plugins/Runtime/Metasound/Source/MetasoundGraphCore/Public/MetasoundGraph.h
@@ -56,24 +56,6 @@ namespace Metasound
 			/** Return the configuration for this graph. */
 			UE_API virtual TSharedPtr<const IOperatorData> GetOperatorData() const override;
 
-			/** Set the vertex interface. If the vertex was successfully changed, returns true. 
-			 *
-			 * @param InInterface - New interface for node. 
-			 *
-			 * @return True on success, false otherwise.
-			 */
-			UE_DEPRECATED(5.6, "FGraphs can only set their interface by adding inputs and outputs.")
-			UE_API virtual bool SetVertexInterface(const FVertexInterface& InInterface) override;
-
-			/** Expresses whether a specific vertex interface is supported.
-			 *
-			 * @param InInterface - New interface. 
-			 *
-			 * @return True if the interface is supported, false otherwise. 
-			 */
-			UE_DEPRECATED(5.6, "This function no longer serves a purpose.")
-			UE_API virtual bool IsVertexInterfaceSupported(const FVertexInterface& InInterface) const override;
-
 			/** Get vertices which contain input parameters. */
 			UE_API virtual const FInputDataDestinationCollection& GetInputDataDestinations() const override;
 
@@ -263,6 +245,11 @@ namespace Metasound
 			FGuid InstanceID;
 
 			FNodeClassMetadata Metadata;
+
+			// Mutable vertex interface used for graph input/output mutations.
+			// Initialized from Metadata.DefaultInterface
+			FVertexInterface GraphInterface;
+
 			TSharedPtr<const IOperatorData> OperatorData;
 
 			TArray<FDataEdge> Edges;
@@ -290,7 +277,7 @@ namespace Metasound
 				INode* GetMutableNode() const;
 
 			private:
-				INode* MutableNode;
+				INode* MutableNode = nullptr;
 				TSharedPtr<const INode> Storage;
 			};
 
diff --git a/Engine/Plugins/Runtime/Metasound/Source/MetasoundGraphCore/Public/MetasoundNode.h b/Engine/Plugins/Runtime/Metasound/Source/MetasoundGraphCore/Public/MetasoundNode.h
--- a/Engine/Plugins/Runtime/Metasound/Source/MetasoundGraphCore/Public/MetasoundNode.h
+++ b/Engine/Plugins/Runtime/Metasound/Source/MetasoundGraphCore/Public/MetasoundNode.h
@@ -39,6 +39,11 @@ namespace Metasound
 			FVertexName InstanceName;
 			FGuid InstanceID;
 			FNodeClassMetadata Info;
+
+			// Mutable vertex interface for SetDefaultInput. Initialized from
+			// Info.DefaultInterface (implicit conversion from FClassInterface).
+			FVertexInterface VertexInterface;
+
 			TSharedPtr<const IOperatorData> OperatorData;
 	};
 }
diff --git a/Engine/Plugins/Runtime/Metasound/Source/MetasoundGraphCore/Public/MetasoundNodeInterface.h b/Engine/Plugins/Runtime/Metasound/Source/MetasoundGraphCore/Public/MetasoundNodeInterface.h
--- a/Engine/Plugins/Runtime/Metasound/Source/MetasoundGraphCore/Public/MetasoundNodeInterface.h
+++ b/Engine/Plugins/Runtime/Metasound/Source/MetasoundGraphCore/Public/MetasoundNodeInterface.h
@@ -159,7 +159,7 @@ namespace Metasound
 			FText InDescription = { },
 			FString InAuthor = { },
 			FText InPromptIfMissing = { },
-			FVertexInterface InDefaultInterface = { },
+			FClassInterface InDefaultInterface = { },
 			TArray<FText> InCategoryHierarchy = { },
 			TArray<FText> InKeywords = { },
 			FNodeDisplayStyle InDisplayStyle = { });
@@ -193,8 +193,8 @@ namespace Metasound
 		/** Human readable prompt for acquiring plugin in case node is not loaded. */
 		FText PromptIfMissing;
 
-		/** Default vertex interface for the node */
-		FVertexInterface DefaultInterface;
+		/** Default class interface for the node. */
+		FClassInterface DefaultInterface;
 
 		/** Hierarchy of categories for displaying node. */
 		TArray<FText> CategoryHierarchy;
@@ -250,24 +250,6 @@ namespace Metasound
 			UE_EXPERIMENTAL(5.6, "Node operator data is still under development")
 			virtual METASOUNDGRAPHCORE_API TSharedPtr<const IOperatorData> GetOperatorData() const;
 #endif
-
-			/** Set the vertex interface. If the vertex was successfully changed, returns true. 
-			 *
-			 * @param InInterface - New interface for node. 
-			 *
-			 * @return True on success, false otherwise.
-			 */
-			UE_DEPRECATED(5.6, "INodeBase will not expose dynamic FVertexInterface operations")
-			virtual bool SetVertexInterface(const FVertexInterface& InInterface) { return false; }
-
-			/** Expresses whether a specific vertex interface is supported.
-			 *
-			 * @param InInterface - New interface. 
-			 *
-			 * @return True if the interface is supported, false otherwise. 
-			 */
-			UE_DEPRECATED(5.6, "INodeBase will not expose dynamic FVertexInterface operations")
-			virtual bool IsVertexInterfaceSupported(const FVertexInterface& InInterface) const { return false; }
 	};
 
 	// Forward declare
diff --git a/Engine/Plugins/Runtime/Metasound/Source/MetasoundGraphCore/Public/MetasoundVertex.h b/Engine/Plugins/Runtime/Metasound/Source/MetasoundGraphCore/Public/MetasoundVertex.h
--- a/Engine/Plugins/Runtime/Metasound/Source/MetasoundGraphCore/Public/MetasoundVertex.h
+++ b/Engine/Plugins/Runtime/Metasound/Source/MetasoundGraphCore/Public/MetasoundVertex.h
@@ -20,7 +20,6 @@ namespace Metasound
 
 	
 	/** Describes a variant on a ClassInterface */
-	struct UE_EXPERIMENTAL(5.8, "MetaSound Variants are experimental") FVariantDescription;
 	struct FVariantDescription
 	{
 		FName VariantName; 
@@ -28,7 +27,6 @@ namespace Metasound
 	};
 
 	/** Configures a variant to be a specific data type. */
-	struct UE_EXPERIMENTAL(5.8, "MetaSound Variants are experimental") FVariantConfiguration;
 	struct FVariantConfiguration
 	{
 		FName VariantName;
@@ -841,7 +839,6 @@ namespace Metasound
 	};
 
 	/** A description of a sub interface which is used when declaring a FClassInterface */
-	struct UE_EXPERIMENTAL(5.6, "Sub interfaces are experimental") FSubInterfaceDescription;
 	struct FSubInterfaceDescription
 	{
 		FName SubInterfaceName;
@@ -852,7 +849,6 @@ namespace Metasound
 
 	/** A sub interface configuration is used to create a FVertexInterface from 
 	 * a FClassInterface. */
-	struct UE_EXPERIMENTAL(5.6, "Sub interfaces are experimental") FSubInterfaceConfiguration;
 	struct FSubInterfaceConfiguration
 	{
 		FName SubInterfaceName;
@@ -867,7 +863,6 @@ namespace Metasound
 	/** A variant vertex represents a vertex that has a variant data type. Vertices
 	 * are associated with a variant by using the VariantName. Multiple vertices
 	 * may share the same variant. */
-	class UE_EXPERIMENTAL(5.8, "MetaSound Variants are experimental") FVariantVertex;
 	class FVariantVertex
 	{
 	protected:
@@ -900,8 +895,6 @@ namespace Metasound
 	/** An input variant vertex represents an input vertex that has a variant data 
 	 * type. Vertices are associated with a variant by using the VariantName. 
 	 * Multiple vertices may share the same variant. */
-	class UE_EXPERIMENTAL(5.8, "MetaSound Variants are experimental") FInputVariantVertex;
-PRAGMA_DISABLE_EXPERIMENTAL_WARNINGS
 	class FInputVariantVertex : public FVariantVertex
 	{
 	protected:
@@ -922,7 +915,6 @@ PRAGMA_DISABLE_EXPERIMENTAL_WARNINGS
 	/** An output variant vertex represents an output vertex that has a variant data 
 	 * type. Vertices are associated with a variant by using the VariantName. 
 	 * Multiple vertices may share the same variant. */
-	class UE_EXPERIMENTAL(5.8, "MetaSound Variants are experimental") FOutputVariantVertex;
 	class FOutputVariantVertex : public FVariantVertex
 	{
 	public:
@@ -932,7 +924,6 @@ PRAGMA_DISABLE_EXPERIMENTAL_WARNINGS
 		friend bool METASOUNDGRAPHCORE_API operator!=(const FOutputVariantVertex& InLHS, const FOutputVariantVertex& InRHS);
 		friend bool METASOUNDGRAPHCORE_API operator<(const FOutputVariantVertex& InLHS, const FOutputVariantVertex& InRHS);
 	};
-PRAGMA_ENABLE_EXPERIMENTAL_WARNINGS
 
 	namespace VertexPrivate
 	{
@@ -966,9 +957,6 @@ PRAGMA_ENABLE_EXPERIMENTAL_WARNINGS
 		}
 	}
 
-	template<typename ... DataTypes>
-	class UE_EXPERIMENTAL(5.8, "MetaSound Variants are experimental") TInputVariantVertex;
-
 	/** An input variant vertex represents an input vertex that has a variant data 
 	 * type. Vertices are associated with a variant by using the VariantName. 
 	 * Multiple vertices may share the same variant. The supported data types of
@@ -1002,10 +990,6 @@ PRAGMA_ENABLE_EXPERIMENTAL_WARNINGS
 		}
 	};
 
-	template<typename ... DataTypes>
-	class UE_EXPERIMENTAL(5.8, "MetaSound Variants are experimental") TOutputVariantVertex;
-
-
 	/** An output variant vertex represents an output vertex that has a variant data 
 	 * type. Vertices are associated with a variant by using the VariantName. 
 	 * Multiple vertices may share the same variant. The supported data types of
@@ -1036,9 +1020,9 @@ PRAGMA_ENABLE_EXPERIMENTAL_WARNINGS
 		}
 	};
 
-	/** FClassInterface describes the interface of a node class. It is an 
-	 * immutable factory for FVertexInterfaces. */
-	class UE_EXPERIMENTAL(5.6, "FClassInterfaces are experimental") FClassInterface;
+	/** FClassInterface describes the interface of a node class. It is a
+	 * factory for FVertexInterfaces. */
+	class FClassInterface;
 
 	class FSubInterface;
 
@@ -1049,7 +1033,6 @@ PRAGMA_ENABLE_EXPERIMENTAL_WARNINGS
 		// duration of a function call. 
 		struct FClassInterfaceData
 		{
-PRAGMA_DISABLE_EXPERIMENTAL_WARNINGS
 			TArray<FSubInterfaceDescription, TInlineAllocator<4>> SubInterfaces;
 			TArray<FSubInterfaceClassLayout, TInlineAllocator<2>> InputSubInterfaceLayouts;
 			TArray<FSubInterfaceClassLayout, TInlineAllocator<2>> OutputSubInterfaceLayouts;
@@ -1057,14 +1040,12 @@ PRAGMA_DISABLE_EXPERIMENTAL_WARNINGS
 			TSortedMap<FVertexKey, uint8, TInlineAllocator<4>> VertexKeyToVariantIndex;
 			TArray<FVariantDescription, TInlineAllocator<4>> Variants;
 			TSortedVertexNameMap<TArray<FLiteral, TInlineAllocator<4>>, TInlineAllocator<4>> VariantLiterals;
-PRAGMA_ENABLE_EXPERIMENTAL_WARNINGS
 
 			TArray<FInputDataVertex, TInlineAllocator<8>> Inputs;
 			TArray<FOutputDataVertex, TInlineAllocator<8>> Outputs;
 			TArray<FEnvironmentVertex, TInlineAllocator<4>> EnvironmentVariables;
 		};
 
-PRAGMA_DISABLE_EXPERIMENTAL_WARNINGS
 		// Stores information to describe variants on a class interface
 		struct FVariantSchema
 		{
@@ -1084,7 +1065,6 @@ PRAGMA_DISABLE_EXPERIMENTAL_WARNINGS
 
 			UE_API static FSubInterfaceSchema FromClassInterfaceData(const VertexPrivate::FClassInterfaceData& InClassInterfaceData);
 		};
-PRAGMA_ENABLE_EXPERIMENTAL_WARNINGS
 
 		/** The FClassInterfaceDataBuilder assembles declared inputs, outputs,
 		 * variants and subinterfaces into a FClassInterfaceData object. This is
@@ -1120,7 +1100,6 @@ PRAGMA_ENABLE_EXPERIMENTAL_WARNINGS
 			UE_API void Add(const FOutputDataVertex* InOutputVertex);
 			UE_API void Add(const FEnvironmentVertex* InEnvironmentVertex);
 
-PRAGMA_DISABLE_EXPERIMENTAL_WARNINGS
 			UE_API void Add(const FSubInterfaceDescription* InSubInterfaceDescription);
 			UE_API void Add(const FSubInterface* InSubInterface);
 			UE_API void Add(const FInputVariantVertex* InInputVariantVertex);
@@ -1132,7 +1111,6 @@ PRAGMA_DISABLE_EXPERIMENTAL_WARNINGS
 
 			uint8 FindOrAddVariant(const FVariantVertex& InVariantVertex);
 			int32 IndexOfVariant(const FName& InName) const;
-PRAGMA_ENABLE_EXPERIMENTAL_WARNINGS
 
 			FClassInterfaceData ClassInterfaceData;
 		};
@@ -1160,8 +1138,6 @@ PRAGMA_ENABLE_EXPERIMENTAL_WARNINGS
 	} // namespace VertexPrivate;
 
 	
-	class UE_EXPERIMENTAL(5.8, "MetaSound SubInterfaces are experimental") FSubInterface;
-PRAGMA_DISABLE_EXPERIMENTAL_WARNINGS
 	class FSubInterface : public VertexPrivate::FBuildableInterface
 	{
 	public:
@@ -1177,10 +1153,7 @@ PRAGMA_DISABLE_EXPERIMENTAL_WARNINGS
 
 		FName SubInterfaceName;
 	};
-PRAGMA_ENABLE_EXPERIMENTAL_WARNINGS
 
-	class UE_EXPERIMENTAL(5.8, "MetaSound ClassInterfaces are experimental") FClassInterface;
-PRAGMA_DISABLE_EXPERIMENTAL_WARNINGS
 	/** An FClassInterface defines the class interface of a MetaSound C++ node class. This
 	 * interface may contain basic vertices, but also subinterfaces and variants. 
 	 *
@@ -1193,6 +1166,26 @@ PRAGMA_DISABLE_EXPERIMENTAL_WARNINGS
 		FClassInterface() = default;
 		~FClassInterface() = default;
 
+		/** Copy constructor (deep-copies schema data). */
+		UE_API FClassInterface(const FClassInterface& Other);
+
+		/** Copy assignment (deep-copies schema data). */
+		UE_API FClassInterface& operator=(const FClassInterface& Other);
+
+		/** Move constructor. */
+		FClassInterface(FClassInterface&&) = default;
+
+		/** Move assignment. */
+		FClassInterface& operator=(FClassInterface&&) = default;
+
+		/** Construct from a FVertexInterface. The resulting FClassInterface has
+		 * no sub-interfaces or variants — it simply wraps the flat vertex lists. */
+		UE_API FClassInterface(const FVertexInterface& InVertexInterface);
+
+		/** Conversion to FVertexInterface. Returns the default-configured
+		 * interface (no sub-interface or variant configurations). */
+		UE_API operator FVertexInterface() const;
+
 		UE_INTERNAL
 		UE_API FClassInterface(const VertexPrivate::FClassInterfaceData& InClassInterfaceData);
 
@@ -1213,10 +1206,17 @@ PRAGMA_DISABLE_EXPERIMENTAL_WARNINGS
 		 * is associated with a variant. */
 		UE_API const FVariantDescription* FindVariantDescriptionForOutput(const FVertexName& InVertexName) const;
 
+		/** Returns true if this class interface contains sub-interface definitions. */
+		bool ContainsSubInterfaces() const { return SubInterfaceSchema && !SubInterfaceSchema->SubInterfaces.IsEmpty(); }
+
+		/** Returns true if this class interface contains variant definitions. */
+		bool ContainsVariants() const { return VariantSchema && !VariantSchema->Variants.IsEmpty(); }
+
 	private:
-		
+
 		// Many nodes do not need this extra configurable information. Save on memory
-		// storage by offloading this extra functionality behind a unique ptr. 
+		// storage by offloading this extra functionality behind a unique ptr.
+		// Deep-copied in copy constructor/assignment to preserve value semantics.
 		TUniquePtr<VertexPrivate::FVariantSchema> VariantSchema;
 		TUniquePtr<VertexPrivate::FSubInterfaceSchema> SubInterfaceSchema;
 
@@ -1238,12 +1238,10 @@ PRAGMA_DISABLE_EXPERIMENTAL_WARNINGS
 	 * );
 	 */
 	template<typename ... ArgTypes>
-	UE_EXPERIMENTAL(5.7, "FClassInterface is experimental")
 	FClassInterface CreateClassInterface(ArgTypes&& ... InArgs)
 	{
 		return FClassInterface{VertexPrivate::FClassInterfaceDataBuilder().Build(Forward<ArgTypes>(InArgs)...).GetClassInterfaceData()};
 	}
-PRAGMA_ENABLE_EXPERIMENTAL_WARNINGS
 
 	/**
 	 * This struct is used to pass in any arguments required for constructing a single node instance.
diff --git a/Engine/Plugins/Runtime/Metasound/Source/MetasoundStandardNodes/Private/MetasoundAudioFormats.cpp b/Engine/Plugins/Runtime/Metasound/Source/MetasoundStandardNodes/Private/MetasoundAudioFormats.cpp
--- a/Engine/Plugins/Runtime/Metasound/Source/MetasoundStandardNodes/Private/MetasoundAudioFormats.cpp
+++ b/Engine/Plugins/Runtime/Metasound/Source/MetasoundStandardNodes/Private/MetasoundAudioFormats.cpp
@@ -206,16 +206,6 @@ namespace Metasound
 			return VertexInterface;
 		}
 
-		bool SetVertexInterface(const FVertexInterface& InInterface) override
-		{
-			return VertexInterface == InInterface;
-		}
-
-		bool IsVertexInterfaceSupported(const FVertexInterface& InInterface) const 
-		{
-			return VertexInterface == InInterface;
-		}
-
 		virtual TSharedRef<IOperatorFactory, ESPMode::ThreadSafe> GetDefaultOperatorFactory() const override
 		{
 			return Factory;
@@ -446,16 +436,6 @@ namespace Metasound
 			return VertexInterface;
 		}
 
-		bool SetVertexInterface(const FVertexInterface& InInterface) override
-		{
-			return VertexInterface == InInterface;
-		}
-
-		bool IsVertexInterfaceSupported(const FVertexInterface& InInterface) const 
-		{
-			return VertexInterface == InInterface;
-		}
-
 		TSharedRef<IOperatorFactory, ESPMode::ThreadSafe> GetDefaultOperatorFactory() const override
 		{
 			return Factory;
