CIS fix for non-proxy interface generation in metasound cooking.
#rnx
#Synced-CL 52515588

diff --git a/Engine/Plugins/Runtime/Metasound/Source/MetasoundFrontend/Private/MetasoundAssetBase.cpp b/Engine/Plugins/Runtime/Metasound/Source/MetasoundFrontend/Private/MetasoundAssetBase.cpp
--- a/Engine/Plugins/Runtime/Metasound/Source/MetasoundFrontend/Private/MetasoundAssetBase.cpp
+++ b/Engine/Plugins/Runtime/Metasound/Source/MetasoundFrontend/Private/MetasoundAssetBase.cpp
@@ -88,7 +88,6 @@ namespace Metasound
 				{
 					// Copy FrontendClass to preserve original document.
 					FrontendClass.Metadata.SetType(EMetasoundFrontendClassType::External);
-					CachedClassInterface = FClassInterface(CreateDefaultVertexInterfaceFromClassNoProxy(FrontendClass));
 				}
 
 				FDocumentNodeRegistryEntryForSerialization(const FDocumentNodeRegistryEntryForSerialization& InOther) = default;
@@ -115,7 +114,8 @@ namespace Metasound
 
 				virtual const FClassInterface& GetClassInterface() const override
 				{
-					return CachedClassInterface;
+					UE_LOG(LogMetaSound, Error, TEXT("ClassInterfaces for assets are not available during cook. Cannot create interface for %s"), *FrontendClass.Metadata.GetClassName().ToString())
+					return InvalidClassInterface;
 				}
 
 				virtual TInstancedStruct<FMetaSoundFrontendNodeConfiguration> CreateFrontendNodeConfiguration() const override
@@ -134,8 +134,11 @@ namespace Metasound
 				TSet<FMetasoundFrontendVersion> Interfaces;
 				FMetasoundFrontendClass FrontendClass;
 				FTopLevelAssetPath AssetPath;
-				FClassInterface CachedClassInterface;
+				static const FClassInterface InvalidClassInterface;
 			};
+
+			const FClassInterface FDocumentNodeRegistryEntryForSerialization::InvalidClassInterface;
+
 		} // namespace AssetBasePrivate
 
 		FConsoleVariableMulticastDelegate CVarMetaSoundBlockRateChanged;
