[Backout] - CL52519867
#rnx
Original CL Desc
-----------------------------------------------------------------
CIS fix for non-proxy interface generation in metasound cooking.
#rnx

diff --git a/Engine/Plugins/Runtime/Metasound/Source/MetasoundFrontend/Private/MetasoundAssetBase.cpp b/Engine/Plugins/Runtime/Metasound/Source/MetasoundFrontend/Private/MetasoundAssetBase.cpp
--- a/Engine/Plugins/Runtime/Metasound/Source/MetasoundFrontend/Private/MetasoundAssetBase.cpp
+++ b/Engine/Plugins/Runtime/Metasound/Source/MetasoundFrontend/Private/MetasoundAssetBase.cpp
@@ -88,6 +88,7 @@ namespace Metasound
 				{
 					// Copy FrontendClass to preserve original document.
 					FrontendClass.Metadata.SetType(EMetasoundFrontendClassType::External);
+					CachedClassInterface = FClassInterface(CreateDefaultVertexInterfaceFromClassNoProxy(FrontendClass));
 				}
 
 				FDocumentNodeRegistryEntryForSerialization(const FDocumentNodeRegistryEntryForSerialization& InOther) = default;
@@ -114,8 +115,7 @@ namespace Metasound
 
 				virtual const FClassInterface& GetClassInterface() const override
 				{
-					UE_LOG(LogMetaSound, Error, TEXT("ClassInterfaces for assets are not available during cook. Cannot create interface for %s"), *FrontendClass.Metadata.GetClassName().ToString())
-					return InvalidClassInterface;
+					return CachedClassInterface;
 				}
 
 				virtual TInstancedStruct<FMetaSoundFrontendNodeConfiguration> CreateFrontendNodeConfiguration() const override
@@ -134,11 +134,8 @@ namespace Metasound
 				TSet<FMetasoundFrontendVersion> Interfaces;
 				FMetasoundFrontendClass FrontendClass;
 				FTopLevelAssetPath AssetPath;
-				static const FClassInterface InvalidClassInterface;
+				FClassInterface CachedClassInterface;
 			};
-
-			const FClassInterface FDocumentNodeRegistryEntryForSerialization::InvalidClassInterface;
-
 		} // namespace AssetBasePrivate
 
 		FConsoleVariableMulticastDelegate CVarMetaSoundBlockRateChanged;
