MetaSound Input Browser for bulk viewing/editing preset input values

#tests Tested a bunch in editor using big asset selections

#Synced-CL 52679804

diff --git a/Engine/Plugins/Runtime/Metasound/Source/MetasoundEditor/Private/MetasoundAssetDefinitions.cpp b/Engine/Plugins/Runtime/Metasound/Source/MetasoundEditor/Private/MetasoundAssetDefinitions.cpp
--- a/Engine/Plugins/Runtime/Metasound/Source/MetasoundEditor/Private/MetasoundAssetDefinitions.cpp
+++ b/Engine/Plugins/Runtime/Metasound/Source/MetasoundEditor/Private/MetasoundAssetDefinitions.cpp
@@ -10,6 +10,7 @@
 #include "ContentBrowserModule.h"
 #include "Editor/EditorPerProjectUserSettings.h"
 #include "FileHelpers.h"
+#include "Framework/Docking/TabManager.h"
 #include "IAssetTools.h"
 #include "IContentBrowserSingleton.h"
 #include "ISourceControlModule.h"
@@ -31,6 +32,7 @@
 #include "Misc/AssetRegistryInterface.h"
 #include "Modules/ModuleManager.h"
 #include "ObjectEditorUtils.h"
+#include "SMetasoundInputBrowser.h"
 #include "Sound/SoundWave.h"
 #include "Styling/AppStyle.h"
 #include "Styling/CoreStyle.h"
@@ -40,6 +42,7 @@
 #include "DocumentTemplates/MetasoundFrontendPresetTemplate.h"
 #include "ToolMenus.h"
 #include "ToolMenuSection.h"
+#include "Widgets/Docking/SDockTab.h"
 #include "Widgets/Input/SButton.h"
 #include "Widgets/SOverlay.h"
 #include "UObject/Class.h"
@@ -1170,6 +1173,84 @@ namespace MenuExtension_MetaSoundTemplate
 					MetaSoundAssetActionsMenu.AddDynamicEntry(NAME_None, AssetActionsSectionDelegate);
 				});
 			}
+
+			// "Browse MetaSound Inputs..." action - opens the MetaSound Input Browser with selected assets.
+			{
+				auto ExecuteBrowseInputs = [](const FToolMenuContext& InContext)
+				{
+					const UContentBrowserAssetContextMenuContext* CBContext = UContentBrowserAssetContextMenuContext::FindContextWithAssets(InContext);
+					if (!CBContext)
+					{
+						return;
+					}
+
+					TArray<TWeakObjectPtr<UObject>> SelectedAssets;
+					for (const FAssetData& AssetData : CBContext->SelectedAssets)
+					{
+						if (UObject* Asset = AssetData.GetAsset())
+						{
+							if (Cast<IMetaSoundDocumentInterface>(Asset))
+							{
+								SelectedAssets.Add(Asset);
+							}
+						}
+					}
+
+					if (SelectedAssets.IsEmpty())
+					{
+						return;
+					}
+
+					TSharedPtr<SDockTab> ExistingTab = FGlobalTabmanager::Get()->TryInvokeTab(SMetasoundInputBrowser::TabId);
+					if (ExistingTab.IsValid())
+					{
+						TSharedRef<SMetasoundInputBrowser> Browser = StaticCastSharedRef<SMetasoundInputBrowser>(ExistingTab->GetContent());
+						Browser->SetAssets(SelectedAssets);
+					}
+				};
+
+				// Register on the base context menu for mixed-type selections
+				{
+					UToolMenu* BaseMenu = UToolMenus::Get()->ExtendMenu("ContentBrowser.AssetContextMenu");
+					FToolMenuSection& Section = BaseMenu->FindOrAddSection("MetaSoundBrowseInputs", LOCTEXT("MetaSoundBrowseInputsHeading", "MetaSound"));
+
+					const FNewToolMenuSectionDelegate BaseDelegate = FNewToolMenuSectionDelegate::CreateLambda(
+						[ExecuteBrowseInputs](FToolMenuSection& InSection)
+					{
+						FToolUIAction UIAction;
+						UIAction.ExecuteAction = FToolMenuExecuteAction::CreateLambda(ExecuteBrowseInputs);
+						UIAction.IsActionVisibleDelegate = FToolMenuIsActionButtonVisible::CreateLambda([](const FToolMenuContext& InContext)
+						{
+							const UContentBrowserAssetContextMenuContext* CBContext = UContentBrowserAssetContextMenuContext::FindContextWithAssets(InContext);
+							if (!CBContext)
+							{
+								return false;
+							}
+							for (const FAssetData& AssetData : CBContext->SelectedAssets)
+							{
+								if (UObject* Asset = AssetData.GetAsset())
+								{
+									if (Cast<IMetaSoundDocumentInterface>(Asset))
+									{
+										return true;
+									}
+								}
+							}
+							return false;
+						});
+
+						InSection.AddMenuEntry(
+							"MetaSound_BrowseInputs",
+							LOCTEXT("BrowseInputsTitle", "Browse MetaSound Inputs..."),
+							LOCTEXT("BrowseInputsTooltip", "Open the MetaSound Input Browser to view and edit input values across selected assets in bulk."),
+							FSlateIcon(FAppStyle::GetAppStyleSetName(), "Icons.Info"),
+							UIAction
+						);
+					});
+
+					Section.AddDynamicEntry(NAME_None, BaseDelegate);
+				}
+			}
  		}));
 	});
 } // namespace MenuExtension_MetaSoundTemplate
diff --git a/Engine/Plugins/Runtime/Metasound/Source/MetasoundEditor/Private/MetasoundEditorModule.cpp b/Engine/Plugins/Runtime/Metasound/Source/MetasoundEditor/Private/MetasoundEditorModule.cpp
--- a/Engine/Plugins/Runtime/Metasound/Source/MetasoundEditor/Private/MetasoundEditorModule.cpp
+++ b/Engine/Plugins/Runtime/Metasound/Source/MetasoundEditor/Private/MetasoundEditorModule.cpp
@@ -22,6 +22,7 @@
 #include "Editor.h"
 #include "FileHelpers.h"
 #include "Framework/Commands/UIAction.h"
+#include "Framework/Docking/TabManager.h"
 #include "Framework/Notifications/NotificationManager.h"
 #include "HAL/IConsoleManager.h"
 #include "IDetailCustomization.h"
@@ -59,6 +60,7 @@
 #include "PropertyEditorDelegates.h"
 #include "PropertyEditorModule.h"
 #include "SMetasoundFilterFrequencyResponsePlots.h"
+#include "SMetasoundInputBrowser.h"
 #include "Styling/AppStyle.h"
 #include "Styling/CoreStyle.h"
 #include "Styling/SlateStyle.h"
@@ -77,6 +79,7 @@
 #include "UObject/NoExportTypes.h"
 #include "UObject/UObjectGlobals.h"
 #include "UObject/UObjectIterator.h"
+#include "Widgets/Docking/SDockTab.h"
 #include "Widgets/Notifications/SNotificationList.h"
 
 
@@ -1221,6 +1224,23 @@ namespace Metasound::Editor
 			{
 				Settings->OverrideTargetPageSettings(EditorSettings->GetPageForAudition());
 			}
+
+			// Register MetaSound Input Browser nomad tab
+			FGlobalTabmanager::Get()->RegisterNomadTabSpawner(
+				SMetasoundInputBrowser::TabId,
+				FOnSpawnTab::CreateLambda([](const FSpawnTabArgs&) -> TSharedRef<SDockTab>
+				{
+					TSharedRef<SMetasoundInputBrowser> Browser = SNew(SMetasoundInputBrowser);
+					return SNew(SDockTab)
+						.TabRole(ETabRole::NomadTab)
+						.Label(LOCTEXT("MetaSoundInputBrowserTabTitle", "MetaSound Input Browser"))
+						[
+							Browser
+						];
+				}))
+				.SetDisplayName(LOCTEXT("MetaSoundInputBrowserDisplayName", "MetaSound Input Browser"))
+				.SetTooltipText(LOCTEXT("MetaSoundInputBrowserTooltip", "Bulk view and edit MetaSound input values across multiple assets"))
+				.SetIcon(FSlateIcon(FAppStyle::GetAppStyleSetName(), "Icons.Info"));
 		}
 
 		virtual void ShutdownModule() override
@@ -1229,6 +1249,8 @@ namespace Metasound::Editor
 
 			METASOUND_LLM_SCOPE;
 
+			FGlobalTabmanager::Get()->UnregisterNomadTabSpawner(SMetasoundInputBrowser::TabId);
+
 			if (ISettingsModule* SettingsModule = FModuleManager::GetModulePtr<ISettingsModule>("Settings"))
 			{
 				SettingsModule->UnregisterSettings("Editor", "Audio", "MetaSound Editor");
diff --git a/Engine/Plugins/Runtime/Metasound/Source/MetasoundEditor/Private/SMetasoundInputBrowser.cpp b/Engine/Plugins/Runtime/Metasound/Source/MetasoundEditor/Private/SMetasoundInputBrowser.cpp
new file mode 100644
--- /dev/null
+++ b/Engine/Plugins/Runtime/Metasound/Source/MetasoundEditor/Private/SMetasoundInputBrowser.cpp
@@ -0,0 +1,1236 @@
+// Copyright Epic Games, Inc. All Rights Reserved.
+
+#include "SMetasoundInputBrowser.h"
+
+#include "ContentBrowserModule.h"
+#include "DocumentTemplates/MetasoundFrontendDocumentVertexTemplate.h"
+#include "DocumentTemplates/MetasoundFrontendPresetTemplate.h"
+#include "Editor.h"
+#include "Framework/MultiBox/MultiBoxBuilder.h"
+#include "HAL/PlatformApplicationMisc.h"
+#include "IContentBrowserSingleton.h"
+#include "MetasoundDocumentInterface.h"
+#include "MetasoundEditorGraphBuilder.h"
+#include "MetasoundEditorModule.h"
+#include "MetasoundFrontendDocumentBuilder.h"
+#include "MetasoundFrontendDocumentBuilderRegistry.h"
+#include "MetasoundFrontendDocumentModifyDelegates.h"
+#include "Misc/MessageDialog.h"
+#include "ScopedTransaction.h"
+#include "Styling/AppStyle.h"
+#include "Subsystems/AssetEditorSubsystem.h"
+#include "Widgets/Input/SButton.h"
+#include "Widgets/Input/SCheckBox.h"
+#include "Widgets/Input/SEditableTextBox.h"
+#include "Widgets/Input/SSearchBox.h"
+#include "Widgets/Input/SSpinBox.h"
+#include "Widgets/Layout/SBox.h"
+#include "Widgets/Layout/SScrollBox.h"
+#include "Widgets/Layout/SSplitter.h"
+#include "Widgets/SNullWidget.h"
+#include "Widgets/Text/STextBlock.h"
+#include "Widgets/Views/SHeaderRow.h"
+
+#define LOCTEXT_NAMESPACE "MetaSoundInputBrowser"
+
+namespace Metasound::Editor
+{
+
+	const FName SMetasoundInputBrowser::TabId("MetaSoundInputBrowser");
+
+	namespace InputBrowserPrivate
+	{
+		const FName Column_AssetName("AssetName");
+		const FName TriggerTypeName("Trigger");
+		const FName AudioTypeName("Audio");
+		constexpr float ResetButtonWidth = 20.0f;
+
+		FLinearColor GetColorForTypeName(FName TypeName)
+		{
+			const IMetasoundEditorModule& EditorModule = FModuleManager::GetModuleChecked<IMetasoundEditorModule>("MetaSoundEditor");
+			if (const FEdGraphPinType* PinType = EditorModule.FindPinType(TypeName))
+			{
+				return FGraphBuilder::GetPinCategoryColor(*PinType);
+			}
+			return FGraphBuilder::GetPinCategoryColor(FEdGraphPinType());
+		}
+
+		const FSlateBrush* GetIconForTypeName(FName TypeName)
+		{
+			const IMetasoundEditorModule& EditorModule = FModuleManager::GetModuleChecked<IMetasoundEditorModule>("MetaSoundEditor");
+			constexpr bool bIsConstructorType = false;
+			return EditorModule.GetIconBrush(TypeName, bIsConstructorType);
+		}
+		// Returns true if the given input has non-default page values that differ
+		// from the default page value, meaning a reset would lose page-specific data.
+		bool HasNonDefaultPageOverrides(const FMetaSoundFrontendDocumentBuilder& Builder, FName InputName)
+		{
+			using namespace Metasound::Frontend;
+
+			const FMetasoundFrontendClassInput* GraphInput = Builder.FindGraphInput(InputName);
+			if (!GraphInput)
+			{
+				return false;
+			}
+
+			const FMetasoundFrontendLiteral* DefaultPageLiteral = GraphInput->FindConstDefault(DefaultPageID);
+			if (!DefaultPageLiteral)
+			{
+				return false;
+			}
+
+			bool bHasPageOverrides = false;
+			GraphInput->IterateDefaults([&](const FGuid& PageID, const FMetasoundFrontendLiteral& Literal)
+			{
+				if (PageID != DefaultPageID && !Literal.IsEqual(*DefaultPageLiteral))
+				{
+					bHasPageOverrides = true;
+				}
+			});
+
+			return bHasPageOverrides;
+		}
+	} // namespace InputBrowserPrivate
+
+	// ---------------------------------------------------------------
+	// FInputBrowserRow
+	// ---------------------------------------------------------------
+
+	IMetaSoundDocumentInterface* FInputBrowserRow::GetDocumentInterface() const
+	{
+		return Cast<IMetaSoundDocumentInterface>(Asset.Get());
+	}
+
+	void FInputBrowserRow::ReadCellData(
+		FInputBrowserCellData& OutCell,
+		const FMetasoundFrontendClassInput& Input,
+		const FMetaSoundFrontendPresetTemplate* PresetTemplate)
+	{
+		using namespace Metasound::Frontend;
+
+		OutCell.NodeID = Input.NodeID;
+		OutCell.bIsPreset = (PresetTemplate != nullptr);
+		OutCell.bInputExists = true;
+		OutCell.bOverridesDefault = false;
+
+		if (PresetTemplate)
+		{
+			const FMetaSoundFrontendPresetVertexMetadata* VertexMeta =
+				PresetTemplate->FindConstVertexMetadata<FMetaSoundFrontendPresetVertexMetadata>(Input.NodeID);
+			OutCell.bOverridesDefault = VertexMeta ? VertexMeta->bOverrideInheritedDefault : false;
+		}
+
+		if (const FMetasoundFrontendLiteral* DefaultLiteral = Input.FindConstDefault(DefaultPageID))
+		{
+			OutCell.Value = *DefaultLiteral;
+		}
+	}
+
+	TArray<FInputBrowserColumnDef> FInputBrowserRow::BuildRow(TSharedPtr<FInputBrowserRow>& OutRow, UObject* InAsset)
+	{
+		using namespace Metasound::Frontend;
+
+		IMetaSoundDocumentInterface* DocInterface = Cast<IMetaSoundDocumentInterface>(InAsset);
+		if (!DocInterface)
+		{
+			return {};
+		}
+
+		TArray<FInputBrowserColumnDef> FoundColumns;
+
+		OutRow = MakeShared<FInputBrowserRow>();
+		OutRow->Asset = InAsset;
+		OutRow->AssetName = InAsset->GetName();
+
+		const FMetasoundFrontendDocument& Document = DocInterface->GetConstDocument();
+		const FMetasoundFrontendClassInterface& Interface = Document.RootGraph.GetDefaultInterface();
+
+		const FMetaSoundFrontendPresetTemplate* PresetTemplate = Document.Template.GetPtr<FMetaSoundFrontendPresetTemplate>();
+
+		for (const FMetasoundFrontendClassInput& Input : Interface.Inputs)
+		{
+			if (Input.TypeName == InputBrowserPrivate::TriggerTypeName || Input.TypeName == InputBrowserPrivate::AudioTypeName)
+			{
+				continue;
+			}
+
+			FInputBrowserCellData CellData;
+			ReadCellData(CellData, Input, PresetTemplate);
+			OutRow->Cells.Add(Input.Name, MoveTemp(CellData));
+
+			FInputBrowserColumnDef ColDef;
+			ColDef.InputName = Input.Name;
+			ColDef.TypeName = Input.TypeName;
+			const FMetasoundFrontendLiteral* DefaultLiteral = Input.FindConstDefault(DefaultPageID);
+			ColDef.LiteralType = DefaultLiteral ? DefaultLiteral->GetType() : EMetasoundFrontendLiteralType::None;
+			FoundColumns.Add(MoveTemp(ColDef));
+		}
+
+		return FoundColumns;
+	}
+
+	// ---------------------------------------------------------------
+	// SMetasoundInputBrowser
+	// ---------------------------------------------------------------
+
+	void SMetasoundInputBrowser::Construct(const FArguments& InArgs)
+	{
+		HeaderRow = SNew(SHeaderRow);
+
+		ChildSlot
+			[
+				SNew(SVerticalBox)
+
+					// Toolbar
+					+ SVerticalBox::Slot()
+					.AutoHeight()
+					.Padding(4.0f)
+					[
+						SNew(SHorizontalBox)
+
+							+ SHorizontalBox::Slot()
+							.AutoWidth()
+							.VAlign(VAlign_Center)
+							.Padding(0.0f, 0.0f, 8.0f, 0.0f)
+							[
+								SNew(STextBlock)
+									.Text(LOCTEXT("FilterLabel", "Filter:"))
+							]
+
+							+ SHorizontalBox::Slot()
+							.FillWidth(1.0f)
+							[
+								SNew(SSearchBox)
+									.HintText(LOCTEXT("SearchHint", "Filter inputs by name..."))
+									.OnTextChanged(this, &SMetasoundInputBrowser::OnNameFilterChanged)
+							]
+
+					]
+
+				// Splitter: Pin list | Table
+				+ SVerticalBox::Slot()
+					.FillHeight(1.0f)
+					[
+						SNew(SSplitter)
+							.Orientation(Orient_Horizontal)
+
+							// Left panel: Table
+							+ SSplitter::Slot()
+							.Value(0.75f)
+							[
+								SNew(SScrollBox)
+									.Orientation(Orient_Horizontal)
+									+ SScrollBox::Slot()
+									[
+										SAssignNew(ListView, SListView<TSharedPtr<FInputBrowserRow>>)
+											.ListItemsSource(&Rows)
+											.OnGenerateRow(this, &SMetasoundInputBrowser::OnGenerateRow)
+											.HeaderRow(HeaderRow)
+											.SelectionMode(ESelectionMode::Single)
+											.OnContextMenuOpening(this, &SMetasoundInputBrowser::OnAssetContextMenuOpening)
+									]
+							]
+
+						// Right panel: Pinned columns list
+						+ SSplitter::Slot()
+							.Value(0.25f)
+							[
+								SNew(SVerticalBox)
+
+									+ SVerticalBox::Slot()
+									.AutoHeight()
+									.Padding(4.0f)
+									[
+										SNew(SHorizontalBox)
+
+											+ SHorizontalBox::Slot()
+											.FillWidth(1.0f)
+											.VAlign(VAlign_Center)
+											[
+												SNew(STextBlock)
+													.Text(LOCTEXT("InputsHeader", "Inputs"))
+													.Font(FAppStyle::GetFontStyle("BoldFont"))
+											]
+
+											+ SHorizontalBox::Slot()
+											.AutoWidth()
+											.Padding(2.0f, 0.0f)
+											[
+												SNew(SButton)
+													.ButtonStyle(FAppStyle::Get(), "SimpleButton")
+													.ToolTipText(LOCTEXT("ShowAllTooltip", "Show All"))
+													.OnClicked_Lambda([this]()
+														{
+															bHasExplicitPins = true;
+															PinnedInputNames.Empty();
+															for (const FInputBrowserColumnDef& ColDef : AllColumns)
+															{
+																PinnedInputNames.Add(ColDef.InputName);
+															}
+															ApplyFilters();
+															RebuildColumns();
+															if (ListView.IsValid()) { ListView->RequestListRefresh(); }
+															if (PinListView.IsValid()) { PinListView->RequestListRefresh(); }
+															return FReply::Handled();
+														})
+													[
+														SNew(SImage)
+															.Image(FAppStyle::GetBrush("Icons.Pinned"))
+															.ColorAndOpacity(FSlateColor::UseForeground())
+													]
+											]
+
+										+ SHorizontalBox::Slot()
+											.AutoWidth()
+											.Padding(2.0f, 0.0f)
+											[
+												SNew(SButton)
+													.ButtonStyle(FAppStyle::Get(), "SimpleButton")
+													.ToolTipText(LOCTEXT("HideAllTooltip", "Hide All"))
+													.OnClicked_Lambda([this]()
+														{
+															bHasExplicitPins = true;
+															PinnedInputNames.Empty();
+															ApplyFilters();
+															RebuildColumns();
+															if (ListView.IsValid()) { ListView->RequestListRefresh(); }
+															if (PinListView.IsValid()) { PinListView->RequestListRefresh(); }
+															return FReply::Handled();
+														})
+													[
+														SNew(SImage)
+															.Image(FAppStyle::GetBrush("Icons.Unpinned"))
+															.ColorAndOpacity(FSlateColor::UseForeground())
+													]
+											]
+									]
+
+								+ SVerticalBox::Slot()
+									.FillHeight(1.0f)
+									[
+										SAssignNew(PinListView, SListView<TSharedPtr<FInputBrowserColumnDef>>)
+											.ListItemsSource(&FilteredPinListItems)
+											.OnGenerateRow(this, &SMetasoundInputBrowser::OnGeneratePinRow)
+											.SelectionMode(ESelectionMode::None)
+									]
+							]
+					]
+			];
+
+		PostUndoRedoHandle = FEditorDelegates::PostUndoRedo.AddSP(this, &SMetasoundInputBrowser::RefreshAllRows);
+	}
+
+	SMetasoundInputBrowser::~SMetasoundInputBrowser()
+	{
+		FEditorDelegates::PostUndoRedo.Remove(PostUndoRedoHandle);
+		UnsubscribeAllBuilders();
+	}
+
+	FReply SMetasoundInputBrowser::OnFocusReceived(const FGeometry& MyGeometry, const FFocusEvent& InFocusEvent)
+	{
+		RefreshAllRows();
+		ScanAndSubscribeBuilders();
+		return SCompoundWidget::OnFocusReceived(MyGeometry, InFocusEvent);
+	}
+
+	void SMetasoundInputBrowser::SetAssets(const TArray<TWeakObjectPtr<UObject>>& InAssets)
+	{
+		// Unsubscribe and resubscribe to builders, because the delegate lambdas refer to row data
+		//  which will no longer be valid after rebuilding.
+		UnsubscribeAllBuilders();
+
+		SourceAssets = InAssets;
+		RebuildData();
+	}
+
+	void SMetasoundInputBrowser::RebuildData()
+	{
+		Rows.Empty();
+		AllColumns.Empty();
+		PinnedInputNames.Empty();
+		bHasExplicitPins = false;
+
+		for (const TWeakObjectPtr<UObject>& WeakAsset : SourceAssets)
+		{
+			UObject* Asset = WeakAsset.Get();
+			if (!Asset)
+			{
+				continue;
+			}
+
+			TSharedPtr<FInputBrowserRow> Row;
+			TArray<FInputBrowserColumnDef> AssetColumns = FInputBrowserRow::BuildRow(Row, Asset);
+
+			if (Row.IsValid())
+			{
+				Rows.Add(Row);
+
+				// Merge into AllColumns (union)
+				for (const FInputBrowserColumnDef& ColDef : AssetColumns)
+				{
+					AllColumns.AddUnique(ColDef);
+				}
+			}
+		}
+
+		// Sort columns alphabetically
+		AllColumns.Sort([](const FInputBrowserColumnDef& A, const FInputBrowserColumnDef& B)
+			{
+				return A.InputName.LexicalLess(B.InputName);
+			});
+
+		ApplyFilters();
+		RebuildColumns();
+
+		if (ListView.IsValid())
+		{
+			ListView->RebuildList();
+		}
+		if (PinListView.IsValid())
+		{
+			PinListView->RequestListRefresh();
+		}
+
+		ScanAndSubscribeBuilders();
+	}
+
+	void SMetasoundInputBrowser::ApplyFilters()
+	{
+		FilteredColumns.Empty();
+		FilteredPinListItems.Empty();
+
+		for (const FInputBrowserColumnDef& ColDef : AllColumns)
+		{
+			// Name filter
+			if (!NameFilter.IsEmpty())
+			{
+				if (!ColDef.InputName.ToString().Contains(NameFilter))
+				{
+					continue;
+				}
+			}
+
+			// Passes name filter - add to pin list (always)
+			FilteredPinListItems.Add(MakeShared<FInputBrowserColumnDef>(ColDef));
+
+			// Pin filter - only add to table columns if pinned (or no explicit pins yet)
+			if (!bHasExplicitPins || PinnedInputNames.Contains(ColDef.InputName))
+			{
+				FilteredColumns.Add(ColDef);
+			}
+		}
+
+		// Update per-row visibility and sort rows without visible inputs to the bottom
+		for (TSharedPtr<FInputBrowserRow>& Row : Rows)
+		{
+			bool bHasAny = false;
+			for (const FInputBrowserColumnDef& ColDef : FilteredColumns)
+			{
+				const FInputBrowserCellData* Cell = Row->Cells.Find(ColDef.InputName);
+				if (Cell && Cell->bInputExists)
+				{
+					bHasAny = true;
+					break;
+				}
+			}
+			Row->bHasVisibleInputs = bHasAny;
+		}
+
+		Rows.StableSort([](const TSharedPtr<FInputBrowserRow>& A, const TSharedPtr<FInputBrowserRow>& B)
+			{
+				if (A->bHasVisibleInputs != B->bHasVisibleInputs)
+				{
+					return A->bHasVisibleInputs; // rows with visible inputs first
+				}
+				return false; // preserve original order otherwise
+			});
+	}
+
+	void SMetasoundInputBrowser::RebuildColumns()
+	{
+		if (!HeaderRow.IsValid())
+		{
+			return;
+		}
+
+		// Build all column args up front, then apply in one batch
+		TArray<SHeaderRow::FColumn::FArguments> ColumnArgs;
+
+		// Asset name column (always first)
+		ColumnArgs.Add(
+			SHeaderRow::Column(InputBrowserPrivate::Column_AssetName)
+			.DefaultLabel(LOCTEXT("AssetNameColumn", "Asset"))
+			.ManualWidth(200.0f)
+			.ShouldGenerateWidget(true)
+		);
+
+		// One column per filtered input
+		for (const FInputBrowserColumnDef& ColDef : FilteredColumns)
+		{
+			FText TypeSuffix = FText::FromName(ColDef.TypeName);
+			FText Label = FText::Format(LOCTEXT("InputColumnLabelFmt", "{0}"), FText::FromName(ColDef.InputName));
+			FText Tooltip = FText::Format(LOCTEXT("InputColumnTooltipFmt", "{0} ({1})"), FText::FromName(ColDef.InputName), TypeSuffix);
+
+			ColumnArgs.Add(
+				SHeaderRow::Column(ColDef.InputName)
+				.DefaultLabel(Label)
+				.DefaultTooltip(Tooltip)
+				.ManualWidth(140.0f)
+				.ShouldGenerateWidget(true)
+			);
+		}
+
+		HeaderRow->ClearColumns();
+		HeaderRow->AddColumns(ColumnArgs);
+	}
+
+	TSharedRef<ITableRow> SMetasoundInputBrowser::OnGenerateRow(
+		TSharedPtr<FInputBrowserRow> InRow,
+		const TSharedRef<STableViewBase>& OwnerTable)
+	{
+		return SNew(SMetasoundInputBrowserTableRow, OwnerTable)
+			.RowData(InRow)
+			.Browser(this);
+	}
+
+	TSharedRef<SWidget> SMetasoundInputBrowser::GenerateCellWidget(
+		const TSharedRef<FInputBrowserRow>& InRow,
+		const FInputBrowserColumnDef& InColumnDef)
+	{
+		FInputBrowserCellData* CellData = InRow->Cells.Find(InColumnDef.InputName);
+
+		if (!CellData || !CellData->bInputExists)
+		{
+
+			return SNew(SHorizontalBox)
+
+				+ SHorizontalBox::Slot()
+				.FillWidth(1.0f)
+				.VAlign(VAlign_Center)
+				[
+					SNew(STextBlock)
+						.Text(LOCTEXT("NAText", "\u2014"))
+						.ColorAndOpacity(FSlateColor::UseSubduedForeground())
+						.Justification(ETextJustify::Center)
+				]
+
+				+ SHorizontalBox::Slot()
+				.AutoWidth()
+				.Padding(2.0f, 0.0f)
+				[
+					SNew(SBox).MinDesiredWidth(InputBrowserPrivate::ResetButtonWidth)
+				];
+		}
+
+		// Determine visual style for inherited preset values
+		const bool bIsInherited = CellData->bIsPreset && !CellData->bOverridesDefault;
+		FText InheritedTooltip = bIsInherited
+			? LOCTEXT("InheritedTooltip", "Inherited from parent preset (not overridden)")
+			: FText::GetEmpty();
+
+		TWeakPtr<FInputBrowserRow> WeakRow = InRow;
+		FName InputName = InColumnDef.InputName;
+
+		// Build the value widget based on type
+		TSharedRef<SWidget> ValueWidget = SNullWidget::NullWidget;
+
+		switch (CellData->Value.GetType())
+		{
+		case EMetasoundFrontendLiteralType::Boolean:
+		{
+			bool bCurrentValue = false;
+			CellData->Value.TryGet(bCurrentValue);
+
+			ValueWidget = SNew(SBox)
+				.HAlign(HAlign_Center)
+				.VAlign(VAlign_Center)
+				[
+					SNew(SCheckBox)
+						.IsChecked(bCurrentValue ? ECheckBoxState::Checked : ECheckBoxState::Unchecked)
+						.OnCheckStateChanged_Lambda([this, WeakRow, InputName](ECheckBoxState NewState)
+							{
+								if (TSharedPtr<FInputBrowserRow> Row = WeakRow.Pin())
+								{
+									FMetasoundFrontendLiteral NewLiteral;
+									NewLiteral.Set(NewState == ECheckBoxState::Checked);
+									CommitCellValue(Row.ToSharedRef(), InputName, NewLiteral);
+								}
+							})
+						.ToolTipText(InheritedTooltip)
+				];
+			break;
+		}
+
+		case EMetasoundFrontendLiteralType::Float:
+		{
+			float CurrentValue = 0.0f;
+			CellData->Value.TryGet(CurrentValue);
+
+			ValueWidget = SNew(SSpinBox<float>)
+				.Value(CurrentValue)
+				.OnValueCommitted_Lambda([this, WeakRow, InputName](float NewValue, ETextCommit::Type)
+					{
+						if (TSharedPtr<FInputBrowserRow> Row = WeakRow.Pin())
+						{
+							FMetasoundFrontendLiteral NewLiteral;
+							NewLiteral.Set(NewValue);
+							CommitCellValue(Row.ToSharedRef(), InputName, NewLiteral);
+						}
+					})
+				.ToolTipText(InheritedTooltip);
+			break;
+		}
+
+		case EMetasoundFrontendLiteralType::Integer:
+		{
+			int32 CurrentValue = 0;
+			CellData->Value.TryGet(CurrentValue);
+
+			ValueWidget = SNew(SSpinBox<int32>)
+				.Value(CurrentValue)
+				.OnValueCommitted_Lambda([this, WeakRow, InputName](int32 NewValue, ETextCommit::Type)
+					{
+						if (TSharedPtr<FInputBrowserRow> Row = WeakRow.Pin())
+						{
+							FMetasoundFrontendLiteral NewLiteral;
+							NewLiteral.Set(NewValue);
+							CommitCellValue(Row.ToSharedRef(), InputName, NewLiteral);
+						}
+					})
+				.ToolTipText(InheritedTooltip);
+			break;
+		}
+
+		case EMetasoundFrontendLiteralType::String:
+		{
+			FString CurrentValue;
+			CellData->Value.TryGet(CurrentValue);
+
+			ValueWidget = SNew(SEditableTextBox)
+				.Text(FText::FromString(CurrentValue))
+				.OnTextCommitted_Lambda([this, WeakRow, InputName](const FText& NewText, ETextCommit::Type)
+					{
+						if (TSharedPtr<FInputBrowserRow> Row = WeakRow.Pin())
+						{
+							FMetasoundFrontendLiteral NewLiteral;
+							NewLiteral.Set(NewText.ToString());
+							CommitCellValue(Row.ToSharedRef(), InputName, NewLiteral);
+						}
+					})
+				.ToolTipText(InheritedTooltip);
+			break;
+		}
+
+		default:
+		{
+			// For arrays, UObjects, and other complex types: display read-only summary
+			FText FullText = FText::FromString(CellData->Value.ToString());
+			FText DisplayText;
+
+			if (CellData->Value.IsArray())
+			{
+				DisplayText = FText::Format(LOCTEXT("ArrayElementCountFmt", "({0} element array)"), CellData->Value.GetArrayNum());
+			}
+			else if (CellData->Value.GetType() == EMetasoundFrontendLiteralType::UObject)
+			{
+				UObject* Object = nullptr;
+				CellData->Value.TryGet(Object);
+				DisplayText = Object
+					? FText::FromString(Object->GetName())
+					: LOCTEXT("NoneObject", "(None)");
+			}
+			else
+			{
+				DisplayText = FullText;
+			}
+
+			ValueWidget = SNew(STextBlock)
+				.Text(DisplayText)
+				.ToolTipText(FullText);
+			break;
+		}
+		}
+
+		// Non-preset cells: value widget + right spacer for alignment
+		if (!CellData->bIsPreset)
+		{
+			return SNew(SHorizontalBox)
+
+				+ SHorizontalBox::Slot()
+				.FillWidth(1.0f)
+				.VAlign(VAlign_Center)
+				[
+					ValueWidget
+				]
+
+				+ SHorizontalBox::Slot()
+				.AutoWidth()
+				.Padding(2.0f, 0.0f)
+				[
+					SNew(SBox).MinDesiredWidth(InputBrowserPrivate::ResetButtonWidth)
+				];
+		}
+
+		// Preset cells: disable the value widget when inheriting, and show either
+		// a pencil icon (to begin overriding) or a reset icon (to revert to inherited).
+		return SNew(SHorizontalBox)
+
+			+ SHorizontalBox::Slot()
+			.FillWidth(1.0f)
+			.VAlign(VAlign_Center)
+			[
+				SNew(SBox)
+					.IsEnabled_Lambda([WeakRow, InputName]()
+						{
+							if (TSharedPtr<FInputBrowserRow> Row = WeakRow.Pin())
+							{
+								if (const FInputBrowserCellData* Cell = Row->Cells.Find(InputName))
+								{
+									return Cell->bOverridesDefault;
+								}
+							}
+							return false;
+						})
+					[
+						ValueWidget
+					]
+			]
+
+		+ SHorizontalBox::Slot()
+			.AutoWidth()
+			.VAlign(VAlign_Center)
+			.Padding(2.0f, 0.0f)
+			[
+				// Pencil icon: visible when inheriting, click to begin overriding
+				SNew(SButton)
+					.ButtonStyle(FAppStyle::Get(), "SimpleButton")
+					.ToolTipText(LOCTEXT("BeginOverrideTooltip", "Override inherited default"))
+					.Visibility_Lambda([WeakRow, InputName]()
+						{
+							if (TSharedPtr<FInputBrowserRow> Row = WeakRow.Pin())
+							{
+								if (const FInputBrowserCellData* Cell = Row->Cells.Find(InputName))
+								{
+									return Cell->bOverridesDefault ? EVisibility::Collapsed : EVisibility::Visible;
+								}
+							}
+							return EVisibility::Collapsed;
+						})
+					.OnClicked_Lambda([this, WeakRow, InputName]()
+						{
+							if (TSharedPtr<FInputBrowserRow> Row = WeakRow.Pin())
+							{
+								BeginOverride(Row.ToSharedRef(), InputName);
+							}
+							return FReply::Handled();
+						})
+					.ContentPadding(0)
+					[
+						SNew(SImage)
+							.Image(FAppStyle::GetBrush("Icons.Edit"))
+							.ColorAndOpacity(FSlateColor::UseForeground())
+					]
+			]
+
+		+ SHorizontalBox::Slot()
+			.AutoWidth()
+			.VAlign(VAlign_Center)
+			.Padding(2.0f, 0.0f)
+			[
+				// Reset icon: visible when overriding, click to revert to inherited
+				SNew(SButton)
+					.ButtonStyle(FAppStyle::Get(), "SimpleButton")
+					.ToolTipText(LOCTEXT("ResetToInheritedTooltip", "Reset to inherited default from parent preset"))
+					.Visibility_Lambda([WeakRow, InputName]()
+						{
+							if (TSharedPtr<FInputBrowserRow> Row = WeakRow.Pin())
+							{
+								if (const FInputBrowserCellData* Cell = Row->Cells.Find(InputName))
+								{
+									return Cell->bOverridesDefault ? EVisibility::Visible : EVisibility::Collapsed;
+								}
+							}
+							return EVisibility::Collapsed;
+						})
+					.OnClicked_Lambda([this, WeakRow, InputName]()
+						{
+							if (TSharedPtr<FInputBrowserRow> Row = WeakRow.Pin())
+							{
+								ResetToDefault(Row.ToSharedRef(), InputName);
+							}
+							return FReply::Handled();
+						})
+					.ContentPadding(0)
+					[
+						SNew(SImage)
+							.Image(FAppStyle::GetBrush("PropertyWindow.DiffersFromDefault"))
+							.ColorAndOpacity(FSlateColor::UseForeground())
+					]
+			];
+	}
+
+	// Apply a new value to the asset
+	void SMetasoundInputBrowser::CommitCellValue(
+		const TSharedRef<FInputBrowserRow>& InRow,
+		FName InputName,
+		const FMetasoundFrontendLiteral& NewValue)
+	{
+		using namespace Metasound::Frontend;
+
+		IMetaSoundDocumentInterface* DocInterface = InRow->GetDocumentInterface();
+		if (!DocInterface)
+		{
+			return;
+		}
+
+		UObject* Asset = InRow->Asset.Get();
+
+		FScopedTransaction Transaction(LOCTEXT("EditMetaSoundInput", "Edit MetaSound Input Value"));
+		Asset->Modify();
+
+		IDocumentBuilderRegistry& Registry = IDocumentBuilderRegistry::GetChecked();
+		FMetaSoundFrontendDocumentBuilder& Builder = Registry.FindOrBeginBuilding(TScriptInterface<IMetaSoundDocumentInterface>(Asset));
+		SubscribeToBuilder(Builder, InRow);
+		Builder.SetGraphInputDefault(InputName, NewValue, &DefaultPageID);
+
+		RefreshRow(InRow);
+
+		if (ListView.IsValid())
+		{
+			ListView->RebuildList();
+		}
+	}
+
+	void SMetasoundInputBrowser::ResetToDefault(
+		const TSharedRef<FInputBrowserRow>& InRow,
+		FName InputName)
+	{
+		using namespace Metasound::Frontend;
+
+		IMetaSoundDocumentInterface* DocInterface = InRow->GetDocumentInterface();
+		if (!DocInterface)
+		{
+			return;
+		}
+
+		UObject* Asset = InRow->Asset.Get();
+
+		IDocumentBuilderRegistry& Registry = IDocumentBuilderRegistry::GetChecked();
+		FMetaSoundFrontendDocumentBuilder& Builder = Registry.FindOrBeginBuilding(TScriptInterface<IMetaSoundDocumentInterface>(Asset));
+		SubscribeToBuilder(Builder, InRow);
+
+		if (InputBrowserPrivate::HasNonDefaultPageOverrides(Builder, InputName))
+		{
+			EAppReturnType::Type Result = FMessageDialog::Open(
+				EAppMsgType::YesNo,
+				LOCTEXT("ResetPagedDefaultsWarning",
+					"This input has overridden values on non-default pages that will be lost.\n\n"
+					"Are you sure you want to reset all pages to inherited default values?"),
+				LOCTEXT("ResetPagedDefaultsTitle", "Reset Input Default"));
+
+			if (Result != EAppReturnType::Yes)
+			{
+				return;
+			}
+		}
+
+		FScopedTransaction Transaction(LOCTEXT("ResetMetaSoundInput", "Reset MetaSound Input to Inherited Default"));
+		Asset->Modify();
+
+		Builder.ResetGraphInputDefault(InputName);
+
+		// ResetGraphInputDefault zeros the literal. Re-run ConfigureDocument so
+		// the inherited value gets repopulated from the parent.
+		Builder.ConfigureDocument();
+
+		RefreshRow(InRow);
+
+		if (ListView.IsValid())
+		{
+			ListView->RebuildList();
+		}
+	}
+
+	void SMetasoundInputBrowser::BeginOverride(
+		const TSharedRef<FInputBrowserRow>& InRow,
+		FName InputName)
+	{
+		FInputBrowserCellData* CellData = InRow->Cells.Find(InputName);
+		if (!CellData || !CellData->bIsPreset)
+		{
+			return;
+		}
+
+		// Read the current (inherited) value from the cell, then write it back
+		// as an explicit override via SetGraphInputDefault, which will set
+		// bOverrideInheritedDefault to true.
+		CommitCellValue(InRow, InputName, CellData->Value);
+	}
+
+	TSharedPtr<SWidget> SMetasoundInputBrowser::OnAssetContextMenuOpening()
+	{
+		if (!ListView.IsValid())
+		{
+			return nullptr;
+		}
+
+		TArray<TSharedPtr<FInputBrowserRow>> SelectedItems = ListView->GetSelectedItems();
+		if (SelectedItems.IsEmpty())
+		{
+			return nullptr;
+		}
+
+		TSharedPtr<FInputBrowserRow> SelectedRow = SelectedItems[0];
+		UObject* Asset = SelectedRow->Asset.Get();
+		if (!Asset)
+		{
+			return nullptr;
+		}
+
+		FMenuBuilder MenuBuilder(true, nullptr);
+
+		FAssetData AssetData(Asset);
+
+		MenuBuilder.BeginSection("AssetActions", LOCTEXT("AssetActionsSection", "Asset Actions"));
+
+		TWeakObjectPtr<UObject> WeakAsset = Asset;
+		MenuBuilder.AddMenuEntry(
+			LOCTEXT("EditAsset", "Edit Asset"),
+			LOCTEXT("EditAssetTooltip", "Open this asset in its editor"),
+			FSlateIcon(FAppStyle::GetAppStyleSetName(), "Icons.Edit"),
+			FUIAction(FExecuteAction::CreateLambda([WeakAsset]()
+				{
+					if (UObject* AssetPtr = WeakAsset.Get())
+					{
+						if (GEditor)
+						{
+							GEditor->GetEditorSubsystem<UAssetEditorSubsystem>()->OpenEditorForAsset(AssetPtr);
+						}
+					}
+				}))
+		);
+
+		MenuBuilder.AddMenuEntry(
+			LOCTEXT("BrowseToAsset", "Browse to Asset"),
+			LOCTEXT("BrowseToAssetTooltip", "Show this asset in the Content Browser"),
+			FSlateIcon(FAppStyle::GetAppStyleSetName(), "Icons.Search"),
+			FUIAction(FExecuteAction::CreateLambda([AssetData]()
+				{
+					FContentBrowserModule& ContentBrowserModule = FModuleManager::LoadModuleChecked<FContentBrowserModule>("ContentBrowser");
+					ContentBrowserModule.Get().SyncBrowserToAssets(TArray<FAssetData>{ AssetData });
+				}))
+		);
+
+		MenuBuilder.AddMenuEntry(
+			LOCTEXT("CopyAssetReference", "Copy Reference"),
+			LOCTEXT("CopyAssetReferenceTooltip", "Copy the asset reference path to the clipboard"),
+			FSlateIcon(FAppStyle::GetAppStyleSetName(), "GenericCommands.Copy"),
+			FUIAction(FExecuteAction::CreateLambda([AssetData]()
+				{
+					FPlatformApplicationMisc::ClipboardCopy(*AssetData.GetExportTextName());
+				}))
+		);
+
+		MenuBuilder.EndSection();
+
+		return MenuBuilder.MakeWidget();
+	}
+
+	void SMetasoundInputBrowser::OnNameFilterChanged(const FText& InText)
+	{
+		NameFilter = InText.ToString();
+		ApplyFilters();
+		RebuildColumns();
+
+		if (ListView.IsValid())
+		{
+			ListView->RequestListRefresh();
+		}
+		if (PinListView.IsValid())
+		{
+			PinListView->RequestListRefresh();
+		}
+	}
+
+	// ---------------------------------------------------------------
+	// Auto-refresh
+	// ---------------------------------------------------------------
+
+	void SMetasoundInputBrowser::RefreshAllRows()
+	{
+		for (const TSharedPtr<FInputBrowserRow>& Row : Rows)
+		{
+			if (Row.IsValid())
+			{
+				RefreshRow(Row.ToSharedRef());
+			}
+		}
+
+		if (ListView.IsValid())
+		{
+			ListView->RebuildList();
+		}
+	}
+
+	void SMetasoundInputBrowser::SubscribeToBuilder(FMetaSoundFrontendDocumentBuilder& Builder, const TSharedRef<FInputBrowserRow>& InRow)
+	{
+		if (SubscribedBuilders.Contains(&Builder))
+		{
+			return;
+		}
+
+		TWeakPtr<FInputBrowserRow> WeakRow = InRow;
+		TWeakPtr<SMetasoundInputBrowser> WeakThis = SharedThis(this);
+
+		FDelegateHandle Handle = Builder.GetDocumentDelegates().InterfaceDelegates.OnInputDefaultChanged.AddLambda(
+			[WeakThis, WeakRow](const Metasound::Frontend::FDocumentArrayPagedInputArgs&)
+			{
+				TSharedPtr<SMetasoundInputBrowser> This = WeakThis.Pin();
+				TSharedPtr<FInputBrowserRow> Row = WeakRow.Pin();
+				if (This && Row.IsValid())
+				{
+					This->RefreshRow(Row.ToSharedRef());
+					if (This->ListView.IsValid())
+					{
+						This->ListView->RebuildList();
+					}
+				}
+			});
+
+		SubscribedBuilders.Add(&Builder, Handle);
+	}
+
+	void SMetasoundInputBrowser::ScanAndSubscribeBuilders()
+	{
+		using namespace Metasound::Frontend;
+
+		IDocumentBuilderRegistry& Registry = IDocumentBuilderRegistry::GetChecked();
+
+		for (const TSharedPtr<FInputBrowserRow>& Row : Rows)
+		{
+			if (!Row.IsValid())
+			{
+				continue;
+			}
+
+			UObject* Asset = Row->Asset.Get();
+			if (!Asset)
+			{
+				continue;
+			}
+
+			FMetaSoundFrontendDocumentBuilder* Builder = Registry.FindBuilder(TScriptInterface<IMetaSoundDocumentInterface>(Asset));
+			if (Builder)
+			{
+				SubscribeToBuilder(*Builder, Row.ToSharedRef());
+			}
+		}
+	}
+
+	void SMetasoundInputBrowser::UnsubscribeAllBuilders()
+	{
+		for (auto& [Builder, Handle] : SubscribedBuilders)
+		{
+			if (Builder)
+			{
+				Builder->GetDocumentDelegates().InterfaceDelegates.OnInputDefaultChanged.Remove(Handle);
+			}
+		}
+		SubscribedBuilders.Empty();
+	}
+
+	void SMetasoundInputBrowser::RefreshRow(const TSharedRef<FInputBrowserRow>& InRow)
+	{
+		using namespace Metasound::Frontend;
+
+		IMetaSoundDocumentInterface* DocInterface = InRow->GetDocumentInterface();
+		if (!DocInterface)
+		{
+			return;
+		}
+
+		const FMetasoundFrontendDocument& Document = DocInterface->GetConstDocument();
+		const FMetasoundFrontendClassInterface& Interface = Document.RootGraph.GetDefaultInterface();
+
+		const FMetaSoundFrontendPresetTemplate* PresetTemplate = Document.Template.GetPtr<FMetaSoundFrontendPresetTemplate>();
+
+		// Update cell data in place (or create new entries)
+		for (const FMetasoundFrontendClassInput& Input : Interface.Inputs)
+		{
+			if (Input.TypeName == InputBrowserPrivate::TriggerTypeName || Input.TypeName == InputBrowserPrivate::AudioTypeName)
+			{
+				continue;
+			}
+
+			FInputBrowserRow::ReadCellData(InRow->Cells.FindOrAdd(Input.Name), Input, PresetTemplate);
+		}
+	}
+
+	// ---------------------------------------------------------------
+	// Pin list
+	// ---------------------------------------------------------------
+
+	void SMetasoundInputBrowser::TogglePin(FName InputName)
+	{
+		if (!bHasExplicitPins)
+		{
+			// First pin toggle: initialize all current columns as pinned
+			bHasExplicitPins = true;
+			for (const FInputBrowserColumnDef& ColDef : AllColumns)
+			{
+				PinnedInputNames.Add(ColDef.InputName);
+			}
+		}
+
+		if (PinnedInputNames.Contains(InputName))
+		{
+			PinnedInputNames.Remove(InputName);
+		}
+		else
+		{
+			PinnedInputNames.Add(InputName);
+		}
+
+		ApplyFilters();
+		RebuildColumns();
+
+		if (ListView.IsValid())
+		{
+			ListView->RequestListRefresh();
+		}
+		if (PinListView.IsValid())
+		{
+			PinListView->RequestListRefresh();
+		}
+	}
+
+	bool SMetasoundInputBrowser::IsPinned(FName InputName) const
+	{
+		return !bHasExplicitPins || PinnedInputNames.Contains(InputName);
+	}
+
+	TSharedRef<ITableRow> SMetasoundInputBrowser::OnGeneratePinRow(
+		TSharedPtr<FInputBrowserColumnDef> InItem,
+		const TSharedRef<STableViewBase>& OwnerTable)
+	{
+		const FName InputName = InItem->InputName;
+		const FName TypeName = InItem->TypeName;
+		const FLinearColor TypeColor = InputBrowserPrivate::GetColorForTypeName(TypeName);
+
+		return SNew(STableRow<TSharedPtr<FInputBrowserColumnDef>>, OwnerTable)
+			[
+				SNew(SHorizontalBox)
+
+					+ SHorizontalBox::Slot()
+					.AutoWidth()
+					.VAlign(VAlign_Center)
+					.Padding(2.0f, 0.0f)
+					[
+						SNew(SButton)
+							.ButtonStyle(FAppStyle::Get(), "SimpleButton")
+							.OnClicked_Lambda([this, InputName]()
+								{
+									TogglePin(InputName);
+									return FReply::Handled();
+								})
+							.ContentPadding(0)
+							[
+								SNew(SImage)
+									.Image_Lambda([this, InputName]()
+										{
+											return IsPinned(InputName)
+												? FAppStyle::GetBrush("Icons.Pinned")
+												: FAppStyle::GetBrush("Icons.Unpinned");
+										})
+									.ColorAndOpacity(FSlateColor::UseForeground())
+							]
+					]
+
+				+ SHorizontalBox::Slot()
+					.AutoWidth()
+					.VAlign(VAlign_Center)
+					.Padding(2.0f, 0.0f, 0.0f, 0.0f)
+					[
+						SNew(SImage)
+							.Image(InputBrowserPrivate::GetIconForTypeName(TypeName))
+							.ColorAndOpacity(TypeColor)
+							.DesiredSizeOverride(FVector2D(12.0f, 12.0f))
+					]
+
+					+ SHorizontalBox::Slot()
+					.FillWidth(1.0f)
+					.VAlign(VAlign_Center)
+					.Padding(2.0f, 0.0f)
+					[
+						SNew(STextBlock)
+							.Text(FText::FromName(InputName))
+					]
+
+					+ SHorizontalBox::Slot()
+					.AutoWidth()
+					.VAlign(VAlign_Center)
+					.Padding(4.0f, 0.0f)
+					[
+						SNew(STextBlock)
+							.Text(FText::FromName(TypeName))
+							.ColorAndOpacity(FSlateColor::UseSubduedForeground())
+					]
+			];
+	}
+
+	// ---------------------------------------------------------------
+	// SMetasoundInputBrowserTableRow
+	// ---------------------------------------------------------------
+
+	void SMetasoundInputBrowserTableRow::Construct(const FArguments& InArgs, const TSharedRef<STableViewBase>& OwnerTable)
+	{
+		RowData = InArgs._RowData;
+		Browser = InArgs._Browser;
+
+		SMultiColumnTableRow<TSharedPtr<FInputBrowserRow>>::Construct(
+			FSuperRowType::FArguments()
+			.Style(FAppStyle::Get(), "TableView.AlternatingRow"),
+			OwnerTable);
+	}
+
+	TSharedRef<SWidget> SMetasoundInputBrowserTableRow::GenerateWidgetForColumn(const FName& ColumnName)
+	{
+		if (!RowData.IsValid() || !Browser)
+		{
+			return SNullWidget::NullWidget;
+		}
+
+		TSharedRef<SWidget> CellContent = SNullWidget::NullWidget;
+
+		// Asset name column
+		if (ColumnName == InputBrowserPrivate::Column_AssetName)
+		{
+			CellContent = SNew(STextBlock)
+				.Text(FText::FromString(RowData->AssetName))
+				.Margin(FMargin(4.0f, 0.0f))
+				.ColorAndOpacity(RowData->bHasVisibleInputs ? FSlateColor::UseForeground() : FSlateColor::UseSubduedForeground());
+		}
+		else
+		{
+			// Find matching column def
+			for (const FInputBrowserColumnDef& ColDef : Browser->FilteredColumns)
+			{
+				if (ColDef.InputName == ColumnName)
+				{
+					CellContent = Browser->GenerateCellWidget(RowData.ToSharedRef(), ColDef);
+					break;
+				}
+			}
+		}
+
+		return SNew(SBox)
+			.MinDesiredHeight(20.0f)
+			.VAlign(VAlign_Center)
+			[
+				CellContent
+			];
+	}
+
+} // namespace Metasound::Editor
+
+#undef LOCTEXT_NAMESPACE
diff --git a/Engine/Plugins/Runtime/Metasound/Source/MetasoundEditor/Private/SMetasoundInputBrowser.h b/Engine/Plugins/Runtime/Metasound/Source/MetasoundEditor/Private/SMetasoundInputBrowser.h
new file mode 100644
--- /dev/null
+++ b/Engine/Plugins/Runtime/Metasound/Source/MetasoundEditor/Private/SMetasoundInputBrowser.h
@@ -0,0 +1,177 @@
+// Copyright Epic Games, Inc. All Rights Reserved.
+
+#pragma once
+
+#include "MetasoundFrontendLiteral.h"
+#include "Widgets/SCompoundWidget.h"
+#include "Widgets/Views/SListView.h"
+#include "Widgets/Views/STableRow.h"
+
+struct FMetasoundFrontendClassInput;
+struct FMetaSoundFrontendDocumentBuilder;
+struct FMetaSoundFrontendPresetTemplate;
+class IMetaSoundDocumentInterface;
+class SHeaderRow;
+
+namespace Metasound::Editor
+{
+
+	/** Describes a single column in the MetaSound Input Browser (one per unique input name across all loaded assets). */
+	struct FInputBrowserColumnDef
+	{
+		FName InputName;
+		FName TypeName;
+		EMetasoundFrontendLiteralType LiteralType = EMetasoundFrontendLiteralType::None;
+
+		bool operator==(const FInputBrowserColumnDef& Other) const
+		{
+			return InputName == Other.InputName;
+		}
+	};
+
+	/** Per-cell data for a single (asset, input) pair. */
+	struct FInputBrowserCellData
+	{
+		FMetasoundFrontendLiteral Value;
+		FGuid NodeID;
+		bool bIsPreset = false;
+		bool bOverridesDefault = false;
+		bool bInputExists = false;
+	};
+
+	/** One row in the Input Browser, representing a single MetaSound asset. */
+	struct FInputBrowserRow
+	{
+		TWeakObjectPtr<UObject> Asset;
+		FString AssetName;
+		TMap<FName, FInputBrowserCellData> Cells;
+		bool bHasVisibleInputs = true;
+
+		/** Returns the document interface for this row's asset, or null if invalid. */
+		IMetaSoundDocumentInterface* GetDocumentInterface() const;
+
+		/** Populate a cell's data from a class input and optional preset template. */
+		static void ReadCellData(
+			FInputBrowserCellData& OutCell,
+			const FMetasoundFrontendClassInput& Input,
+			const FMetaSoundFrontendPresetTemplate* PresetTemplate);
+
+		/** Populate this row's cell data from the given MetaSound asset. Returns the set of column defs found. */
+		static TArray<FInputBrowserColumnDef> BuildRow(TSharedPtr<FInputBrowserRow>& OutRow, UObject* InAsset);
+	};
+
+	/**
+	 * A bulk-editing tool for MetaSound input values, similar to the Property Matrix.
+	 * Displays a table with rows = assets and columns = inputs.
+	 */
+	class SMetasoundInputBrowser : public SCompoundWidget
+	{
+	public:
+		static const FName TabId;
+
+		SLATE_BEGIN_ARGS(SMetasoundInputBrowser) {}
+		SLATE_END_ARGS()
+
+		void Construct(const FArguments& InArgs);
+		virtual ~SMetasoundInputBrowser();
+
+		virtual FReply OnFocusReceived(const FGeometry& MyGeometry, const FFocusEvent& InFocusEvent) override;
+
+		/** Populate the browser with the given MetaSound assets. */
+		void SetAssets(const TArray<TWeakObjectPtr<UObject>>& InAssets);
+
+	private:
+		/** Rebuild rows and column union from the current asset set. */
+		void RebuildData();
+
+		/** Rebuild the SHeaderRow columns from FilteredColumns. */
+		void RebuildColumns();
+
+		/** Apply name/type filters to AllColumns to produce FilteredColumns. */
+		void ApplyFilters();
+
+		friend class SMetasoundInputBrowserTableRow;
+
+		/** SListView callbacks. */
+		TSharedRef<ITableRow> OnGenerateRow(TSharedPtr<FInputBrowserRow> InRow, const TSharedRef<STableViewBase>& OwnerTable);
+
+		/** Generate a cell widget for a given row and column. */
+		TSharedRef<SWidget> GenerateCellWidget(const TSharedRef<FInputBrowserRow>& InRow, const FInputBrowserColumnDef& InColumnDef);
+
+		/** Commit an edited value back to the asset. */
+		void CommitCellValue(const TSharedRef<FInputBrowserRow>& InRow, FName InputName, const FMetasoundFrontendLiteral& NewValue);
+
+		/** Reset a preset input to its inherited default value. */
+		void ResetToDefault(const TSharedRef<FInputBrowserRow>& InRow, FName InputName);
+
+		/** Begin overriding an inherited preset input (copies the inherited value as an explicit override). */
+		void BeginOverride(const TSharedRef<FInputBrowserRow>& InRow, FName InputName);
+
+		/** Filter callback. */
+		void OnNameFilterChanged(const FText& InText);
+
+		/** Context menu for right-clicking on an asset row. */
+		TSharedPtr<SWidget> OnAssetContextMenuOpening();
+
+		/** Refresh a single row's cell data from its asset. */
+		void RefreshRow(const TSharedRef<FInputBrowserRow>& InRow);
+
+		/** Refresh all rows from their assets. Called when the widget regains focus. */
+		void RefreshAllRows();
+
+		/** Subscribe to OnInputDefaultChanged on a builder if not already subscribed. */
+		void SubscribeToBuilder(FMetaSoundFrontendDocumentBuilder& Builder, const TSharedRef<FInputBrowserRow>& InRow);
+
+		/** Scan all rows for existing builders and subscribe to their delegates. */
+		void ScanAndSubscribeBuilders();
+
+		/** Unsubscribe from all builder delegates and clear the subscriptions map. */
+		void UnsubscribeAllBuilders();
+
+		/** Pin column callbacks. */
+		void TogglePin(FName InputName);
+		bool IsPinned(FName InputName) const;
+		TSharedRef<ITableRow> OnGeneratePinRow(TSharedPtr<FInputBrowserColumnDef> InItem, const TSharedRef<STableViewBase>& OwnerTable);
+
+		// Data
+		TArray<TWeakObjectPtr<UObject>> SourceAssets;
+		TArray<TSharedPtr<FInputBrowserRow>> Rows;
+		TArray<FInputBrowserColumnDef> AllColumns;
+		TArray<FInputBrowserColumnDef> FilteredColumns;
+
+		// Filter state
+		FString NameFilter;
+
+		// Pin state
+		TSet<FName> PinnedInputNames;
+		bool bHasExplicitPins = false; // When false, all columns are shown
+		TArray<TSharedPtr<FInputBrowserColumnDef>> FilteredPinListItems;
+
+		// Builders we've subscribed to, mapped to delegate handles for cleanup
+		TMap<FMetaSoundFrontendDocumentBuilder*, FDelegateHandle> SubscribedBuilders;
+		FDelegateHandle PostUndoRedoHandle;
+
+		// Widgets
+		TSharedPtr<SHeaderRow> HeaderRow;
+		TSharedPtr<SListView<TSharedPtr<FInputBrowserRow>>> ListView;
+		TSharedPtr<SListView<TSharedPtr<FInputBrowserColumnDef>>> PinListView;
+	};
+
+	/** Multi-column table row for the Input Browser. */
+	class SMetasoundInputBrowserTableRow : public SMultiColumnTableRow<TSharedPtr<FInputBrowserRow>>
+	{
+	public:
+		SLATE_BEGIN_ARGS(SMetasoundInputBrowserTableRow) {}
+			SLATE_ARGUMENT(TSharedPtr<FInputBrowserRow>, RowData)
+			SLATE_ARGUMENT(SMetasoundInputBrowser*, Browser)
+		SLATE_END_ARGS()
+
+		void Construct(const FArguments& InArgs, const TSharedRef<STableViewBase>& OwnerTable);
+		virtual TSharedRef<SWidget> GenerateWidgetForColumn(const FName& ColumnName) override;
+
+	private:
+		TSharedPtr<FInputBrowserRow> RowData;
+		SMetasoundInputBrowser* Browser = nullptr;
+	};
+
+} // namespace Metasound::Editor
