MetaSound: Consolidate pin type registration and associated pin-related MetaSound Editor behavior into IMetaSoundGraphPanelPinFactory, exposing registration calls to enable 3rd party developers to customize pins further.

[RN] Exposed registration calls to enable further customization of MetaSound pins and colors (icon, color, etc.)

diff --git a/Engine/Plugins/Experimental/MetasoundExperimental/Source/MetasoundExperimentalEditor/Private/MetasoundExperimentalEditorModule.cpp b/Engine/Plugins/Experimental/MetasoundExperimental/Source/MetasoundExperimentalEditor/Private/MetasoundExperimentalEditorModule.cpp
--- a/Engine/Plugins/Experimental/MetasoundExperimental/Source/MetasoundExperimentalEditor/Private/MetasoundExperimentalEditorModule.cpp
+++ b/Engine/Plugins/Experimental/MetasoundExperimental/Source/MetasoundExperimentalEditor/Private/MetasoundExperimentalEditorModule.cpp
@@ -1,5 +1,6 @@
 // Copyright Epic Games, Inc. All Rights Reserved.
 
+#include "IMetaSoundGraphPanelPinFactory.h"
 #include "MetasoundEditorModule.h"
 #include "MetasoundExampleNodeConfiguration.h"
 #include "MetasoundExampleNodeDetailsCustomization.h"
@@ -22,7 +23,8 @@ public:
 
 void FMetasoundExperimentalEditorModule::StartupModule()
 {
-	Metasound::Editor::IMetasoundEditorModule& MetasoundEditorModule = FModuleManager::GetModuleChecked<Metasound::Editor::IMetasoundEditorModule>("MetasoundEditor");
+	using namespace Metasound::Editor;
+	IMetasoundEditorModule& MetasoundEditorModule = FModuleManager::GetModuleChecked<IMetasoundEditorModule>(IMetasoundEditorModule::ModuleName);
 	MetasoundEditorModule.RegisterCustomNodeConfigurationDetailsCustomization(
 		FMetaSoundWidgetExampleNodeConfiguration::StaticStruct()->GetFName(),
 		[](TSharedPtr<IPropertyHandle> InStructProperty, TWeakObjectPtr<UMetasoundEditorGraphNode> InNode)
@@ -35,21 +37,28 @@ void FMetasoundExperimentalEditorModule::StartupModule()
 		{
 			return MakeShareable(new FMappingFunctionNodeConfigurationCustomization(InStructProperty, InNode));
 		});
-	MetasoundEditorModule.RegisterCustomNodeConfigurationDetailsCustomization(		
+	MetasoundEditorModule.RegisterCustomNodeConfigurationDetailsCustomization(
 		FMetaSoundGranulatorNodeConfiguration::StaticStruct()->GetFName(),
 		[](TSharedPtr<IPropertyHandle> InStructProperty, TWeakObjectPtr<UMetasoundEditorGraphNode> InNode)
 		{
 			return MakeShareable(new FGranularNodeConfigurationCustomization(InStructProperty, InNode));
 		});
-	MetasoundEditorModule.RegisterPinType("Enum:CurveFunction", "Int32");
+	TSharedRef<IMetaSoundGraphPanelPinFactory> PinFactory = MetasoundEditorModule.GetGraphPanelPinFactory();
+	PinFactory->RegisterPin("Enum:CurveFunction", FGraphPinParams { .PinCategory = "Int32" });
 }
 
 void FMetasoundExperimentalEditorModule::ShutdownModule()
 {
-	Metasound::Editor::IMetasoundEditorModule& MetasoundEditorModule = FModuleManager::GetModuleChecked<Metasound::Editor::IMetasoundEditorModule>("MetasoundEditor");
-	MetasoundEditorModule.UnregisterCustomNodeConfigurationDetailsCustomization(FMetaSoundWidgetExampleNodeConfiguration::StaticStruct()->GetFName());
-	MetasoundEditorModule.UnregisterCustomNodeConfigurationDetailsCustomization(FMetaSoundMappingFunctionNodeConfiguration::StaticStruct()->GetFName());
-	MetasoundEditorModule.UnregisterCustomNodeConfigurationDetailsCustomization(FMetaSoundGranulatorNodeConfiguration::StaticStruct()->GetFName());
+	using namespace Metasound::Editor;
+	if (FModuleManager::Get().IsModuleLoaded(IMetasoundEditorModule::ModuleName))
+	{
+		IMetasoundEditorModule& MetasoundEditorModule = FModuleManager::GetModuleChecked<IMetasoundEditorModule>(IMetasoundEditorModule::ModuleName);
+		MetasoundEditorModule.UnregisterCustomNodeConfigurationDetailsCustomization(FMetaSoundWidgetExampleNodeConfiguration::StaticStruct()->GetFName());
+		MetasoundEditorModule.UnregisterCustomNodeConfigurationDetailsCustomization(FMetaSoundMappingFunctionNodeConfiguration::StaticStruct()->GetFName());
+		MetasoundEditorModule.UnregisterCustomNodeConfigurationDetailsCustomization(FMetaSoundGranulatorNodeConfiguration::StaticStruct()->GetFName());
+		TSharedRef<IMetaSoundGraphPanelPinFactory> PinFactory = MetasoundEditorModule.GetGraphPanelPinFactory();
+		PinFactory->UnregisterPin("Enum:CurveFunction");
+	}
 }
 
 #undef LOCTEXT_NAMESPACE
diff --git a/Engine/Plugins/Experimental/TechAudioTools/Source/TechAudioToolsMetaSoundEditor/Private/ViewModels/MetaSoundEditorViewModelConversionFunctions.cpp b/Engine/Plugins/Experimental/TechAudioTools/Source/TechAudioToolsMetaSoundEditor/Private/ViewModels/MetaSoundEditorViewModelConversionFunctions.cpp
--- a/Engine/Plugins/Experimental/TechAudioTools/Source/TechAudioToolsMetaSoundEditor/Private/ViewModels/MetaSoundEditorViewModelConversionFunctions.cpp
+++ b/Engine/Plugins/Experimental/TechAudioTools/Source/TechAudioToolsMetaSoundEditor/Private/ViewModels/MetaSoundEditorViewModelConversionFunctions.cpp
@@ -3,6 +3,7 @@
 #include "ViewModels/MetaSoundEditorViewModelConversionFunctions.h"
 
 #include "Editor/EditorEngine.h"
+#include "IMetaSoundGraphPanelPinFactory.h"
 #include "MetasoundEditorGraphSchema.h"
 #include "MetasoundEditorModule.h"
 #include "Modules/ModuleManager.h"
@@ -30,8 +31,8 @@ FLinearColor UMetaSoundEditorViewModelConversionFunctions::GetMetaSoundDataTypeP
 {
 	using namespace Metasound::Editor;
 
-	const IMetasoundEditorModule& MetaSoundEditorModule = FModuleManager::GetModuleChecked<IMetasoundEditorModule>("MetaSoundEditor");
-	if (const FEdGraphPinType* PinType = MetaSoundEditorModule.FindPinType(DataType))
+	IMetasoundEditorModule& MetaSoundEditorModule = FModuleManager::GetModuleChecked<IMetasoundEditorModule>(IMetasoundEditorModule::ModuleName);
+	if (const FEdGraphPinType* PinType = MetaSoundEditorModule.GetGraphPanelPinFactory()->FindPinType(DataType))
 	{
 		if (const UMetasoundEditorGraphSchema* Schema = GetDefault<UMetasoundEditorGraphSchema>())
 		{
diff --git a/Engine/Plugins/Runtime/AudioModulation/Source/AudioModulationEditor/Private/AudioModulationEditor.cpp b/Engine/Plugins/Runtime/AudioModulation/Source/AudioModulationEditor/Private/AudioModulationEditor.cpp
--- a/Engine/Plugins/Runtime/AudioModulation/Source/AudioModulationEditor/Private/AudioModulationEditor.cpp
+++ b/Engine/Plugins/Runtime/AudioModulation/Source/AudioModulationEditor/Private/AudioModulationEditor.cpp
@@ -29,6 +29,7 @@
 #include "AudioControlBusMixTrackEditor.h"
 
 #if WITH_AUDIOMODULATION_METASOUND_SUPPORT
+#include "IMetaSoundGraphPanelPinFactory.h"
 #include "MetasoundDataReference.h"
 #include "MetasoundEditorModule.h"
 #endif // WITH_AUDIOMODULATION_METASOUND_SUPPORT
@@ -121,11 +122,12 @@ void FAudioModulationEditorModule::StartupModule()
 	{
 		using namespace Metasound;
 		using namespace Metasound::Editor;
-		IMetasoundEditorModule& MetaSoundEditorModule = FModuleManager::GetModuleChecked<IMetasoundEditorModule>("MetaSoundEditor");
-		MetaSoundEditorModule.RegisterPinType("Modulator");
-		MetaSoundEditorModule.RegisterPinType(CreateArrayTypeNameFromElementTypeName("Modulator"));
-		MetaSoundEditorModule.RegisterPinType("ModulationParameter");
-		MetaSoundEditorModule.RegisterPinType(CreateArrayTypeNameFromElementTypeName("ModulationParameter"));
+		IMetasoundEditorModule& Module = FModuleManager::GetModuleChecked<IMetasoundEditorModule>(IMetasoundEditorModule::ModuleName);
+		TSharedRef<IMetaSoundGraphPanelPinFactory> PinFactory = Module.GetGraphPanelPinFactory();
+		PinFactory->RegisterPin("Modulator");
+		PinFactory->RegisterPin(CreateArrayTypeNameFromElementTypeName("Modulator"));
+		PinFactory->RegisterPin("ModulationParameter");
+		PinFactory->RegisterPin(CreateArrayTypeNameFromElementTypeName("ModulationParameter"));
 	}
 #endif // WITH_AUDIOMODULATION_METASOUND_SUPPORT
 
@@ -195,6 +197,20 @@ void FAudioModulationEditorModule::ShutdownModule()
 
 	FSlateStyleRegistry::UnRegisterSlateStyle(*StyleSet.Get());
 
+#if WITH_AUDIOMODULATION_METASOUND_SUPPORT
+	if (FModuleManager::Get().IsModuleLoaded(Metasound::Editor::IMetasoundEditorModule::ModuleName))
+	{
+		using namespace Metasound;
+		using namespace Metasound::Editor;
+		IMetasoundEditorModule& Module = FModuleManager::GetModuleChecked<IMetasoundEditorModule>(IMetasoundEditorModule::ModuleName);
+		TSharedRef<IMetaSoundGraphPanelPinFactory> PinFactory = Module.GetGraphPanelPinFactory();
+		PinFactory->UnregisterPin("Modulator");
+		PinFactory->UnregisterPin(CreateArrayTypeNameFromElementTypeName("Modulator"));
+		PinFactory->UnregisterPin("ModulationParameter");
+		PinFactory->UnregisterPin(CreateArrayTypeNameFromElementTypeName("ModulationParameter"));
+	}
+#endif // WITH_AUDIOMODULATION_METASOUND_SUPPORT
+
 #if WITH_EDITOR
 	if (GIsEditor)
 	{
diff --git a/Engine/Plugins/Runtime/Harmonix/Source/HarmonixMetasoundEditor/HarmonixMetasoundEditor.Build.cs b/Engine/Plugins/Runtime/Harmonix/Source/HarmonixMetasoundEditor/HarmonixMetasoundEditor.Build.cs
--- a/Engine/Plugins/Runtime/Harmonix/Source/HarmonixMetasoundEditor/HarmonixMetasoundEditor.Build.cs
+++ b/Engine/Plugins/Runtime/Harmonix/Source/HarmonixMetasoundEditor/HarmonixMetasoundEditor.Build.cs
@@ -31,6 +31,7 @@ public class HarmonixMetasoundEditor : ModuleRules
 				"HarmonixMetasound",
 				"MetasoundEditor",
 				"MetasoundFrontend",
+				"MetasoundGraphCore",
 				"DetailCustomizations",
 				"Projects"
 			}
diff --git a/Engine/Plugins/Runtime/Harmonix/Source/HarmonixMetasoundEditor/Private/HarmonixMetasoundEditorModule.cpp b/Engine/Plugins/Runtime/Harmonix/Source/HarmonixMetasoundEditor/Private/HarmonixMetasoundEditorModule.cpp
--- a/Engine/Plugins/Runtime/Harmonix/Source/HarmonixMetasoundEditor/Private/HarmonixMetasoundEditorModule.cpp
+++ b/Engine/Plugins/Runtime/Harmonix/Source/HarmonixMetasoundEditor/Private/HarmonixMetasoundEditorModule.cpp
@@ -3,11 +3,14 @@
 #include "HarmonixMetasoundEditorModule.h"
 #include "Modules/ModuleManager.h"
 #include "Logging/LogMacros.h"
+#include "IMetaSoundGraphPanelPinFactory.h"
 #include "MetasoundEditorModule.h"
 #include "HarmonixMetasoundSlateStyle.h"
 #include "MidiStepSequenceDetailCustomization.h"
 #include "HarmonixMetasound/DataTypes/MidiStepSequence.h"
 #include "AssetDefinition_MidiStepSequence.h"
+#include "MetasoundDataReference.h"
+
 
 #define LOCTEXT_NAMESPACE "HarmonixMetaSound"
 
@@ -15,35 +18,42 @@ DEFINE_LOG_CATEGORY(LogHarmonixMetasoundEditor)
 
 void FHarmonixMetasoundEditorModule::StartupModule()
 {
+	using namespace Metasound::Editor;
+
 	const HarmonixMetasoundEditor::FSlateStyle& Style = HarmonixMetasoundEditor::FSlateStyle::Get();
-	Metasound::Editor::IMetasoundEditorModule& MetasoundEditorModule = FModuleManager::GetModuleChecked<Metasound::Editor::IMetasoundEditorModule>("MetasoundEditor");
-	MetasoundEditorModule.RegisterPinType("MIDIAsset");
-	MetasoundEditorModule.RegisterPinType("MIDIStepSequenceAsset");
-	MetasoundEditorModule.RegisterPinType("StutterSequenceAsset");
-	MetasoundEditorModule.RegisterPinType("FusionPatchAsset");
-	MetasoundEditorModule.RegisterPinType("Enum:SubdivisionQuantizationType","Int32");
-	MetasoundEditorModule.RegisterPinType("Enum:DelayFilterType", "Int32");
-	MetasoundEditorModule.RegisterPinType("Enum:DelayStereoType", "Int32");
-	MetasoundEditorModule.RegisterPinType("Enum:TimeSyncOption", "Int32");
-	MetasoundEditorModule.RegisterPinType("Enum:DistortionType", "Int32");
-	MetasoundEditorModule.RegisterPinType("Enum:Harmonix:BiquadFilterType", "Int32");
-	MetasoundEditorModule.RegisterPinType("Enum:Distortion:FilterPasses", "Int32");
-	MetasoundEditorModule.RegisterPinType("Enum:StdMIDIControllerID", "Int32");
-
-	auto RegisterCustomPinType = [&Style, &MetasoundEditorModule](FName PinType)
+	IMetasoundEditorModule& MetasoundEditorModule = FModuleManager::GetModuleChecked<IMetasoundEditorModule>(IMetasoundEditorModule::ModuleName);
+	TSharedRef<IMetaSoundGraphPanelPinFactory> PinFactory = MetasoundEditorModule.GetGraphPanelPinFactory();
+	PinFactory->RegisterPin("MIDIAsset");
+	PinFactory->RegisterPin("MIDIStepSequenceAsset");
+	PinFactory->RegisterPin("StutterSequenceAsset");
+	PinFactory->RegisterPin("FusionPatchAsset");
+	PinFactory->RegisterPin("Enum:SubdivisionQuantizationType", FGraphPinParams { .PinCategory = "Int32" });
+	PinFactory->RegisterPin("Enum:DelayFilterType", FGraphPinParams { .PinCategory = "Int32" });
+	PinFactory->RegisterPin("Enum:DelayStereoType", FGraphPinParams { .PinCategory = "Int32" });
+	PinFactory->RegisterPin("Enum:TimeSyncOption", FGraphPinParams { .PinCategory = "Int32" });
+	PinFactory->RegisterPin("Enum:DistortionType", FGraphPinParams { .PinCategory = "Int32" });
+	PinFactory->RegisterPin("Enum:Harmonix:BiquadFilterType", FGraphPinParams { .PinCategory = "Int32" });
+	PinFactory->RegisterPin("Enum:Distortion:FilterPasses", FGraphPinParams { .PinCategory = "Int32" });
+	PinFactory->RegisterPin("Enum:StdMIDIControllerID", FGraphPinParams { .PinCategory = "Int32" });
+
+	auto RegisterCustomPinType = [&Style, &PinFactory](FName PinType)
 	{
-		Metasound::Editor::FGraphPinParams MidiStreamParams;
-		MidiStreamParams.PinCategory = PinType;
-		MidiStreamParams.PinColor = &Style.GetPinColor(PinType);
-		MidiStreamParams.PinConnectedIcon = Style.GetConnectedIcon(PinType);
-		MidiStreamParams.PinDisconnectedIcon = Style.GetDisconnectedIcon(PinType);
-		MetasoundEditorModule.RegisterCustomPinType(PinType, MidiStreamParams);
+		FGraphPinParams PinParams;
+		PinParams.PinCategory = PinType;
+		PinParams.PinColor = &Style.GetPinColor(PinType);
+		PinParams.PinConnectedIcon = Style.GetConnectedIcon(PinType);
+		PinParams.PinDisconnectedIcon = Style.GetDisconnectedIcon(PinType);
+		PinFactory->RegisterPin(PinType, PinParams);
+
+		PinParams.PinConnectedIcon = nullptr;
+		PinParams.PinDisconnectedIcon = nullptr;
+		PinFactory->RegisterPin(Metasound::CreateArrayTypeNameFromElementTypeName(PinType), PinParams);
 	};
 
 	RegisterCustomPinType("MIDIStream");
 	RegisterCustomPinType("MIDIClock");
 	RegisterCustomPinType("MusicTransport");
-	
+
 	FPropertyEditorModule& PropertyModule = FModuleManager::LoadModuleChecked<FPropertyEditorModule>("PropertyEditor");
 	PropertyModule.RegisterCustomClassLayout(
 		UMidiStepSequence::StaticClass()->GetFName(),
@@ -55,7 +65,33 @@ void FHarmonixMetasoundEditorModule::StartupModule()
 
 void FHarmonixMetasoundEditorModule::ShutdownModule()
 {
+	using namespace Metasound::Editor;
+	if (FModuleManager::Get().IsModuleLoaded(IMetasoundEditorModule::ModuleName))
+	{
+		IMetasoundEditorModule& Module = FModuleManager::GetModuleChecked<IMetasoundEditorModule>(IMetasoundEditorModule::ModuleName);
+		TSharedRef<IMetaSoundGraphPanelPinFactory> PinFactory = Module.GetGraphPanelPinFactory();
+		PinFactory->UnregisterPin("MIDIAsset");
+		PinFactory->UnregisterPin("MIDIStepSequenceAsset");
+		PinFactory->UnregisterPin("StutterSequenceAsset");
+		PinFactory->UnregisterPin("FusionPatchAsset");
+		PinFactory->UnregisterPin("Enum:SubdivisionQuantizationType");
+		PinFactory->UnregisterPin("Enum:DelayFilterType");
+		PinFactory->UnregisterPin("Enum:DelayStereoType");
+		PinFactory->UnregisterPin("Enum:TimeSyncOption");
+		PinFactory->UnregisterPin("Enum:DistortionType");
+		PinFactory->UnregisterPin("Enum:Harmonix:BiquadFilterType");
+		PinFactory->UnregisterPin("Enum:Distortion:FilterPasses");
+		PinFactory->UnregisterPin("Enum:StdMIDIControllerID");
 
+		auto UnregisterCustomPinType = [&PinFactory](FName PinType)
+		{
+			PinFactory->UnregisterPin(PinType);
+			PinFactory->UnregisterPin(Metasound::CreateArrayTypeNameFromElementTypeName(PinType));
+		};
+		UnregisterCustomPinType("MIDIStream");
+		UnregisterCustomPinType("MIDIClock");
+		UnregisterCustomPinType("MusicTransport");
+	}
 }
 
 IMPLEMENT_MODULE(FHarmonixMetasoundEditorModule, HarmonixMetasoundEditor);
diff --git a/Engine/Plugins/Runtime/Metasound/Source/MetasoundEditor/Private/FindInMetasound.cpp b/Engine/Plugins/Runtime/Metasound/Source/MetasoundEditor/Private/FindInMetasound.cpp
--- a/Engine/Plugins/Runtime/Metasound/Source/MetasoundEditor/Private/FindInMetasound.cpp
+++ b/Engine/Plugins/Runtime/Metasound/Source/MetasoundEditor/Private/FindInMetasound.cpp
@@ -268,7 +268,7 @@ namespace Metasound::Editor
 				}
 
 				// Get brush 
-				const IMetasoundEditorModule& EditorModule = FModuleManager::GetModuleChecked<IMetasoundEditorModule>("MetaSoundEditor");
+				const IMetasoundEditorModule& EditorModule = FModuleManager::GetModuleChecked<IMetasoundEditorModule>(IMetasoundEditorModule::ModuleName);
 				Brush = EditorModule.GetIconBrush(DataTypeName, bIsConstructorPin);
 			}
 			else
diff --git a/Engine/Plugins/Runtime/Metasound/Source/MetasoundEditor/Private/MetaSoundGraphPanelPinFactory.cpp b/Engine/Plugins/Runtime/Metasound/Source/MetasoundEditor/Private/MetaSoundGraphPanelPinFactory.cpp
new file mode 100644
--- /dev/null
+++ b/Engine/Plugins/Runtime/Metasound/Source/MetasoundEditor/Private/MetaSoundGraphPanelPinFactory.cpp
@@ -0,0 +1,415 @@
+// Copyright Epic Games, Inc. All Rights Reserved.
+
+#include "MetaSoundGraphPanelPinFactory.h"
+
+#include "EdGraph/EdGraphPin.h"
+#include "MetasoundAudioBuffer.h"
+#include "MetasoundEditorGraphNode.h"
+#include "MetasoundEditorModule.h"
+#include "MetasoundEditorSettings.h"
+#include "MetasoundFrontendDataTypeRegistry.h"
+#include "MetasoundTime.h"
+#include "MetasoundTrigger.h"
+#include "MetasoundWaveTable.h"
+#include "SGraphPin.h"
+#include "SMetasoundGraphEnumPin.h"
+#include "SMetasoundGraphPin.h"
+#include "WaveTable.h"
+
+
+namespace Metasound::Editor
+{
+	// Pin category constants
+	const FName FMetaSoundGraphPanelPinFactory::PinCategoryAudio = GetMetasoundDataTypeName<FAudioBuffer>();
+	const FName FMetaSoundGraphPanelPinFactory::PinCategoryBoolean = GetMetasoundDataTypeName<bool>();
+	const FName FMetaSoundGraphPanelPinFactory::PinCategoryFloat = GetMetasoundDataTypeName<float>();
+	const FName FMetaSoundGraphPanelPinFactory::PinCategoryInt32 = GetMetasoundDataTypeName<int32>();
+	const FName FMetaSoundGraphPanelPinFactory::PinCategoryObject = "object";
+	const FName FMetaSoundGraphPanelPinFactory::PinCategoryString = GetMetasoundDataTypeName<FString>();
+	const FName FMetaSoundGraphPanelPinFactory::PinCategoryTime = GetMetasoundDataTypeName<FTime>();
+	const FName FMetaSoundGraphPanelPinFactory::PinCategoryTimeArray = GetMetasoundDataTypeName<TArray<FTime>>();
+	const FName FMetaSoundGraphPanelPinFactory::PinCategoryTrigger = GetMetasoundDataTypeName<FTrigger>();
+	const FName FMetaSoundGraphPanelPinFactory::PinCategoryWaveTable = GetMetasoundDataTypeName<WaveTable::FWaveTable>();
+
+	bool FMetaSoundGraphPanelPinFactory::IsInternalMetaSoundPinCategory(FName InPinCategoryName)
+	{
+		return InPinCategoryName == PinCategoryAudio
+			|| InPinCategoryName == PinCategoryTime
+			|| InPinCategoryName == PinCategoryTimeArray
+			|| InPinCategoryName == PinCategoryTrigger
+			|| InPinCategoryName == PinCategoryWaveTable;
+	}
+
+	FMetaSoundGraphPanelPinFactory::FMetaSoundGraphPanelPinFactory()
+	{
+		RegisterCorePinTypes();
+	}
+
+	TSharedPtr<SGraphPin> FMetaSoundGraphPanelPinFactory::CreatePin(UEdGraphPin* InPin) const
+	{
+		if (!InPin)
+		{
+			return nullptr;
+		}
+
+		const UMetasoundEditorGraphNode* MetaSoundNode = Cast<UMetasoundEditorGraphNode>(InPin->GetOwningNode());
+		if (!MetaSoundNode)
+		{
+			return nullptr;
+		}
+
+		// Data type lookup takes priority to allow external plugins to override core types
+		const Frontend::FDataTypeRegistryInfo RegistryInfo = MetaSoundNode->GetPinDataTypeInfo(*InPin);
+		if (const FOnCreateMetaSoundPinWidget* DataTypeDelegate = DataTypeToWidgetDelegateMap.Find(RegistryInfo.DataTypeName))
+		{
+			if (DataTypeDelegate->IsBound())
+			{
+				TSharedPtr<SGraphPin> PinWidget = DataTypeDelegate->Execute(InPin);
+				if (PinWidget.IsValid())
+				{
+					return PinWidget;
+				}
+			}
+		}
+
+		const FName& PinCategory = InPin->PinType.PinCategory;
+		if (const FOnCreateMetaSoundPinWidget* CategoryDelegate = CategoryToWidgetDelegateMap.Find(PinCategory))
+		{
+			if (CategoryDelegate->IsBound())
+			{
+				TSharedPtr<SGraphPin> PinWidget = CategoryDelegate->Execute(InPin);
+				if (PinWidget.IsValid())
+				{
+					return PinWidget;
+				}
+			}
+		}
+
+		return nullptr;
+	}
+
+	const FEdGraphPinType* FMetaSoundGraphPanelPinFactory::FindPinType(FName InDataTypeName) const
+	{
+		if (const FPinConfiguration* PinConfiguration = PinTypes.Find(InDataTypeName))
+		{
+			return &PinConfiguration->PinType;
+		}
+
+		return nullptr;
+	}
+
+	TSharedRef<FMetaSoundGraphPanelPinFactory> FMetaSoundGraphPanelPinFactory::GetChecked()
+	{
+		IMetasoundEditorModule& EditorModule = FModuleManager::GetModuleChecked<IMetasoundEditorModule>(IMetasoundEditorModule::ModuleName);
+		return StaticCastSharedRef<FMetaSoundGraphPanelPinFactory>(EditorModule.GetGraphPanelPinFactory());
+	}
+
+	bool FMetaSoundGraphPanelPinFactory::GetCustomPinIcons(UEdGraphPin* InPin, const FSlateBrush*& OutConnectedIcon, const FSlateBrush*& OutDisconnectedIcon) const
+	{
+		if (InPin)
+		{
+			if (const UMetasoundEditorGraphNode* MetaSoundNode = Cast<UMetasoundEditorGraphNode>(InPin->GetOwningNode()))
+			{
+				Frontend::FDataTypeRegistryInfo RegistryInfo = MetaSoundNode->GetPinDataTypeInfo(*InPin);
+				return GetCustomPinIcons(RegistryInfo.DataTypeName, OutConnectedIcon, OutDisconnectedIcon);
+			}
+		}
+		return false;
+	}
+
+	bool FMetaSoundGraphPanelPinFactory::GetCustomPinIcons(FName InDataType, const FSlateBrush*& OutConnectedIcon, const FSlateBrush*& OutDisconnectedIcon) const
+	{
+		const FPinConfiguration* PinConfiguration = PinTypes.Find(InDataType);
+		if (!PinConfiguration || (!PinConfiguration->PinConnectedIcon && !PinConfiguration->PinDisconnectedIcon))
+		{
+			return false;
+		}
+
+		OutConnectedIcon = PinConfiguration->PinConnectedIcon;
+		OutDisconnectedIcon = PinConfiguration->PinDisconnectedIcon ? PinConfiguration->PinDisconnectedIcon : PinConfiguration->PinConnectedIcon;
+		return true;
+	}
+
+	FLinearColor FMetaSoundGraphPanelPinFactory::GetPinColor(const FEdGraphPinType& InPinType) const
+	{
+		const UMetasoundEditorSettings* Settings = GetDefault<UMetasoundEditorSettings>();
+		check(Settings);
+
+		if (InPinType.PinCategory == PinCategoryAudio)
+		{
+			return Settings->AudioPinTypeColor;
+		}
+
+		if (InPinType.PinCategory == PinCategoryBoolean)
+		{
+			return Settings->BooleanPinTypeColor;
+		}
+
+		if (InPinType.PinCategory == PinCategoryFloat)
+		{
+			return Settings->FloatPinTypeColor;
+		}
+
+		if (InPinType.PinCategory == PinCategoryInt32)
+		{
+			return Settings->IntPinTypeColor;
+		}
+
+		if (InPinType.PinCategory == PinCategoryObject)
+		{
+			return Settings->ObjectPinTypeColor;
+		}
+
+		if (InPinType.PinCategory == PinCategoryString)
+		{
+			return Settings->StringPinTypeColor;
+		}
+
+		if (InPinType.PinCategory == PinCategoryTime || InPinType.PinCategory == PinCategoryTimeArray)
+		{
+			return Settings->TimePinTypeColor;
+		}
+
+		if (InPinType.PinCategory == PinCategoryTrigger)
+		{
+			return Settings->TriggerPinTypeColor;
+		}
+
+		if (InPinType.PinCategory == PinCategoryWaveTable)
+		{
+			return Settings->WaveTablePinTypeColor;
+		}
+
+		// Check registered custom pin colors
+		if (const FLinearColor* CustomColor = CustomCategoryColors.Find(InPinType.PinCategory))
+		{
+			return *CustomColor;
+		}
+
+		return Settings->DefaultPinTypeColor;
+	}
+
+	void FMetaSoundGraphPanelPinFactory::RegisterCorePinTypes()
+	{
+		using namespace Frontend;
+
+		RegisterCategoryPin(PinCategoryBoolean, FOnCreateMetaSoundPinWidget::CreateLambda([](UEdGraphPin* Pin)
+		{
+			return SNew(SMetasoundGraphPinBool, Pin);
+		}));
+
+		RegisterCategoryPin(PinCategoryFloat, FOnCreateMetaSoundPinWidget::CreateLambda([](UEdGraphPin* Pin)
+		{
+			return SNew(SMetasoundGraphPinFloat, Pin);
+		}));
+
+		RegisterCategoryPin(PinCategoryTime, FOnCreateMetaSoundPinWidget::CreateLambda([](UEdGraphPin* Pin)
+		{
+			return SNew(SMetasoundGraphPinFloat, Pin);
+		}));
+
+		RegisterCategoryPin(PinCategoryInt32, FOnCreateMetaSoundPinWidget::CreateLambda([](UEdGraphPin* Pin) -> TSharedPtr<SGraphPin>
+		{
+			if (SMetasoundGraphEnumPin::FindEnumInterfaceFromPin(Pin))
+			{
+				return SNew(SMetasoundGraphEnumPin, Pin);
+			}
+			return SNew(SMetasoundGraphPinInteger, Pin);
+		}));
+
+		RegisterCategoryPin(PinCategoryObject, FOnCreateMetaSoundPinWidget::CreateLambda([](UEdGraphPin* Pin)
+		{
+			return SNew(SMetasoundGraphPinObject, Pin);
+		}));
+
+		RegisterCategoryPin(PinCategoryString, FOnCreateMetaSoundPinWidget::CreateLambda([](UEdGraphPin* Pin)
+		{
+			return SNew(SMetasoundGraphPinString, Pin);
+		}));
+
+		RegisterCategoryPin(PinCategoryAudio, FOnCreateMetaSoundPinWidget::CreateLambda([](UEdGraphPin* Pin)
+		{
+			return SNew(SMetasoundGraphPin, Pin);
+		}));
+
+		RegisterCategoryPin(PinCategoryTrigger, FOnCreateMetaSoundPinWidget::CreateLambda([](UEdGraphPin* Pin)
+		{
+			TSharedPtr<SGraphPin> PinWidget = SNew(SMetasoundGraphPin, Pin);
+			const FSlateBrush& PinConnectedBrush = Style::GetSlateBrushSafe("MetasoundEditor.Graph.TriggerPin.Connected");
+			const FSlateBrush& PinDisconnectedBrush = Style::GetSlateBrushSafe("MetasoundEditor.Graph.TriggerPin.Disconnected");
+			PinWidget->SetCustomPinIcon(&PinConnectedBrush, &PinDisconnectedBrush);
+			return PinWidget;
+		}));
+
+		// Register all known data types from the MetaSound data type registry
+		const IDataTypeRegistry& DataTypeRegistry = IDataTypeRegistry::Get();
+
+		TArray<FName> DataTypeNames;
+		DataTypeRegistry.GetRegisteredDataTypeNames(DataTypeNames);
+
+		for (FName DataTypeName : DataTypeNames)
+		{
+			FDataTypeRegistryInfo RegistryInfo;
+			if (ensure(DataTypeRegistry.GetDataTypeInfo(DataTypeName, RegistryInfo)))
+			{
+				FName PinCategory = DataTypeName;
+				FName PinSubCategory;
+
+				if (!IsInternalMetaSoundPinCategory(PinCategory))
+				{
+					switch (RegistryInfo.PreferredLiteralType)
+					{
+						case ELiteralType::Boolean:
+						case ELiteralType::BooleanArray:
+						{
+							PinCategory = PinCategoryBoolean;
+						}
+						break;
+
+						case ELiteralType::Float:
+						{
+							PinCategory = PinCategoryFloat;
+						}
+						break;
+
+						case ELiteralType::FloatArray:
+						{
+							if (RegistryInfo.bIsArrayType)
+							{
+								PinCategory = PinCategoryFloat;
+							}
+						}
+						break;
+
+						case ELiteralType::Integer:
+						{
+							PinCategory = PinCategoryInt32;
+						}
+						break;
+
+						case ELiteralType::IntegerArray:
+						{
+							if (RegistryInfo.bIsArrayType)
+							{
+								PinCategory = PinCategoryInt32;
+							}
+						}
+						break;
+
+						case ELiteralType::String:
+						{
+							PinCategory = PinCategoryString;
+						}
+						break;
+
+						case ELiteralType::StringArray:
+						{
+							if (RegistryInfo.bIsArrayType)
+							{
+								PinCategory = PinCategoryString;
+							}
+						}
+						break;
+
+						case ELiteralType::UObjectProxy:
+						case ELiteralType::UObjectProxyArray:
+						{
+							PinCategory = PinCategoryObject;
+						}
+						break;
+
+						case ELiteralType::None:
+						case ELiteralType::NoneArray:
+						case ELiteralType::Invalid:
+						default:
+						{
+							static_assert(static_cast<int32>(ELiteralType::Invalid) == 12, "Possible missing binding of pin category to primitive type");
+						}
+						break;
+					}
+				}
+
+				RegisterPin(DataTypeName, FGraphPinParams { .PinCategory = PinCategory, .PinSubCategory = PinSubCategory });
+			}
+		}
+	}
+
+	void FMetaSoundGraphPanelPinFactory::RegisterPin(FName InDataTypeName, const FGraphPinParams& Params)
+	{
+		using namespace Frontend;
+		FDataTypeRegistryInfo DataTypeInfo;
+		if (!ensureMsgf(IDataTypeRegistry::Get().GetDataTypeInfo(InDataTypeName, DataTypeInfo), TEXT("Failed to find data type with name '%s'. Pin type not registered"), *InDataTypeName.ToString()))
+		{
+			return;
+		}
+
+		const FName PinCategory = Params.PinCategory.IsNone() ? PinCategoryObject : Params.PinCategory;
+
+		FEdGraphPinType PinType;
+		PinType.PinCategory = PinCategory;
+		PinType.PinSubCategory = Params.PinSubCategory;
+		PinType.ContainerType = DataTypeInfo.bIsArrayType ? EPinContainerType::Array : EPinContainerType::None;
+		PinType.PinSubCategoryObject = IDataTypeRegistry::Get().GetUClassForDataType(InDataTypeName);
+
+		FPinConfiguration PinConfiguration;
+		PinConfiguration.PinType = PinType;
+		PinConfiguration.PinColor = Params.PinColor ? TOptional<FLinearColor>(*Params.PinColor) : TOptional<FLinearColor>();
+		PinConfiguration.PinConnectedIcon = Params.PinConnectedIcon;
+		PinConfiguration.PinDisconnectedIcon = Params.PinDisconnectedIcon;
+		PinTypes.Emplace(InDataTypeName, MoveTemp(PinConfiguration));
+
+		if (Params.PinColor && !Params.PinCategory.IsNone())
+		{
+			CustomCategoryColors.Add(Params.PinCategory, *Params.PinColor);
+		}
+	}
+
+	void FMetaSoundGraphPanelPinFactory::UnregisterPin(FName InDataTypeName)
+	{
+		FPinConfiguration RemovedConfig;
+		if (PinTypes.RemoveAndCopyValue(InDataTypeName, RemovedConfig))
+		{
+			const FName& RemovedCategory = RemovedConfig.PinType.PinCategory;
+			if (!RemovedCategory.IsNone() && RemovedConfig.PinColor.IsSet())
+			{
+				// Only remove the cached color if no other pin type still uses this category
+				bool bCategoryStillInUse = false;
+				for (const TPair<FName, FPinConfiguration>& Pair : PinTypes)
+				{
+					if (Pair.Value.PinType.PinCategory == RemovedCategory && Pair.Value.PinColor.IsSet())
+					{
+						bCategoryStillInUse = true;
+						break;
+					}
+				}
+
+				if (!bCategoryStillInUse)
+				{
+					CustomCategoryColors.Remove(RemovedCategory);
+				}
+			}
+		}
+	}
+
+	void FMetaSoundGraphPanelPinFactory::RegisterCategoryPin(FName InPinCategory, FOnCreateMetaSoundPinWidget InDelegate)
+	{
+		check(InDelegate.IsBound());
+		CategoryToWidgetDelegateMap.Add(InPinCategory, MoveTemp(InDelegate));
+	}
+
+	void FMetaSoundGraphPanelPinFactory::RegisterDataTypePin(FName InDataTypeName, FOnCreateMetaSoundPinWidget InDelegate)
+	{
+		check(InDelegate.IsBound());
+		DataTypeToWidgetDelegateMap.Add(InDataTypeName, MoveTemp(InDelegate));
+	}
+
+	void FMetaSoundGraphPanelPinFactory::UnregisterCategoryPin(FName InPinCategory)
+	{
+		CategoryToWidgetDelegateMap.Remove(InPinCategory);
+	}
+
+	void FMetaSoundGraphPanelPinFactory::UnregisterDataTypePin(FName InDataTypeName)
+	{
+		DataTypeToWidgetDelegateMap.Remove(InDataTypeName);
+	}
+} // namespace Metasound::Editor
diff --git a/Engine/Plugins/Runtime/Metasound/Source/MetasoundEditor/Private/MetaSoundGraphPanelPinFactory.h b/Engine/Plugins/Runtime/Metasound/Source/MetasoundEditor/Private/MetaSoundGraphPanelPinFactory.h
new file mode 100644
--- /dev/null
+++ b/Engine/Plugins/Runtime/Metasound/Source/MetasoundEditor/Private/MetaSoundGraphPanelPinFactory.h
@@ -0,0 +1,73 @@
+// Copyright Epic Games, Inc. All Rights Reserved.
+#pragma once
+
+#include "IMetaSoundGraphPanelPinFactory.h"
+
+
+namespace Metasound::Editor
+{
+	class FMetaSoundGraphPanelPinFactory : public IMetaSoundGraphPanelPinFactory
+	{
+	public:
+		FMetaSoundGraphPanelPinFactory();
+		virtual ~FMetaSoundGraphPanelPinFactory() = default;
+
+		virtual TSharedPtr<SGraphPin> CreatePin(UEdGraphPin* InPin) const override;
+
+		virtual const FEdGraphPinType* FindPinType(FName InDataTypeName) const override;
+
+		// Convenience getter for returning concrete private implementation
+		// from owning editor module and casting accordingly.
+		static TSharedRef<FMetaSoundGraphPanelPinFactory> GetChecked();
+
+		virtual bool GetCustomPinIcons(UEdGraphPin* InPin, const FSlateBrush*& OutConnectedIcon, const FSlateBrush*& OutDisconnectedIcon) const override;
+		virtual bool GetCustomPinIcons(FName InDataType, const FSlateBrush*& OutConnectedIcon, const FSlateBrush*& OutDisconnectedIcon) const override;
+
+		virtual FLinearColor GetPinColor(const FEdGraphPinType& InPinType) const override;
+
+		virtual void RegisterCategoryPin(FName InPinCategory, FOnCreateMetaSoundPinWidget InDelegate) override;
+		virtual void UnregisterCategoryPin(FName InPinCategory) override;
+
+		virtual void RegisterDataTypePin(FName InDataTypeName, FOnCreateMetaSoundPinWidget InDelegate) override;
+		virtual void UnregisterDataTypePin(FName InDataTypeName) override;
+
+		virtual void RegisterPin(FName InDataTypeName, const FGraphPinParams& Params = FGraphPinParams { }) override;
+		virtual void UnregisterPin(FName InDataTypeName) override;
+
+		// Pin category constants
+		static const FName PinCategoryAudio;
+		static const FName PinCategoryBoolean;
+		static const FName PinCategoryFloat;
+		static const FName PinCategoryInt32;
+		static const FName PinCategoryObject;
+		static const FName PinCategoryString;
+		static const FName PinCategoryTime;
+		static const FName PinCategoryTimeArray;
+		static const FName PinCategoryTrigger;
+		static const FName PinCategoryWaveTable;
+
+		// Returns whether the given name is an internally-managed MetaSound pin category
+		static bool IsInternalMetaSoundPinCategory(FName InPinCategoryName);
+
+	private:
+		void RegisterCorePinTypes();
+
+		/** Map from pin category to widget creation delegate */
+		TMap<FName, FOnCreateMetaSoundPinWidget> CategoryToWidgetDelegateMap;
+
+		/** Map from data type name to widget creation delegate (higher priority than category) */
+		TMap<FName, FOnCreateMetaSoundPinWidget> DataTypeToWidgetDelegateMap;
+
+		struct FPinConfiguration
+		{
+			FEdGraphPinType PinType;
+			TOptional<FLinearColor> PinColor;
+			const FSlateBrush* PinConnectedIcon = nullptr;
+			const FSlateBrush* PinDisconnectedIcon = nullptr;
+		};
+		TMap<FName, FPinConfiguration> PinTypes;
+
+		/** Cached custom pin category colors for O(1) lookup in GetPinColor */
+		TMap<FName, FLinearColor> CustomCategoryColors;
+	};
+} // namespace Metasound::Editor
diff --git a/Engine/Plugins/Runtime/Metasound/Source/MetasoundEditor/Private/MetasoundAssetDefinitions.cpp b/Engine/Plugins/Runtime/Metasound/Source/MetasoundEditor/Private/MetasoundAssetDefinitions.cpp
--- a/Engine/Plugins/Runtime/Metasound/Source/MetasoundEditor/Private/MetasoundAssetDefinitions.cpp
+++ b/Engine/Plugins/Runtime/Metasound/Source/MetasoundEditor/Private/MetasoundAssetDefinitions.cpp
@@ -306,7 +306,7 @@ namespace Metasound::Editor
 					TInstancedStruct<FPresetTemplate> PresetTemplate = TInstancedStruct<FPresetTemplate>::Make();
 					PresetTemplate.GetMutable().Parent = ParentMetaSound;
 
-					IMetasoundEditorModule& MetaSoundEditorModule = FModuleManager::GetModuleChecked<IMetasoundEditorModule>("MetaSoundEditor");
+					IMetasoundEditorModule& MetaSoundEditorModule = FModuleManager::GetModuleChecked<IMetasoundEditorModule>(IMetasoundEditorModule::ModuleName);
 					if (ParentPackage->bIsCookedForEditor)
 					{
 						// Cannot duplicate cooked assets or when underlying asset is cooked, so create new object and copy over properties
@@ -426,7 +426,7 @@ namespace Metasound::Editor
 					InSection.AddMenuEntry("MetaSoundSource_BrowseToPresetParent", Label, ToolTip, FindInContentBrowserIcon, UIAction);
 				}
 
-				IMetasoundEditorModule& EditorModule = FModuleManager::GetModuleChecked<Metasound::Editor::IMetasoundEditorModule>("MetaSoundEditor");
+				IMetasoundEditorModule& EditorModule = FModuleManager::GetModuleChecked<Metasound::Editor::IMetasoundEditorModule>(Metasound::Editor::IMetasoundEditorModule::ModuleName);
 				if (!EditorModule.IsRestrictedMode())
 				{
 					const TAttribute<FText> Label = LOCTEXT("MetaSound_EditPresetParentFormat", "Edit Preset Parent...");
@@ -593,7 +593,7 @@ TConstArrayView<FAssetCategoryPath> UAssetDefinition_MetaSoundPatch::GetAssetCat
 
 EAssetCommandResult UAssetDefinition_MetaSoundPatch::OpenAssets(const FAssetOpenArgs& OpenArgs) const
 {
-	Metasound::Editor::IMetasoundEditorModule& MetaSoundEditorModule = FModuleManager::GetModuleChecked<Metasound::Editor::IMetasoundEditorModule>("MetaSoundEditor");
+	Metasound::Editor::IMetasoundEditorModule& MetaSoundEditorModule = FModuleManager::GetModuleChecked<Metasound::Editor::IMetasoundEditorModule>(Metasound::Editor::IMetasoundEditorModule::ModuleName);
 	if (!MetaSoundEditorModule.IsRestrictedMode())
 	{
 		for (UMetaSoundPatch* Metasound : OpenArgs.LoadObjects<UMetaSoundPatch>())
@@ -675,7 +675,7 @@ EAssetCommandResult UAssetDefinition_MetaSoundSource::OpenAssets(const FAssetOpe
 				TSharedPtr<IToolkitHost> HostPtr = ToolkitHost.Pin();
 				if (bHostNull || HostPtr)
 				{
-					Editor::IMetasoundEditorModule* EditorModule = FModuleManager::GetModulePtr<Editor::IMetasoundEditorModule>("MetaSoundEditor");
+					Editor::IMetasoundEditorModule* EditorModule = FModuleManager::GetModulePtr<Editor::IMetasoundEditorModule>(Editor::IMetasoundEditorModule::ModuleName);
 					if (EditorModule)
 					{
 						if (EditorModule->IsRestrictedMode())
diff --git a/Engine/Plugins/Runtime/Metasound/Source/MetasoundEditor/Private/MetasoundEditor.cpp b/Engine/Plugins/Runtime/Metasound/Source/MetasoundEditor/Private/MetasoundEditor.cpp
--- a/Engine/Plugins/Runtime/Metasound/Source/MetasoundEditor/Private/MetasoundEditor.cpp
+++ b/Engine/Plugins/Runtime/Metasound/Source/MetasoundEditor/Private/MetasoundEditor.cpp
@@ -46,6 +46,7 @@
 #include "MetasoundEditorDocumentClipboardUtils.h"
 #include "MetasoundEditorGraph.h"
 #include "MetasoundEditorGraphBuilder.h"
+#include "MetaSoundGraphPanelPinFactory.h"
 #include "MetasoundEditorGraphCommentNode.h"
 #include "MetasoundEditorGraphInputNode.h"
 #include "MetasoundEditorGraphSchema.h"
@@ -300,13 +301,13 @@ namespace Metasound
 
 		bool IsPageAuditioningEnabled()
 		{
-			IMetasoundEditorModule& MetaSoundEditorModule = FModuleManager::GetModuleChecked<IMetasoundEditorModule>("MetaSoundEditor");
+			IMetasoundEditorModule& MetaSoundEditorModule = FModuleManager::GetModuleChecked<IMetasoundEditorModule>(IMetasoundEditorModule::ModuleName);
 			return !MetaSoundEditorModule.IsRestrictedMode();
 		}
 
 		bool IsPageGraphEditingEnabled(const FMetaSoundFrontendDocumentBuilder& Builder, bool bHasLocalPageValues)
 		{
-			IMetasoundEditorModule& MetaSoundEditorModule = FModuleManager::GetModuleChecked<IMetasoundEditorModule>("MetaSoundEditor");
+			IMetasoundEditorModule& MetaSoundEditorModule = FModuleManager::GetModuleChecked<IMetasoundEditorModule>(IMetasoundEditorModule::ModuleName);
 			if (MetaSoundEditorModule.IsRestrictedMode())
 			{
 				return false;
@@ -347,7 +348,7 @@ namespace Metasound
 				return false;
 			}
 
-			IMetasoundEditorModule& MetaSoundEditorModule = FModuleManager::GetModuleChecked<IMetasoundEditorModule>("MetaSoundEditor");
+			IMetasoundEditorModule& MetaSoundEditorModule = FModuleManager::GetModuleChecked<IMetasoundEditorModule>(IMetasoundEditorModule::ModuleName);
 			if (MetaSoundEditorModule.IsRestrictedMode())
 			{
 				return false;
@@ -1019,8 +1020,7 @@ namespace Metasound
 						}
 						FName DataTypeName = GraphMember->GetDataType();
 
-						const IMetasoundEditorModule& EditorModule = FModuleManager::GetModuleChecked<IMetasoundEditorModule>("MetaSoundEditor");
-						if (const FEdGraphPinType* PinType = EditorModule.FindPinType(DataTypeName))
+						if (const FEdGraphPinType* PinType = FMetaSoundGraphPanelPinFactory::GetChecked()->FindPinType(DataTypeName))
 						{
 							if (const UMetasoundEditorGraphSchema* Schema = GetDefault<UMetasoundEditorGraphSchema>())
 							{
@@ -1028,6 +1028,7 @@ namespace Metasound
 							}
 						}
 
+						const IMetasoundEditorModule& EditorModule = FModuleManager::GetModuleChecked<IMetasoundEditorModule>(IMetasoundEditorModule::ModuleName);
 						IconBrush = EditorModule.GetIconBrush(DataTypeName, bIsConstructorPin);
 					}
 				}
@@ -1178,7 +1179,7 @@ namespace Metasound
 			return { };
 		}
 
-		const FName FEditor::EditorName = "MetaSoundEditor";
+		const FName FEditor::EditorName = IMetasoundEditorModule::ModuleName;
 
 		FEditor::FEditor()
 		: GraphConnectionManager(MakeUnique<FGraphConnectionManager>())
@@ -1399,7 +1400,7 @@ namespace Metasound
 			check(ObjectToEdit);
 			checkf(IMetasoundUObjectRegistry::Get().IsRegisteredClass(ObjectToEdit), TEXT("Object passed in was not registered as a valid metasound interface!"));
 
-			IMetasoundEditorModule& MetaSoundEditorModule = FModuleManager::GetModuleChecked<IMetasoundEditorModule>("MetaSoundEditor");
+			IMetasoundEditorModule& MetaSoundEditorModule = FModuleManager::GetModuleChecked<IMetasoundEditorModule>(IMetasoundEditorModule::ModuleName);
 			if (!MetaSoundEditorModule.IsRestrictedMode())
 			{
 				IMetasoundEngineModule& MetaSoundEngineModule = FModuleManager::GetModuleChecked<IMetasoundEngineModule>("MetaSoundEngine");
@@ -2519,7 +2520,7 @@ namespace Metasound
 		{
 			TAttribute<EVisibility> ClearVisibility = TAttribute<EVisibility>::CreateSPLambda(this, [this]()
 			{
-				const IMetasoundEditorModule & MetaSoundEditorModule = FModuleManager::GetModuleChecked<IMetasoundEditorModule>("MetaSoundEditor");
+				const IMetasoundEditorModule & MetaSoundEditorModule = FModuleManager::GetModuleChecked<IMetasoundEditorModule>(IMetasoundEditorModule::ModuleName);
 				if (!MetaSoundEditorModule.IsRestrictedMode())
 				{
 					return GetDocumentTemplate() ? EVisibility::Visible : EVisibility::Collapsed;
@@ -2529,7 +2530,7 @@ namespace Metasound
 			});
 
 
-			IMetasoundEditorModule& MetaSoundEditorModule = FModuleManager::GetModuleChecked<IMetasoundEditorModule>("MetaSoundEditor");
+			IMetasoundEditorModule& MetaSoundEditorModule = FModuleManager::GetModuleChecked<IMetasoundEditorModule>(IMetasoundEditorModule::ModuleName);
 			if (!MetaSoundEditorModule.IsRestrictedMode())
 			{
 				MenuBuilder.BeginSection("ConvertTemplateHeader", LOCTEXT("ConvertTemplateDisplayName", "Convert Template"), ClearVisibility);
@@ -3248,7 +3249,7 @@ namespace Metasound
 						// ensure when generating asset key.
 						if (Class->Metadata.GetType() == EMetasoundFrontendClassType::External)
 						{
-							IMetasoundEditorModule& MetaSoundEditorModule = FModuleManager::GetModuleChecked<IMetasoundEditorModule>("MetaSoundEditor");
+							IMetasoundEditorModule& MetaSoundEditorModule = FModuleManager::GetModuleChecked<IMetasoundEditorModule>(IMetasoundEditorModule::ModuleName);
 							if (!MetaSoundEditorModule.IsRestrictedMode())
 							{
 								const FMetaSoundAssetKey AssetKey(Class->Metadata);
diff --git a/Engine/Plugins/Runtime/Metasound/Source/MetasoundEditor/Private/MetasoundEditorGraph.cpp b/Engine/Plugins/Runtime/Metasound/Source/MetasoundEditor/Private/MetasoundEditorGraph.cpp
--- a/Engine/Plugins/Runtime/Metasound/Source/MetasoundEditor/Private/MetasoundEditorGraph.cpp
+++ b/Engine/Plugins/Runtime/Metasound/Source/MetasoundEditor/Private/MetasoundEditorGraph.cpp
@@ -178,7 +178,7 @@ void UMetasoundEditorGraphMember::InitializeLiteral(bool bForceRebind)
 	using namespace Metasound::Frontend;
 
 	FDataTypeRegistryInfo DataTypeInfo;
-	IMetasoundEditorModule& EditorModule = FModuleManager::GetModuleChecked<IMetasoundEditorModule>("MetaSoundEditor");
+	IMetasoundEditorModule& EditorModule = FModuleManager::GetModuleChecked<IMetasoundEditorModule>(IMetasoundEditorModule::ModuleName);
 	IDataTypeRegistry::Get().GetDataTypeInfo(TypeName, DataTypeInfo);
 	const EMetasoundFrontendLiteralType LiteralType = static_cast<EMetasoundFrontendLiteralType>(DataTypeInfo.PreferredLiteralType);
 
diff --git a/Engine/Plugins/Runtime/Metasound/Source/MetasoundEditor/Private/MetasoundEditorGraphBuilder.cpp b/Engine/Plugins/Runtime/Metasound/Source/MetasoundEditor/Private/MetasoundEditorGraphBuilder.cpp
--- a/Engine/Plugins/Runtime/Metasound/Source/MetasoundEditor/Private/MetasoundEditorGraphBuilder.cpp
+++ b/Engine/Plugins/Runtime/Metasound/Source/MetasoundEditor/Private/MetasoundEditorGraphBuilder.cpp
@@ -26,6 +26,7 @@
 #include "MetasoundEditorGraphSchema.h"
 #include "MetasoundEditorGraphValidation.h"
 #include "MetasoundEditorModule.h"
+#include "MetaSoundGraphPanelPinFactory.h"
 #include "MetasoundEditorSettings.h"
 #include "MetasoundEditorSubsystem.h"
 #include "MetasoundFactory.h"
@@ -149,29 +150,6 @@ namespace Metasound
 			}
 		} // namespace GraphBuilderPrivate
 
-		// Categories corresponding with POD DataTypes
-		const FName FGraphBuilder::PinCategoryObject = "object"; // Basket for all UObject proxy types (corresponds to multiple DataTypes)
-		const FName FGraphBuilder::PinCategoryBoolean = GetMetasoundDataTypeName<bool>();
-		const FName FGraphBuilder::PinCategoryFloat = GetMetasoundDataTypeName<float>();
-		const FName FGraphBuilder::PinCategoryInt32 = GetMetasoundDataTypeName<int32>();
-		const FName FGraphBuilder::PinCategoryString = GetMetasoundDataTypeName<FString>();
-
-		// Categories corresponding with MetaSound DataTypes with custom visualization
-		const FName FGraphBuilder::PinCategoryAudio = GetMetasoundDataTypeName<FAudioBuffer>();
-		const FName FGraphBuilder::PinCategoryTime = GetMetasoundDataTypeName<FTime>();
-		const FName FGraphBuilder::PinCategoryTimeArray = GetMetasoundDataTypeName<TArray<FTime>>();
-		const FName FGraphBuilder::PinCategoryTrigger = GetMetasoundDataTypeName<FTrigger>();
-		const FName FGraphBuilder::PinCategoryWaveTable = GetMetasoundDataTypeName<WaveTable::FWaveTable>();
-
-		bool FGraphBuilder::IsPinCategoryMetaSoundCustomDataType(FName InPinCategoryName)
-		{
-			return InPinCategoryName == PinCategoryAudio
-				|| InPinCategoryName == PinCategoryTime
-				|| InPinCategoryName == PinCategoryTimeArray
-				|| InPinCategoryName == PinCategoryTrigger
-				|| InPinCategoryName == PinCategoryWaveTable;
-		}
-
 		bool FGraphBuilder::CanInspectPin(const UEdGraphPin* InPin)
 		{
 			// Can't inspect the value on an invalid pin object.
@@ -659,65 +637,6 @@ namespace Metasound
 			return { };
 		}
 
-		FLinearColor FGraphBuilder::GetPinCategoryColor(const FEdGraphPinType& PinType)
-		{
-			const UMetasoundEditorSettings* Settings = GetDefault<UMetasoundEditorSettings>();
-			check(Settings);
-
-			if (PinType.PinCategory == PinCategoryAudio)
-			{
-				return Settings->AudioPinTypeColor;
-			}
-
-			if (PinType.PinCategory == PinCategoryBoolean)
-			{
-				return Settings->BooleanPinTypeColor;
-			}
-
-			if (PinType.PinCategory == PinCategoryFloat)
-			{
-				return Settings->FloatPinTypeColor;
-			}
-
-			if (PinType.PinCategory == PinCategoryInt32)
-			{
-				return Settings->IntPinTypeColor;
-			}
-
-			if (PinType.PinCategory == PinCategoryObject)
-			{
-				return Settings->ObjectPinTypeColor;
-			}
-
-			if (PinType.PinCategory == PinCategoryString)
-			{
-				return Settings->StringPinTypeColor;
-			}
-
-			if (PinType.PinCategory == PinCategoryTime || PinType.PinCategory == PinCategoryTimeArray)
-			{
-				return Settings->TimePinTypeColor;
-			}
-
-			if (PinType.PinCategory == PinCategoryTrigger)
-			{
-				return Settings->TriggerPinTypeColor;
-			}
-
-			if (PinType.PinCategory == PinCategoryWaveTable)
-			{
-				return Settings->WaveTablePinTypeColor;
-			}
-
-			// custom colors
-			if (const FLinearColor* Color = Settings->CustomPinTypeColors.Find(PinType.PinCategory))
-			{
-				return *Color;
-			}
-
-			return Settings->DefaultPinTypeColor;
-		}
-
 		Frontend::FInputHandle FGraphBuilder::GetInputHandleFromPin(const UEdGraphPin* InPin)
 		{
 			using namespace Frontend;
@@ -1119,7 +1038,7 @@ PRAGMA_ENABLE_DEPRECATION_WARNINGS
 		{
 			using namespace Frontend;
 
-			IMetasoundEditorModule& EditorModule = FModuleManager::GetModuleChecked<IMetasoundEditorModule>("MetaSoundEditor");
+			IMetasoundEditorModule& EditorModule = FModuleManager::GetModuleChecked<IMetasoundEditorModule>(IMetasoundEditorModule::ModuleName);
 
 			const FMetaSoundFrontendDocumentBuilder& Builder = GetBuilderFromPinChecked(InInputPin).GetConstBuilder();
 			FMetasoundFrontendVertexHandle InputHandle = GetPinVertexHandle(Builder, &InInputPin);
@@ -1629,7 +1548,7 @@ PRAGMA_ENABLE_DEPRECATION_WARNINGS
 			using namespace Frontend;
 
 			FEdGraphPinType PinType;
-			if (const FEdGraphPinType* RegisteredPinType = EditorModule.FindPinType(InVertex.TypeName))
+			if (const FEdGraphPinType* RegisteredPinType = FMetaSoundGraphPanelPinFactory::GetChecked()->FindPinType(InVertex.TypeName))
 			{
 				PinType = *RegisteredPinType;
 			}
@@ -1683,7 +1602,7 @@ PRAGMA_ENABLE_DEPRECATION_WARNINGS
 		bool FGraphBuilder::SynchronizePinType(const IMetasoundEditorModule& InEditorModule, UEdGraphPin& InPin, const FName InDataType)
 		{
 			FEdGraphPinType PinType;
-			if (const FEdGraphPinType* RegisteredPinType = InEditorModule.FindPinType(InDataType))
+			if (const FEdGraphPinType* RegisteredPinType = FMetaSoundGraphPanelPinFactory::GetChecked()->FindPinType(InDataType))
 			{
 				PinType = *RegisteredPinType;
 			}
@@ -2243,7 +2162,7 @@ PRAGMA_ENABLE_DEPRECATION_WARNINGS
 
 			bool bIsNodeDirty = false;
 
-			IMetasoundEditorModule& EditorModule = FModuleManager::GetModuleChecked<IMetasoundEditorModule>("MetaSoundEditor");
+			IMetasoundEditorModule& EditorModule = FModuleManager::GetModuleChecked<IMetasoundEditorModule>(IMetasoundEditorModule::ModuleName);
 
 			const FGuid& NodeID = InNode.GetID();
 			TArray<const FMetasoundFrontendVertex*> NodeInputs = InBuilder.FindUserModifiableNodeInputs(NodeID);
diff --git a/Engine/Plugins/Runtime/Metasound/Source/MetasoundEditor/Private/MetasoundEditorGraphBuilder.h b/Engine/Plugins/Runtime/Metasound/Source/MetasoundEditor/Private/MetasoundEditorGraphBuilder.h
--- a/Engine/Plugins/Runtime/Metasound/Source/MetasoundEditor/Private/MetasoundEditorGraphBuilder.h
+++ b/Engine/Plugins/Runtime/Metasound/Source/MetasoundEditor/Private/MetasoundEditorGraphBuilder.h
@@ -49,18 +49,6 @@ namespace Metasound
 		class FGraphBuilder
 		{
 		public:
-			static const FName PinCategoryAudio;
-			static const FName PinCategoryBoolean;
-			static const FName PinCategoryFloat;
-			static const FName PinCategoryInt32;
-			static const FName PinCategoryObject;
-			static const FName PinCategoryString;
-			static const FName PinCategoryTime;
-			static const FName PinCategoryTimeArray;
-			static const FName PinCategoryTrigger;
-			static const FName PinCategoryWaveTable;
-
-
 			static const FText FunctionMenuName;
 			static const FText GraphMenuName;
 
@@ -91,9 +79,6 @@ namespace Metasound
 			// @param RegOptions - Registration options utilized
 			static void RegisterGraphWithFrontend(UObject& InMetaSound, Frontend::FMetaSoundAssetRegistrationOptions RegOptions = GetDefaultRegistrationOptions());
 
-			// Returns whether pin category is a custom MetaSound DataType
-			static bool IsPinCategoryMetaSoundCustomDataType(FName InPinCategoryName);
-
 			// Determines if pin supports inspection/probe view.
 			static bool CanInspectPin(const UEdGraphPin* InPin);
 
@@ -204,9 +189,6 @@ namespace Metasound
 			// Returns the default literal stored on the respective Frontend Node's Input.
 			static bool GetPinLiteral(const UEdGraphPin& InInputPin, FMetasoundFrontendLiteral& OutLiteralDefault);
 
-			// Retrieves the proper pin color for the given PinType
-			static FLinearColor GetPinCategoryColor(const FEdGraphPinType& PinType);
-
 			// Deletes both the editor graph & frontend nodes from respective graphs
 			static bool DeleteNode(UEdGraphNode& InNode, bool bRemoveUnusedDependencies = true);
 
diff --git a/Engine/Plugins/Runtime/Metasound/Source/MetasoundEditor/Private/MetasoundEditorGraphConnectionDrawingPolicy.cpp b/Engine/Plugins/Runtime/Metasound/Source/MetasoundEditor/Private/MetasoundEditorGraphConnectionDrawingPolicy.cpp
--- a/Engine/Plugins/Runtime/Metasound/Source/MetasoundEditor/Private/MetasoundEditorGraphConnectionDrawingPolicy.cpp
+++ b/Engine/Plugins/Runtime/Metasound/Source/MetasoundEditor/Private/MetasoundEditorGraphConnectionDrawingPolicy.cpp
@@ -10,6 +10,7 @@
 #include "MetasoundEditorGraph.h"
 #include "MetasoundEditorGraphBuilder.h"
 #include "MetasoundEditorGraphSchema.h"
+#include "MetaSoundGraphPanelPinFactory.h"
 #include "MetasoundTrace.h"
 #include "Misc/App.h"
 #include "NodeTemplates/MetasoundFrontendNodeTemplateReroute.h"
@@ -360,7 +361,7 @@ namespace Metasound
 			}
 			else if (AnyPin)
 			{
-				OutParams.WireColor = FGraphBuilder::GetPinCategoryColor(AnyPin->PinType);
+				OutParams.WireColor = FMetaSoundGraphPanelPinFactory::GetChecked()->GetPinColor(AnyPin->PinType);
 			}
 
 			bool bThicknessSet = false;
@@ -378,14 +379,14 @@ namespace Metasound
 								check(OutputNode);
 								const bool bIsPlaying = EditorPtr->IsPlaying();
 								const FGuid& NodeID = OutputNode->GetID();
-								if (AnyPin->PinType.PinCategory == FGraphBuilder::PinCategoryTrigger)
+								if (AnyPin->PinType.PinCategory == FMetaSoundGraphPanelPinFactory::PinCategoryTrigger)
 								{
 									FGraphConnectionManager& ConnectionManager = EditorPtr->GetConnectionManager();
 									DrawingPolicyPrivate::DetermineWiringStyle_Envelope<FVertexAnalyzerTriggerDensity>(bIsPlaying, ConnectionManager, NodeID, OutputVertex->Name, InputPin, OutParams);
 									ScaleEnvelopeWireThickness(OutParams);
 									bThicknessSet = true;
 								}
-								else if (AnyPin->PinType.PinCategory == FGraphBuilder::PinCategoryAudio)
+								else if (AnyPin->PinType.PinCategory == FMetaSoundGraphPanelPinFactory::PinCategoryAudio)
 								{
 									FGraphConnectionManager& ConnectionManager = EditorPtr->GetConnectionManager();
 									DrawingPolicyPrivate::DetermineWiringStyle_Envelope<FVertexAnalyzerEnvelopeFollower>(bIsPlaying, ConnectionManager, NodeID, OutputVertex->Name, InputPin, OutParams);
diff --git a/Engine/Plugins/Runtime/Metasound/Source/MetasoundEditor/Private/MetasoundEditorGraphNode.cpp b/Engine/Plugins/Runtime/Metasound/Source/MetasoundEditor/Private/MetasoundEditorGraphNode.cpp
--- a/Engine/Plugins/Runtime/Metasound/Source/MetasoundEditor/Private/MetasoundEditorGraphNode.cpp
+++ b/Engine/Plugins/Runtime/Metasound/Source/MetasoundEditor/Private/MetasoundEditorGraphNode.cpp
@@ -656,7 +656,7 @@ TSet<FString> UMetasoundEditorGraphNode::GetDisallowedPinClassNames(const UEdGra
 	using namespace Metasound::Editor;
 	using namespace Metasound::Frontend;
 
-	const IMetasoundEditorModule& EditorModule = FModuleManager::GetModuleChecked<IMetasoundEditorModule>("MetaSoundEditor");
+	const IMetasoundEditorModule& EditorModule = FModuleManager::GetModuleChecked<IMetasoundEditorModule>(IMetasoundEditorModule::ModuleName);
 
 	const FDataTypeRegistryInfo DataTypeInfo = GetPinDataTypeInfo(InPin);
 	if (DataTypeInfo.PreferredLiteralType != Metasound::ELiteralType::UObjectProxy)
diff --git a/Engine/Plugins/Runtime/Metasound/Source/MetasoundEditor/Private/MetasoundEditorGraphNodeVisualization.h b/Engine/Plugins/Runtime/Metasound/Source/MetasoundEditor/Private/MetasoundEditorGraphNodeVisualization.h
--- a/Engine/Plugins/Runtime/Metasound/Source/MetasoundEditor/Private/MetasoundEditorGraphNodeVisualization.h
+++ b/Engine/Plugins/Runtime/Metasound/Source/MetasoundEditor/Private/MetasoundEditorGraphNodeVisualization.h
@@ -4,6 +4,7 @@
 #include "Containers/Map.h"
 #include "MetasoundEditorGraphBuilder.h"
 #include "MetasoundEditorGraphNode.h"
+#include "MetaSoundGraphPanelPinFactory.h"
 #include "MetasoundEditorModule.h"
 
 namespace Metasound::Editor
@@ -33,7 +34,7 @@ namespace Metasound::Editor
 		{
 			UEdGraphPin* const* FoundNamedInputPin = MetaSoundNode.Pins.FindByPredicate([&PinName](const UEdGraphPin* InPin)
 				{
-					return (InPin->Direction == EGPD_Input && InPin->PinType.PinCategory != FGraphBuilder::PinCategoryAudio && InPin->PinName == PinName);
+					return (InPin->Direction == EGPD_Input && InPin->PinType.PinCategory != FMetaSoundGraphPanelPinFactory::PinCategoryAudio && InPin->PinName == PinName);
 				});
 
 			if (!FoundNamedInputPin)
diff --git a/Engine/Plugins/Runtime/Metasound/Source/MetasoundEditor/Private/MetasoundEditorGraphSchema.cpp b/Engine/Plugins/Runtime/Metasound/Source/MetasoundEditor/Private/MetasoundEditorGraphSchema.cpp
--- a/Engine/Plugins/Runtime/Metasound/Source/MetasoundEditor/Private/MetasoundEditorGraphSchema.cpp
+++ b/Engine/Plugins/Runtime/Metasound/Source/MetasoundEditor/Private/MetasoundEditorGraphSchema.cpp
@@ -28,6 +28,7 @@
 #include "MetasoundEditorCommands.h"
 #include "MetasoundEditorGraph.h"
 #include "MetasoundEditorGraphBuilder.h"
+#include "MetaSoundGraphPanelPinFactory.h"
 #include "MetasoundEditorGraphCommentNode.h"
 #include "MetasoundEditorGraphInputNode.h"
 #include "MetasoundEditorGraphNode.h"
@@ -2684,7 +2685,7 @@ float UMetasoundEditorGraphSchema::GetActionFilteredWeight(const FEdGraphSchemaA
 
 FLinearColor UMetasoundEditorGraphSchema::GetPinTypeColor(const FEdGraphPinType& PinType) const
 {
-	return Metasound::Editor::FGraphBuilder::GetPinCategoryColor(PinType);
+	return Metasound::Editor::FMetaSoundGraphPanelPinFactory::GetChecked()->GetPinColor(PinType);
 }
 
 void UMetasoundEditorGraphSchema::BreakNodeLinks(UEdGraphNode& TargetNode) const
diff --git a/Engine/Plugins/Runtime/Metasound/Source/MetasoundEditor/Private/MetasoundEditorModule.cpp b/Engine/Plugins/Runtime/Metasound/Source/MetasoundEditor/Private/MetasoundEditorModule.cpp
--- a/Engine/Plugins/Runtime/Metasound/Source/MetasoundEditor/Private/MetasoundEditorModule.cpp
+++ b/Engine/Plugins/Runtime/Metasound/Source/MetasoundEditor/Private/MetasoundEditorModule.cpp
@@ -39,6 +39,7 @@
 #include "MetasoundEditorGraphBuilder.h"
 #include "MetasoundEditorGraphConnectionDrawingPolicy.h"
 #include "MetasoundEditorGraphMemberDefaults.h"
+#include "MetaSoundGraphPanelPinFactory.h"
 #include "MetasoundEditorGraphNode.h"
 #include "MetasoundEditorGraphNodeFactory.h"
 #include "MetasoundEditorGraphNodeVisualization.h"
@@ -89,6 +90,8 @@ DEFINE_LOG_CATEGORY(LogMetasoundEditor);
 
 namespace Metasound::Editor
 {
+	const FName IMetasoundEditorModule::ModuleName = "MetaSoundEditor";
+
 	EAssetPrimeStatus IMetasoundEditorModule::GetAssetRegistryPrimeStatus() const
 	{
 		return EAssetPrimeStatus::NotRequested;
@@ -665,15 +668,6 @@ namespace Metasound::Editor
 		}
 	} // namespace Style
 
-	// A structure that contains information about registered custom pin types. 
-	struct FGraphPinConfiguration
-	{
-		FEdGraphPinType PinType;
-		const FSlateBrush* PinConnectedIcon = nullptr;
-		const FSlateBrush* PinDisconnectedIcon = nullptr;
-	};
-
-
 	class FModule : public IMetasoundEditorModule
 	{
 		void RegisterInputDefaultClasses()
@@ -704,105 +698,6 @@ namespace Metasound::Editor
 			}
 		}
 
-		void RegisterCorePinTypes()
-		{
-			using namespace Metasound::Frontend;
-
-			const IDataTypeRegistry& DataTypeRegistry = IDataTypeRegistry::Get();
-
-			TArray<FName> DataTypeNames;
-			DataTypeRegistry.GetRegisteredDataTypeNames(DataTypeNames);
-
-			for (FName DataTypeName : DataTypeNames)
-			{
-				FDataTypeRegistryInfo RegistryInfo;
-				if (ensure(DataTypeRegistry.GetDataTypeInfo(DataTypeName, RegistryInfo)))
-				{
-					FName PinCategory = DataTypeName;
-					FName PinSubCategory;
-
-					// Types like triggers & AudioBuffer are specialized, so ignore their preferred
-					// literal types to classify the category.
-					if (!FGraphBuilder::IsPinCategoryMetaSoundCustomDataType(PinCategory) && !CustomPinCategories.Contains(PinCategory))
-					{
-						// Primitives
-						switch (RegistryInfo.PreferredLiteralType)
-						{
-							case ELiteralType::Boolean:
-							case ELiteralType::BooleanArray:
-							{
-								PinCategory = FGraphBuilder::PinCategoryBoolean;
-							}
-							break;
-
-							case ELiteralType::Float:
-							{
-								PinCategory = FGraphBuilder::PinCategoryFloat;
-							}
-							break;
-
-							case ELiteralType::FloatArray:
-							{
-								if (RegistryInfo.bIsArrayType)
-								{
-									PinCategory = FGraphBuilder::PinCategoryFloat;
-								}
-							}
-							break;
-
-							case ELiteralType::Integer:
-							{
-								PinCategory = FGraphBuilder::PinCategoryInt32;
-							}
-							break;
-
-							case ELiteralType::IntegerArray:
-							{
-								if (RegistryInfo.bIsArrayType)
-								{
-									PinCategory = FGraphBuilder::PinCategoryInt32;
-								}
-							}
-							break;
-
-							case ELiteralType::String:
-							{
-								PinCategory = FGraphBuilder::PinCategoryString;
-							}
-							break;
-
-							case ELiteralType::StringArray:
-							{
-								if (RegistryInfo.bIsArrayType)
-								{
-									PinCategory = FGraphBuilder::PinCategoryString;
-								}
-							}
-							break;
-
-							case ELiteralType::UObjectProxy:
-							case ELiteralType::UObjectProxyArray:
-							{
-								PinCategory = FGraphBuilder::PinCategoryObject;
-							}
-							break;
-
-							case ELiteralType::None:
-							case ELiteralType::NoneArray:
-							case ELiteralType::Invalid:
-							default:
-							{
-								static_assert(static_cast<int32>(ELiteralType::Invalid) == 12, "Possible missing binding of pin category to primitive type");
-							}
-							break;
-						}
-					}
-
-					RegisterPinType(DataTypeName, PinCategory, PinSubCategory);
-				}
-			}
-		}
-
 		void RegisterFolderBulkOperations()
 		{
 			if (UToolMenu* Menu = UToolMenus::Get()->ExtendMenu("ContentBrowser.FolderContextMenu"))
@@ -822,44 +717,17 @@ namespace Metasound::Editor
 
 		virtual void RegisterPinType(FName InDataTypeName, FName InPinCategory, FName InPinSubCategory, const FSlateBrush* InPinConnectedIcon = nullptr, const FSlateBrush* InPinDisconnectedIcon = nullptr) override
 		{
-			using namespace Frontend;
-
-			FDataTypeRegistryInfo DataTypeInfo;
-			IDataTypeRegistry::Get().GetDataTypeInfo(InDataTypeName, DataTypeInfo);
-
-			// Default to object as most calls to this outside of the MetaSound Editor will be for custom UObject types
-			const FName PinCategory = InPinCategory.IsNone() ? FGraphBuilder::PinCategoryObject : InPinCategory;
-
-			const EPinContainerType ContainerType = DataTypeInfo.bIsArrayType ? EPinContainerType::Array : EPinContainerType::None;
-			FGraphPinConfiguration PinConfiguration;
-			PinConfiguration.PinType.PinCategory = PinCategory;
-			PinConfiguration.PinType.PinSubCategory = InPinSubCategory;
-			PinConfiguration.PinType.ContainerType = ContainerType;
-			UClass* ClassToUse = IDataTypeRegistry::Get().GetUClassForDataType(InDataTypeName);
-			PinConfiguration.PinType.PinSubCategoryObject = Cast<UObject>(ClassToUse);
-			PinConfiguration.PinConnectedIcon = InPinConnectedIcon;
-			PinConfiguration.PinDisconnectedIcon = InPinDisconnectedIcon;
-			PinTypes.Emplace(InDataTypeName, MoveTemp(PinConfiguration));
+			FGraphPinParams Params;
+			Params.PinCategory = InPinCategory;
+			Params.PinSubCategory = InPinSubCategory;
+			Params.PinConnectedIcon = InPinConnectedIcon;
+			Params.PinDisconnectedIcon = InPinDisconnectedIcon;
+			GraphPanelPinFactory->RegisterPin(InDataTypeName, Params);
 		}
 
 		virtual void RegisterCustomPinType(FName InDataTypeName, const FGraphPinParams& Params) override
 		{
-			RegisterPinType(InDataTypeName, Params.PinCategory, Params.PinSubcategory, Params.PinConnectedIcon, Params.PinDisconnectedIcon);
-			if (Params.PinCategory.IsNone())
-			{
-				return;
-			}
-				
-			if (FGraphBuilder::IsPinCategoryMetaSoundCustomDataType(InDataTypeName))
-			{
-				UE_LOGF(LogMetasoundEditor, Warning, "Attempted to register a \"Custom Pin Type\": \"%ls\", but this is already a Metasound Custom Data Type", *InDataTypeName.ToString());
-				return;
-			}
-				
-			CustomPinCategories.Add(Params.PinCategory);
-			UMetasoundEditorSettings* Settings = GetMutableDefault<UMetasoundEditorSettings>();
-			Settings->CustomPinTypeColors.Add(Params.PinCategory, Params.PinColor ? *Params.PinColor : Settings->DefaultPinTypeColor);
-
+			GraphPanelPinFactory->RegisterPin(InDataTypeName, Params);
 		}
 
 		virtual void RegisterCustomNodeConfigurationDetailsCustomization(FName InNodeConfigurationStructType, FCreateNodeConfigurationDetails InCreateDetailsFunc) override
@@ -1000,37 +868,22 @@ namespace Metasound::Editor
 			
 		virtual bool GetCustomPinIcons(UEdGraphPin* InPin, const FSlateBrush*& PinConnectedIcon, const FSlateBrush*& PinDisconnectedIcon) const override
 		{
-			if (const UEdGraphNode* Node = InPin->GetOwningNode())
-			{
-				if (const UMetasoundEditorGraphNode* MetaSoundNode = Cast<UMetasoundEditorGraphNode>(InPin->GetOwningNode()))
-				{
-					Metasound::Frontend::FDataTypeRegistryInfo RegistryInfo = MetaSoundNode->GetPinDataTypeInfo(*InPin);
-					return GetCustomPinIcons(RegistryInfo.DataTypeName, PinConnectedIcon, PinDisconnectedIcon);
-				}
-			}
-			return false;
+			return GraphPanelPinFactory->GetCustomPinIcons(InPin, PinConnectedIcon, PinDisconnectedIcon);
 		}
 
 		virtual bool GetCustomPinIcons(FName InDataType, const FSlateBrush*& PinConnectedIcon, const FSlateBrush*& PinDisconnectedIcon) const override
 		{
-			const FGraphPinConfiguration* PinConfiguration = PinTypes.Find(InDataType);
-			if (!PinConfiguration || (!PinConfiguration->PinConnectedIcon && !PinConfiguration->PinDisconnectedIcon))
-			{
-				return false;
-			}
-			PinConnectedIcon = PinConfiguration->PinConnectedIcon;
-			PinDisconnectedIcon = PinConfiguration->PinDisconnectedIcon ? PinConfiguration->PinDisconnectedIcon : PinConfiguration->PinConnectedIcon;
-			return true;
+			return GraphPanelPinFactory->GetCustomPinIcons(InDataType, PinConnectedIcon, PinDisconnectedIcon);
 		}
 
-		virtual const FEdGraphPinType* FindPinType(FName InDataTypeName) const
+		virtual const FEdGraphPinType* FindPinType(FName InDataTypeName) const override
 		{
-			const FGraphPinConfiguration* PinConfiguration = PinTypes.Find(InDataTypeName);
-			if (PinConfiguration)
-			{
-				return &PinConfiguration->PinType;
-			}
-			return nullptr;
+			return GraphPanelPinFactory->FindPinType(InDataTypeName);
+		}
+
+		virtual TSharedRef<IMetaSoundGraphPanelPinFactory> GetGraphPanelPinFactory() const override
+		{
+			return StaticCastSharedRef<IMetaSoundGraphPanelPinFactory>(GraphPanelPinFactory.ToSharedRef());
 		}
 
 		virtual bool IsMetaSoundAssetClass(const FTopLevelAssetPath& InClassName) const override
@@ -1110,7 +963,6 @@ namespace Metasound::Editor
 
 			StyleSet = MakeShared<FSlateStyle>();
 
-			RegisterCorePinTypes();
 			RegisterInputDefaultClasses();
 
 			GraphConnectionFactory = MakeShared<FGraphConnectionDrawingPolicyFactory>();
@@ -1119,7 +971,7 @@ namespace Metasound::Editor
 			GraphNodeFactory = MakeShared<FMetasoundGraphNodeFactory>();
 			FEdGraphUtilities::RegisterVisualNodeFactory(GraphNodeFactory);
 
-			GraphPanelPinFactory = MakeShared<FGraphPanelPinFactory>();
+			GraphPanelPinFactory = MakeShared<FMetaSoundGraphPanelPinFactory>();
 			FEdGraphUtilities::RegisterVisualPinFactory(GraphPanelPinFactory);
 
 			RegisterGraphNodeVisualization(
@@ -1241,6 +1093,8 @@ namespace Metasound::Editor
 				.SetDisplayName(LOCTEXT("MetaSoundInputBrowserDisplayName", "MetaSound Input Browser"))
 				.SetTooltipText(LOCTEXT("MetaSoundInputBrowserTooltip", "Bulk view and edit MetaSound input values across multiple assets"))
 				.SetIcon(FSlateIcon(FAppStyle::GetAppStyleSetName(), "Icons.Info"));
+
+			PropertyModule.NotifyCustomizationModuleChanged();
 		}
 
 		virtual void ShutdownModule() override
@@ -1279,8 +1133,6 @@ namespace Metasound::Editor
 				GraphPanelPinFactory.Reset();
 			}
 
-			PinTypes.Reset();
-
 			LiteralCustomizationFactories.Reset();
 
 			FGraphNodeVisualizationRegistry::TearDown();
@@ -1333,15 +1185,13 @@ namespace Metasound::Editor
 		}
 
 		TMap<EMetasoundFrontendLiteralType, const TSubclassOf<UMetasoundEditorGraphMemberDefaultLiteral>> InputDefaultLiteralClassRegistry;
-		TMap<FName, FGraphPinConfiguration> PinTypes;
-		TSet<FName> CustomPinCategories;
 
 		TMap<UClass*, TUniquePtr<IMemberDefaultLiteralCustomizationFactory>> LiteralCustomizationFactories;
 		TMap<FName, FCreateNodeConfigurationDetails> CustomNodeConfigurationDetails;
 
 		TSharedPtr<FMetasoundGraphNodeFactory> GraphNodeFactory;
 		TSharedPtr<FGraphPanelPinConnectionFactory> GraphConnectionFactory;
-		TSharedPtr<FGraphPanelPinFactory> GraphPanelPinFactory;
+		TSharedPtr<FMetaSoundGraphPanelPinFactory> GraphPanelPinFactory;
 		TSharedPtr<FSlateStyleSet> StyleSet;
 
 		TSet<const UClass*> ExplicitProxyClasses;
diff --git a/Engine/Plugins/Runtime/Metasound/Source/MetasoundEditor/Private/MetasoundEditorSettings.h b/Engine/Plugins/Runtime/Metasound/Source/MetasoundEditor/Private/MetasoundEditorSettings.h
--- a/Engine/Plugins/Runtime/Metasound/Source/MetasoundEditor/Private/MetasoundEditorSettings.h
+++ b/Engine/Plugins/Runtime/Metasound/Source/MetasoundEditor/Private/MetasoundEditorSettings.h
@@ -197,7 +197,7 @@ public:
 #endif
 
 
-	/** Maps Pin Category To Pin Color */
+	UE_DEPRECATED(5.8, "Custom pin colors are now managed by FMetaSoundGraphPanelPinFactory. Use IMetaSoundGraphPanelPinFactory::GetPinColor instead.")
 	TMap<FName, FLinearColor> CustomPinTypeColors;
 	
 	/** Default pin type color */
diff --git a/Engine/Plugins/Runtime/Metasound/Source/MetasoundEditor/Private/MetasoundEditorSubsystem.cpp b/Engine/Plugins/Runtime/Metasound/Source/MetasoundEditor/Private/MetasoundEditorSubsystem.cpp
--- a/Engine/Plugins/Runtime/Metasound/Source/MetasoundEditor/Private/MetasoundEditorSubsystem.cpp
+++ b/Engine/Plugins/Runtime/Metasound/Source/MetasoundEditor/Private/MetasoundEditorSubsystem.cpp
@@ -275,7 +275,7 @@ TSubclassOf<UMetasoundEditorGraphMemberDefaultLiteral> UMetaSoundEditorSubsystem
 
 	// Get literal class
 	FDataTypeRegistryInfo DataTypeInfo;
-	IMetasoundEditorModule& EditorModule = FModuleManager::GetModuleChecked<IMetasoundEditorModule>("MetaSoundEditor");
+	IMetasoundEditorModule& EditorModule = FModuleManager::GetModuleChecked<IMetasoundEditorModule>(IMetasoundEditorModule::ModuleName);
 	IDataTypeRegistry::Get().GetDataTypeInfo(TypeName, DataTypeInfo);
 	const EMetasoundFrontendLiteralType LiteralType = static_cast<EMetasoundFrontendLiteralType>(DataTypeInfo.PreferredLiteralType);
 
@@ -333,7 +333,7 @@ void UMetaSoundEditorSubsystem::InitAsset(UObject& InNewMetaSound, UObject* InPa
 		.SelectedObjects = { InParentMetaSound }
 	});
 
-	Metasound::Editor::IMetasoundEditorModule& MetaSoundEditorModule = FModuleManager::GetModuleChecked<Metasound::Editor::IMetasoundEditorModule>("MetaSoundEditor");
+	Metasound::Editor::IMetasoundEditorModule& MetaSoundEditorModule = FModuleManager::GetModuleChecked<Metasound::Editor::IMetasoundEditorModule>(Metasound::Editor::IMetasoundEditorModule::ModuleName);
 	if (MetaSoundEditorModule.IsRestrictedMode())
 	{
 		if (USoundWave* MetaSoundWave = Cast<USoundWave>(&InNewMetaSound))
diff --git a/Engine/Plugins/Runtime/Metasound/Source/MetasoundEditor/Private/MetasoundFactory.cpp b/Engine/Plugins/Runtime/Metasound/Source/MetasoundEditor/Private/MetasoundFactory.cpp
--- a/Engine/Plugins/Runtime/Metasound/Source/MetasoundEditor/Private/MetasoundFactory.cpp
+++ b/Engine/Plugins/Runtime/Metasound/Source/MetasoundEditor/Private/MetasoundFactory.cpp
@@ -68,7 +68,7 @@ UObject* UMetaSoundSourceFactory::FactoryCreateNew(UClass* InClass, UObject* InP
 	});
 
 
-	const IMetasoundEditorModule& MetaSoundEditorModule = FModuleManager::GetModuleChecked<IMetasoundEditorModule>("MetaSoundEditor");
+	const IMetasoundEditorModule& MetaSoundEditorModule = FModuleManager::GetModuleChecked<IMetasoundEditorModule>(IMetasoundEditorModule::ModuleName);
 	if (MetaSoundEditorModule.IsRestrictedMode())
 	{
 		if (const FMetaSoundFrontendPresetTemplate* PresetTemplate = Template.GetPtr<FMetaSoundFrontendPresetTemplate>())
diff --git a/Engine/Plugins/Runtime/Metasound/Source/MetasoundEditor/Private/MetasoundNodeDetailCustomization.cpp b/Engine/Plugins/Runtime/Metasound/Source/MetasoundEditor/Private/MetasoundNodeDetailCustomization.cpp
--- a/Engine/Plugins/Runtime/Metasound/Source/MetasoundEditor/Private/MetasoundNodeDetailCustomization.cpp
+++ b/Engine/Plugins/Runtime/Metasound/Source/MetasoundEditor/Private/MetasoundNodeDetailCustomization.cpp
@@ -713,7 +713,7 @@ namespace Metasound
 		{
 			TSharedPtr<IPropertyHandle> PropertyHandle = StructPropertyHandle->GetChildHandle(GET_MEMBER_NAME_CHECKED(FMetasoundEditorGraphMemberDefaultObjectRef, Object));
 
-			const IMetasoundEditorModule& EditorModule = FModuleManager::GetModuleChecked<IMetasoundEditorModule>("MetaSoundEditor");
+			const IMetasoundEditorModule& EditorModule = FModuleManager::GetModuleChecked<IMetasoundEditorModule>(IMetasoundEditorModule::ModuleName);
 			auto FilterAsset = [InEditorModule = &EditorModule, &InDataTypeInfo = DataTypeInfo](const FAssetData& InAsset)
 			{
 				if (InDataTypeInfo.ProxyGeneratorClass)
@@ -955,7 +955,7 @@ PRAGMA_ENABLE_DEPRECATION_WARNINGS
 				BaseTypeName = GraphMember->GetDataType();
 			}
 
-			IMetasoundEditorModule& EditorModule = FModuleManager::GetModuleChecked<IMetasoundEditorModule>("MetaSoundEditor");
+			IMetasoundEditorModule& EditorModule = FModuleManager::GetModuleChecked<IMetasoundEditorModule>(IMetasoundEditorModule::ModuleName);
 
 			// Not all types have an equivalent array type. Base types without array
 			// types should have the "Is Array" checkbox disabled.
@@ -1121,7 +1121,7 @@ PRAGMA_ENABLE_DEPRECATION_WARNINGS
 						// and hide the struct picker since users should only be able edit child values, not the struct type 
 						if (ConfigurationHandle && ConfigurationHandle->IsValidHandle())
 						{
-							const IMetasoundEditorModule& MetaSoundEditorModule = FModuleManager::GetModuleChecked<IMetasoundEditorModule>("MetaSoundEditor");
+							const IMetasoundEditorModule& MetaSoundEditorModule = FModuleManager::GetModuleChecked<IMetasoundEditorModule>(IMetasoundEditorModule::ModuleName);
 
 							// Get underlying struct name
 							FName StructName;
@@ -1198,7 +1198,7 @@ PRAGMA_ENABLE_DEPRECATION_WARNINGS
 				check(MemberClass);
 
 				IDetailCategoryBuilder& DefaultCategoryBuilder = GetDefaultCategoryBuilder(InDetailLayout);
-				IMetasoundEditorModule& EditorModule = FModuleManager::GetModuleChecked<IMetasoundEditorModule>("MetaSoundEditor");
+				IMetasoundEditorModule& EditorModule = FModuleManager::GetModuleChecked<IMetasoundEditorModule>(IMetasoundEditorModule::ModuleName);
 				LiteralCustomization = EditorModule.CreateMemberDefaultLiteralCustomization(*MemberClass, DefaultCategoryBuilder);
 
 				TAttribute<EVisibility> Visibility = TAttribute<EVisibility>::CreateSPLambda(AsShared(), [this]()
diff --git a/Engine/Plugins/Runtime/Metasound/Source/MetasoundEditor/Private/MetasoundPinAudioInspector.cpp b/Engine/Plugins/Runtime/Metasound/Source/MetasoundEditor/Private/MetasoundPinAudioInspector.cpp
--- a/Engine/Plugins/Runtime/Metasound/Source/MetasoundEditor/Private/MetasoundPinAudioInspector.cpp
+++ b/Engine/Plugins/Runtime/Metasound/Source/MetasoundEditor/Private/MetasoundPinAudioInspector.cpp
@@ -10,6 +10,7 @@
 #include "MetasoundEditorGraphBuilder.h"
 #include "MetasoundEditorGraphNode.h"
 #include "MetasoundEditorModule.h"
+#include "MetaSoundGraphPanelPinFactory.h"
 
 #define LOCTEXT_NAMESPACE "MetasoundEditor"
 
@@ -64,7 +65,7 @@ namespace Metasound::Editor
 		Oscilloscope->StartProcessing();
 
 		// Analyze Audio Pin
-		if (GraphPinObj && !GraphPinObj->LinkedTo.IsEmpty() && GraphPinObj->PinType.PinCategory == FGraphBuilder::PinCategoryAudio)
+		if (GraphPinObj && !GraphPinObj->LinkedTo.IsEmpty() && GraphPinObj->PinType.PinCategory == FMetaSoundGraphPanelPinFactory::PinCategoryAudio)
 		{
 			if (FGraphConnectionManager* ConnectionManager = GetConnectionManager())
 			{
diff --git a/Engine/Plugins/Runtime/Metasound/Source/MetasoundEditor/Private/SMetasoundActionMenu.cpp b/Engine/Plugins/Runtime/Metasound/Source/MetasoundEditor/Private/SMetasoundActionMenu.cpp
--- a/Engine/Plugins/Runtime/Metasound/Source/MetasoundEditor/Private/SMetasoundActionMenu.cpp
+++ b/Engine/Plugins/Runtime/Metasound/Source/MetasoundEditor/Private/SMetasoundActionMenu.cpp
@@ -10,6 +10,7 @@
 #include "MetasoundEditorGraphSchema.h"
 #include "MetasoundEditorModule.h"
 #include "MetasoundEditorSettings.h"
+#include "MetaSoundGraphPanelPinFactory.h"
 #include "MetasoundFrontendDocument.h"
 #include "SGraphActionMenu.h"
 #include "SMetasoundPalette.h"
@@ -195,13 +196,13 @@ namespace Metasound
 						IDataTypeRegistry::Get().GetDataTypeInfo(OutputHandle->GetDataType(), RegistryInfo);
 					}
 
-					TypeColor = FGraphBuilder::GetPinCategoryColor(Pin->PinType);
+					TypeColor = FMetaSoundGraphPanelPinFactory::GetChecked()->GetPinColor(Pin->PinType);
 
 					if (RegistryInfo.bIsArrayType)
 					{
 						PinBrush = FAppStyle::GetBrush("Graph.ArrayPin.Connected");
 					}
-					else if (Pin->PinType.PinCategory == FGraphBuilder::PinCategoryTrigger)
+					else if (Pin->PinType.PinCategory == FMetaSoundGraphPanelPinFactory::PinCategoryTrigger)
 					{
 						if (const ISlateStyle* MetasoundStyle = FSlateStyleRegistry::FindSlateStyle("MetaSoundStyle"))
 						{
diff --git a/Engine/Plugins/Runtime/Metasound/Source/MetasoundEditor/Private/SMetasoundFilterFrequencyResponsePlots.cpp b/Engine/Plugins/Runtime/Metasound/Source/MetasoundEditor/Private/SMetasoundFilterFrequencyResponsePlots.cpp
--- a/Engine/Plugins/Runtime/Metasound/Source/MetasoundEditor/Private/SMetasoundFilterFrequencyResponsePlots.cpp
+++ b/Engine/Plugins/Runtime/Metasound/Source/MetasoundEditor/Private/SMetasoundFilterFrequencyResponsePlots.cpp
@@ -6,6 +6,7 @@
 #include "MetasoundEditorGraphBuilder.h"
 #include "MetasoundEditorGraphNode.h"
 #include "MetasoundEditorModule.h"
+#include "MetaSoundGraphPanelPinFactory.h"
 #include "Widgets/Layout/SBox.h"
 
 #define LOCTEXT_NAMESPACE "SMetasoundFilterFrequencyResponsePlots"
@@ -17,7 +18,7 @@ namespace Metasound::Editor
 		inline bool IsConnectedAudioOutputPin(const UEdGraphPin* Pin)
 		{
 			const bool bIsPinConnected = !Pin->LinkedTo.IsEmpty();
-			const bool bIsAudioOutputPin = Pin->Direction == EGPD_Output && Pin->PinType.PinCategory == FGraphBuilder::PinCategoryAudio;
+			const bool bIsAudioOutputPin = Pin->Direction == EGPD_Output && Pin->PinType.PinCategory == FMetaSoundGraphPanelPinFactory::PinCategoryAudio;
 			return (bIsAudioOutputPin && bIsPinConnected);
 		}
 	} // namespace SMetasoundFilterFrequencyResponsePlotsPrivate
@@ -453,7 +454,7 @@ namespace Metasound::Editor
 			const bool bHasConnectedAudioOutputPins = GraphNode->Pins.ContainsByPredicate(IsConnectedAudioOutputPin);
 			for (UEdGraphPin* Pin : GraphNode->Pins)
 			{
-				const bool bIsAudioOutputPin = Pin->Direction == EGPD_Output && Pin->PinType.PinCategory == FGraphBuilder::PinCategoryAudio;
+				const bool bIsAudioOutputPin = Pin->Direction == EGPD_Output && Pin->PinType.PinCategory == FMetaSoundGraphPanelPinFactory::PinCategoryAudio;
 				if (bIsAudioOutputPin)
 				{
 					const bool bIsPinConnected = !Pin->LinkedTo.IsEmpty();
diff --git a/Engine/Plugins/Runtime/Metasound/Source/MetasoundEditor/Private/SMetasoundGraphNode.cpp b/Engine/Plugins/Runtime/Metasound/Source/MetasoundEditor/Private/SMetasoundGraphNode.cpp
--- a/Engine/Plugins/Runtime/Metasound/Source/MetasoundEditor/Private/SMetasoundGraphNode.cpp
+++ b/Engine/Plugins/Runtime/Metasound/Source/MetasoundEditor/Private/SMetasoundGraphNode.cpp
@@ -21,6 +21,7 @@
 #include "MetasoundDocumentBuilderRegistry.h"
 #include "MetasoundEditorGraph.h"
 #include "MetasoundEditorGraphBuilder.h"
+#include "MetaSoundGraphPanelPinFactory.h"
 #include "MetasoundEditorGraphInputNode.h"
 #include "MetasoundEditorGraphMemberDefaults.h"
 #include "MetasoundEditorGraphNode.h"
@@ -288,8 +289,6 @@ namespace Metasound
 
 		TSharedPtr<SGraphPin> SMetaSoundGraphNode::CreatePinWidget(UEdGraphPin* InPin) const
 		{
-			using namespace Frontend;
-
 			TSharedPtr<SGraphPin> PinWidget;
 
 			if (const UMetasoundEditorGraphSchema* GraphSchema = Cast<const UMetasoundEditorGraphSchema>(InPin->GetSchema()))
@@ -299,52 +298,9 @@ namespace Metasound
 				{
 					PinWidget = SNew(SMetasoundGraphPin, InPin);
 				}
-
-				else if (InPin->PinType.PinCategory == FGraphBuilder::PinCategoryAudio)
-				{
-					PinWidget = SNew(SMetasoundGraphPin, InPin);
-				}
-
-				else if (InPin->PinType.PinCategory == FGraphBuilder::PinCategoryBoolean)
-				{
-					PinWidget = SNew(SMetasoundGraphPinBool, InPin);
-				}
-				
-				else if (InPin->PinType.PinCategory == FGraphBuilder::PinCategoryFloat
-					|| InPin->PinType.PinCategory == FGraphBuilder::PinCategoryTime)
-				{
-					PinWidget = SNew(SMetasoundGraphPinFloat, InPin);
-				}
-
-				else if (InPin->PinType.PinCategory == FGraphBuilder::PinCategoryInt32)
-				{
-					if (SMetasoundGraphEnumPin::FindEnumInterfaceFromPin(InPin))
-					{
-						PinWidget = SNew(SMetasoundGraphEnumPin, InPin);
-					}
-					else
-					{
-						PinWidget = SNew(SMetasoundGraphPinInteger, InPin);
-					}
-				}
-
-				else if (InPin->PinType.PinCategory == FGraphBuilder::PinCategoryObject)
-				{
-					PinWidget = SNew(SMetasoundGraphPinObject, InPin);
-				}
-
-				else if (InPin->PinType.PinCategory == FGraphBuilder::PinCategoryString)
+				else
 				{
-					PinWidget = SNew(SMetasoundGraphPinString, InPin);
-				}
-
-				else if (InPin->PinType.PinCategory == FGraphBuilder::PinCategoryTrigger)
-				{
-					PinWidget = SNew(SMetasoundGraphPin, InPin);
-
-					const FSlateBrush& PinConnectedBrush = Editor::Style::GetSlateBrushSafe("MetasoundEditor.Graph.TriggerPin.Connected");
-					const FSlateBrush& PinDisconnectedBrush = Editor::Style::GetSlateBrushSafe("MetasoundEditor.Graph.TriggerPin.Disconnected");
-					PinWidget->SetCustomPinIcon(&PinConnectedBrush, &PinDisconnectedBrush);
+					PinWidget = FMetaSoundGraphPanelPinFactory::GetChecked()->CreatePin(InPin);
 				}
 			}
 
@@ -353,10 +309,9 @@ namespace Metasound
 				PinWidget = SNew(SMetasoundGraphPin, InPin);
 			}
 
-			IMetasoundEditorModule& EditorModule = FModuleManager::GetModuleChecked<IMetasoundEditorModule>("MetaSoundEditor");
 			const FSlateBrush* PinConnectedIcon = nullptr;
 			const FSlateBrush* PinDisconnectedIcon = nullptr;
-			if (EditorModule.GetCustomPinIcons(InPin, PinConnectedIcon, PinDisconnectedIcon))
+			if (FMetaSoundGraphPanelPinFactory::GetChecked()->GetCustomPinIcons(InPin, PinConnectedIcon, PinDisconnectedIcon))
 			{
 				PinWidget->SetCustomPinIcon(PinConnectedIcon, PinDisconnectedIcon);
 			}
diff --git a/Engine/Plugins/Runtime/Metasound/Source/MetasoundEditor/Private/SMetasoundInputBrowser.cpp b/Engine/Plugins/Runtime/Metasound/Source/MetasoundEditor/Private/SMetasoundInputBrowser.cpp
--- a/Engine/Plugins/Runtime/Metasound/Source/MetasoundEditor/Private/SMetasoundInputBrowser.cpp
+++ b/Engine/Plugins/Runtime/Metasound/Source/MetasoundEditor/Private/SMetasoundInputBrowser.cpp
@@ -12,9 +12,11 @@
 #include "MetasoundDocumentInterface.h"
 #include "MetasoundEditorGraphBuilder.h"
 #include "MetasoundEditorModule.h"
+#include "MetaSoundGraphPanelPinFactory.h"
 #include "MetasoundFrontendDocumentBuilder.h"
 #include "MetasoundFrontendDocumentBuilderRegistry.h"
 #include "MetasoundFrontendDocumentModifyDelegates.h"
+#include "MetaSoundGraphPanelPinFactory.h"
 #include "Misc/MessageDialog.h"
 #include "ScopedTransaction.h"
 #include "Styling/AppStyle.h"
@@ -47,17 +49,17 @@ namespace Metasound::Editor
 
 		FLinearColor GetColorForTypeName(FName TypeName)
 		{
-			const IMetasoundEditorModule& EditorModule = FModuleManager::GetModuleChecked<IMetasoundEditorModule>("MetaSoundEditor");
-			if (const FEdGraphPinType* PinType = EditorModule.FindPinType(TypeName))
+			TSharedRef<FMetaSoundGraphPanelPinFactory> PinFactory = FMetaSoundGraphPanelPinFactory::GetChecked();
+			if (const FEdGraphPinType* PinType = PinFactory->FindPinType(TypeName))
 			{
-				return FGraphBuilder::GetPinCategoryColor(*PinType);
+				return PinFactory->GetPinColor(*PinType);
 			}
-			return FGraphBuilder::GetPinCategoryColor(FEdGraphPinType());
+			return PinFactory->GetPinColor(FEdGraphPinType());
 		}
 
 		const FSlateBrush* GetIconForTypeName(FName TypeName)
 		{
-			const IMetasoundEditorModule& EditorModule = FModuleManager::GetModuleChecked<IMetasoundEditorModule>("MetaSoundEditor");
+			const IMetasoundEditorModule& EditorModule = FModuleManager::GetModuleChecked<IMetasoundEditorModule>(IMetasoundEditorModule::ModuleName);
 			constexpr bool bIsConstructorType = false;
 			return EditorModule.GetIconBrush(TypeName, bIsConstructorType);
 		}
diff --git a/Engine/Plugins/Runtime/Metasound/Source/MetasoundEditor/Private/SMetasoundPinValueInspector.cpp b/Engine/Plugins/Runtime/Metasound/Source/MetasoundEditor/Private/SMetasoundPinValueInspector.cpp
--- a/Engine/Plugins/Runtime/Metasound/Source/MetasoundEditor/Private/SMetasoundPinValueInspector.cpp
+++ b/Engine/Plugins/Runtime/Metasound/Source/MetasoundEditor/Private/SMetasoundPinValueInspector.cpp
@@ -9,6 +9,7 @@
 #include "MetasoundEditorGraphInputNode.h"
 #include "MetasoundEditorGraphNode.h"
 #include "MetasoundEditorSettings.h"
+#include "MetaSoundGraphPanelPinFactory.h"
 #include "MetasoundFrontendDocument.h"
 #include "SPinValueInspector.h"
 #include "Widgets/Colors/SColorBlock.h"
@@ -674,7 +675,7 @@ namespace Metasound
 			TSharedPtr<FMetasoundNumericDebugLineItem> NewLine;
 
 			static const FString UnsetValueString = TEXT("Unset");
-			if (GraphPinObj->PinType.PinCategory == FGraphBuilder::PinCategoryFloat)
+			if (GraphPinObj->PinType.PinCategory == FMetaSoundGraphPanelPinFactory::PinCategoryFloat)
 			{
 				NewLine = MakeShared<FMetasoundNumericDebugLineItem>(GraphPinObj, [&](const FGraphConnectionManager& ConnectionManager, const FGuid& NodeID, const FName OutputName)
 				{
@@ -686,7 +687,7 @@ namespace Metasound
 					return UnsetValueString;
 				});
 			}
-			else if (GraphPinObj->PinType.PinCategory == FGraphBuilder::PinCategoryInt32)
+			else if (GraphPinObj->PinType.PinCategory == FMetaSoundGraphPanelPinFactory::PinCategoryInt32)
 			{
 				NewLine = MakeShared<FMetasoundNumericDebugLineItem>(GraphPinObj, [&](const FGraphConnectionManager& ConnectionManager, const FGuid& NodeID, const FName OutputName)
 				{
@@ -698,7 +699,7 @@ namespace Metasound
 					return UnsetValueString;
 				});
 			}
-			else if (GraphPinObj->PinType.PinCategory == FGraphBuilder::PinCategoryString)
+			else if (GraphPinObj->PinType.PinCategory == FMetaSoundGraphPanelPinFactory::PinCategoryString)
 			{
 				NewLine = MakeShared<FMetasoundNumericDebugLineItem>(GraphPinObj, [&](const FGraphConnectionManager& ConnectionManager, const FGuid& NodeID, const FName OutputName)
 				{
@@ -707,7 +708,7 @@ namespace Metasound
 					return Value;
 				});
 			}
-			else if (GraphPinObj->PinType.PinCategory == FGraphBuilder::PinCategoryBoolean)
+			else if (GraphPinObj->PinType.PinCategory == FMetaSoundGraphPanelPinFactory::PinCategoryBoolean)
 			{
 				NewLine = MakeShared<FMetasoundNumericDebugLineItem>(GraphPinObj, [&](const FGraphConnectionManager& ConnectionManager, const FGuid& NodeID, const FName OutputName)
 				{
diff --git a/Engine/Plugins/Runtime/Metasound/Source/MetasoundEditor/Private/SMetasoundSpectrumAnalyzerGraphNode.cpp b/Engine/Plugins/Runtime/Metasound/Source/MetasoundEditor/Private/SMetasoundSpectrumAnalyzerGraphNode.cpp
--- a/Engine/Plugins/Runtime/Metasound/Source/MetasoundEditor/Private/SMetasoundSpectrumAnalyzerGraphNode.cpp
+++ b/Engine/Plugins/Runtime/Metasound/Source/MetasoundEditor/Private/SMetasoundSpectrumAnalyzerGraphNode.cpp
@@ -6,6 +6,7 @@
 #include "MetasoundEditorGraphBuilder.h"
 #include "MetasoundEditorGraphNode.h"
 #include "MetasoundEditorModule.h"
+#include "MetaSoundGraphPanelPinFactory.h"
 #include "Widgets/Layout/SBox.h"
 
 namespace Metasound::Editor
@@ -31,7 +32,7 @@ namespace Metasound::Editor
 		if (SpectrumAnalyzer.IsValid() && ensure(MetaSoundNode.Pins.Num() == 1))
 		{
 			const UEdGraphPin* Pin = MetaSoundNode.Pins.Last();
-			if (!Pin->LinkedTo.IsEmpty() && ensure(Pin->PinType.PinCategory == FGraphBuilder::PinCategoryAudio && Pin->Direction == EGPD_Input))
+			if (!Pin->LinkedTo.IsEmpty() && ensure(Pin->PinType.PinCategory == FMetaSoundGraphPanelPinFactory::PinCategoryAudio && Pin->Direction == EGPD_Input))
 			{
 				// Find connected output for the input (only ever one):
 				UEdGraphPin* SourcePin = Pin->LinkedTo.Last();
diff --git a/Engine/Plugins/Runtime/Metasound/Source/MetasoundEditor/Public/IMetaSoundGraphPanelPinFactory.h b/Engine/Plugins/Runtime/Metasound/Source/MetasoundEditor/Public/IMetaSoundGraphPanelPinFactory.h
new file mode 100644
--- /dev/null
+++ b/Engine/Plugins/Runtime/Metasound/Source/MetasoundEditor/Public/IMetaSoundGraphPanelPinFactory.h
@@ -0,0 +1,85 @@
+// Copyright Epic Games, Inc. All Rights Reserved.
+#pragma once
+
+#include "EdGraphUtilities.h"
+#include "Math/Color.h"
+#include "Misc/CoreDelegates.h"
+#include "Styling/SlateBrush.h"
+
+#define UE_API METASOUNDEDITOR_API
+
+
+namespace Metasound::Editor
+{
+	// Delegate for creating a custom MetaSound pin widget
+	DECLARE_DELEGATE_RetVal_OneParam(TSharedPtr<SGraphPin>, FOnCreateMetaSoundPinWidget, UEdGraphPin*);
+
+	/**
+	 * Params for registering an external pin type.
+	 *
+	 * Used with IMetaSoundGraphPanelPinFactory::RegisterPin to provide optional
+	 * styling data (category, color, and icons) for a registered data type.
+	 */
+	struct FGraphPinParams
+	{
+		FName PinCategory = { };
+		FName PinSubCategory = { };
+		const FLinearColor* PinColor = nullptr;
+		const FSlateBrush* PinConnectedIcon = nullptr;
+		const FSlateBrush* PinDisconnectedIcon = nullptr;
+	};
+
+	/**
+	 * Factory for creating MetaSound Editor Graph pin icons and widgets.
+	 *
+	 * Provides a single registration API for external plugins to register custom pin types
+	 * with optional styling (category, color, icons). Internal core MetaSound pin types
+	 * (Boolean, Float, Trigger, etc.) are registered automatically at startup.
+	 *
+	 * Lifecycle:
+	 * - Register pins in your module's StartupModule() after MetaSoundEditor is loaded.
+	 * - Unregister pins in your module's ShutdownModule(). Guard with
+	 *   FModuleManager::Get().IsModuleLoaded(IMetasoundEditorModule::ModuleName)
+	 *   in case MetaSoundEditor shuts down before your module.
+	 * - All Register/Unregister calls must be made on the game thread.
+	 */
+	class IMetaSoundGraphPanelPinFactory : public FGraphPanelPinFactory
+	{
+	public:
+		// Register a pin widget factory for a specific pin category. Categories are broader
+		// classifications like Boolean, Float, Int32, etc. Category delegates are used as a
+		// fallback when no data-type-specific delegate is found.
+		UE_API virtual void RegisterCategoryPin(FName InPinCategory, FOnCreateMetaSoundPinWidget InCreateDelegate) = 0;
+
+		// Unregister a pin widget factory for a pin category.
+		UE_API virtual void UnregisterCategoryPin(FName InPinCategory) = 0;
+
+		// Register a pin widget factory for a specific data type. Data type registrations
+		// take priority over category registrations in the widget creation lookup.
+		UE_API virtual void RegisterDataTypePin(FName InDataTypeName, FOnCreateMetaSoundPinWidget InCreateDelegate) = 0;
+
+		// Unregister a pin widget factory for a data type.
+		UE_API virtual void UnregisterDataTypePin(FName InDataTypeName) = 0;
+
+		// Register a pin type for a MetaSound data type with optional styling parameters.
+		UE_API virtual void RegisterPin(FName InDataTypeName, const FGraphPinParams& Params = FGraphPinParams { }) = 0;
+
+		// Unregister a previously registered pin type.
+		UE_API virtual void UnregisterPin(FName InDataTypeName) = 0;
+
+		// Find the pin type configuration for a registered data type. Returns nullptr if not registered.
+		UE_API virtual const FEdGraphPinType* FindPinType(FName InDataTypeName) const = 0;
+
+		// Get custom pin icons for a registered data type via pin reference.
+		UE_API virtual bool GetCustomPinIcons(UEdGraphPin* InPin, const FSlateBrush*& OutConnectedIcon, const FSlateBrush*& OutDisconnectedIcon) const = 0;
+
+		// Get custom pin icons for a registered data type by name.
+		UE_API virtual bool GetCustomPinIcons(FName InDataType, const FSlateBrush*& OutConnectedIcon, const FSlateBrush*& OutDisconnectedIcon) const = 0;
+
+		// Get pin color for a pin type. Checks built-in categories (from UMetasoundEditorSettings),
+		// then registered custom colors, then falls back to DefaultPinTypeColor.
+		UE_API virtual FLinearColor GetPinColor(const FEdGraphPinType& InPinType) const = 0;
+	};
+} // namespace Metasound::Editor
+
+#undef UE_API
diff --git a/Engine/Plugins/Runtime/Metasound/Source/MetasoundEditor/Public/MetasoundEditorModule.h b/Engine/Plugins/Runtime/Metasound/Source/MetasoundEditor/Public/MetasoundEditorModule.h
--- a/Engine/Plugins/Runtime/Metasound/Source/MetasoundEditor/Public/MetasoundEditorModule.h
+++ b/Engine/Plugins/Runtime/Metasound/Source/MetasoundEditor/Public/MetasoundEditorModule.h
@@ -5,6 +5,7 @@
 #include "EdGraph/EdGraphNode.h"
 #include "EdGraph/EdGraphPin.h"
 #include "IDetailPropertyRow.h"
+#include "IMetaSoundGraphPanelPinFactory.h"
 #include "MetasoundDefaultLiteralCustomization.h"
 #include "MetasoundFrontendDataTypeRegistry.h"
 #include "MetasoundFrontendLiteral.h"
@@ -31,7 +32,7 @@ class UMetasoundEditorGraph;
 class UMetasoundEditorGraphMemberDefaultLiteral;
 class UMetasoundEditorGraphNode;
 
-namespace Metasound::Engine 
+namespace Metasound::Engine
 {
 	enum class EAssetScanStatus : uint8;
 	enum class ENodeClassRegistryPrimeStatus : uint8;
@@ -40,103 +41,102 @@ namespace Metasound::Engine
 DECLARE_LOG_CATEGORY_EXTERN(LogMetasoundEditor, Log, All);
 
 
-namespace Metasound
+namespace Metasound::Editor
 {
-	namespace Editor
+	namespace Style
 	{
-		namespace Style
-		{
-			METASOUNDEDITOR_API FSlateIcon CreateSlateIcon(FName InName);
-			METASOUNDEDITOR_API const FSlateBrush& GetSlateBrushSafe(FName InName);
+		METASOUNDEDITOR_API FSlateIcon CreateSlateIcon(FName InName);
+		METASOUNDEDITOR_API const FSlateBrush& GetSlateBrushSafe(FName InName);
 
-			const FSlateColor& GetPageExecutingColor();
-		} // namespace Style
+		const FSlateColor& GetPageExecutingColor();
+	} // namespace Style
 
-		class FMetaSoundNodeConfigurationDataDetails;
+	class FMetaSoundNodeConfigurationDataDetails;
 
-		using EAssetScanStatus = Metasound::Engine::EAssetScanStatus;
-		using EAssetPrimeStatus = Metasound::Engine::ENodeClassRegistryPrimeStatus;
-		using FCreateNodeConfigurationDetails = TFunction<TSharedRef<FMetaSoundNodeConfigurationDataDetails>(
-			TSharedPtr<IPropertyHandle> InStructProperty, TWeakObjectPtr<UMetasoundEditorGraphNode> InNode)>;
+	using EAssetScanStatus = Metasound::Engine::EAssetScanStatus;
+	using EAssetPrimeStatus = Metasound::Engine::ENodeClassRegistryPrimeStatus;
+	using FCreateNodeConfigurationDetails = TFunction<TSharedRef<FMetaSoundNodeConfigurationDataDetails>(
+		TSharedPtr<IPropertyHandle> InStructProperty, TWeakObjectPtr<UMetasoundEditorGraphNode> InNode)>;
 
-		struct FGraphPinParams
-		{
-			FName PinCategory = {};
-			FName PinSubcategory = {};
-			const FLinearColor* PinColor = nullptr;
-			const FSlateBrush* PinConnectedIcon = nullptr;
-			const FSlateBrush* PinDisconnectedIcon = nullptr;
-		};
+	struct FCreateGraphNodeVisualizationWidgetParams
+	{
+		UMetasoundEditorGraphNode* MetaSoundNode = nullptr;
+	};
+
+	DECLARE_DELEGATE_RetVal_OneParam(TSharedRef<SWidget>, FOnCreateGraphNodeVisualizationWidget, const FCreateGraphNodeVisualizationWidgetParams&);
+
+
+	class IMetasoundEditorModule : public IModuleInterface
+	{
+	public:
+		/** Module name constant for use with FModuleManager. */
+		UE_API static const FName ModuleName;
 
+		UE_DEPRECATED(5.3, "IsExplicitProxyClass is deprecated, use Metasound::Frontend::FDataTypeRegistryInfo::bIsExplicit")
+		virtual bool IsExplicitProxyClass(const UClass& InClass) const = 0;
 
-		struct FCreateGraphNodeVisualizationWidgetParams
-		{
-			UMetasoundEditorGraphNode* MetaSoundNode = nullptr;
-		};
+		UE_DEPRECATED(5.3, "RegisterExplicitProxyClass is deprecated, use Metasound::TIsExplicit<>")
+		virtual void RegisterExplicitProxyClass(const UClass& InClass) = 0;
 
-		DECLARE_DELEGATE_RetVal_OneParam(TSharedRef<SWidget>, FOnCreateGraphNodeVisualizationWidget, const FCreateGraphNodeVisualizationWidgetParams&);
+		UE_DEPRECATED(5.3, "IsMetaSoundAssetClass is deprecated, use IMetasoundUObjectRegistry::IsRegisteredClass")
+		virtual bool IsMetaSoundAssetClass(const FTopLevelAssetPath& InClassName) const = 0;
 
+		UE_DEPRECATED(5.5, "Prime no longer supported (nor recommended as AssetTag updates now enable querying MetaSound assets for interface information without loading full asset.")
+		virtual void PrimeAssetRegistryAsync() { }
 
-		class IMetasoundEditorModule : public IModuleInterface
-		{
-		public:
-			UE_DEPRECATED(5.3, "IsExplicitProxyClass is deprecated, use Metasound::Frontend::FDataTypeRegistryInfo::bIsExplicit")
-			virtual bool IsExplicitProxyClass(const UClass& InClass) const = 0;
+		UE_DEPRECATED(5.5, "Prime no longer supported (nor recommended as AssetTag updates now enable querying MetaSound assets for interface information without loading full asset.")
+		UE_API virtual EAssetPrimeStatus GetAssetRegistryPrimeStatus() const;
 
-			UE_DEPRECATED(5.3, "RegisterExplicitProxyClass is deprecated, use Metasound::TIsExplicit<>")
-			virtual void RegisterExplicitProxyClass(const UClass& InClass) = 0;
+		UE_DEPRECATED(5.5, "Use FMetaSoundAssetManager::IsInitialAssetScanComplete instead")
+		UE_API virtual EAssetScanStatus GetAssetRegistryScanStatus() const;
 
-			UE_DEPRECATED(5.3, "IsMetaSoundAssetClass is deprecated, use IMetasoundUObjectRegistry::IsRegisteredClass")
-			virtual bool IsMetaSoundAssetClass(const FTopLevelAssetPath& InClassName) const = 0;
+		virtual TUniquePtr<FMetasoundDefaultLiteralCustomizationBase> CreateMemberDefaultLiteralCustomization(UClass& InClass, IDetailCategoryBuilder& DefaultCategoryBuilder) const = 0;
 
-			UE_DEPRECATED(5.5, "Prime no longer supported (nor recommended as AssetTag updates now enable querying MetaSound assets for interface information without loading full asset.")
-			virtual void PrimeAssetRegistryAsync() { }
+		UE_EXPERIMENTAL(5.6, "Node configuration custom details customizations are experimental.")
+		virtual const FCreateNodeConfigurationDetails* FindCreateCustomNodeConfigurationDetailsCustomization(FName InNodeConfigurationStructType) const = 0;
 
-			UE_DEPRECATED(5.5, "Prime no longer supported (nor recommended as AssetTag updates now enable querying MetaSound assets for interface information without loading full asset.")
-			UE_API virtual EAssetPrimeStatus GetAssetRegistryPrimeStatus() const;
+		virtual const TSubclassOf<UMetasoundEditorGraphMemberDefaultLiteral> FindDefaultLiteralClass(EMetasoundFrontendLiteralType InLiteralType) const = 0;
 
-			UE_DEPRECATED(5.5, "Use FMetaSoundAssetManager::IsInitialAssetScanComplete instead")
-			UE_API virtual EAssetScanStatus GetAssetRegistryScanStatus() const;
+		UE_DEPRECATED(5.8, "Use IMetaSoundGraphPanelPinFactory::FindPinType instead.")
+		virtual const FEdGraphPinType* FindPinType(FName InDataTypeName) const = 0;
 
-			virtual TUniquePtr<FMetasoundDefaultLiteralCustomizationBase> CreateMemberDefaultLiteralCustomization(UClass& InClass, IDetailCategoryBuilder& DefaultCategoryBuilder) const = 0;
+		virtual const FSlateBrush* GetIconBrush(FName InDataType, const bool bIsConstructorType) const = 0;
 
-			UE_EXPERIMENTAL(5.6, "Node configuration custom details customizations are experimental.")
-			virtual const FCreateNodeConfigurationDetails* FindCreateCustomNodeConfigurationDetailsCustomization(FName InNodeConfigurationStructType) const = 0;
+		UE_DEPRECATED(5.8, "Use IMetaSoundGraphPanelPinFactory::GetCustomPinIcons instead.")
+		virtual bool GetCustomPinIcons(UEdGraphPin* InPin, const FSlateBrush*& PinConnectedIcon, const FSlateBrush*& PinDisconnectedIcon) const = 0;
 
-			virtual const TSubclassOf<UMetasoundEditorGraphMemberDefaultLiteral> FindDefaultLiteralClass(EMetasoundFrontendLiteralType InLiteralType) const = 0;
+		UE_DEPRECATED(5.8, "Use IMetaSoundGraphPanelPinFactory::GetCustomPinIcons instead.")
+		virtual bool GetCustomPinIcons(FName InDataType, const FSlateBrush*& PinConnectedIcon, const FSlateBrush*& PinDisconnectedIcon) const = 0;
 
-			virtual const FEdGraphPinType* FindPinType(FName InDataTypeName) const = 0;
+		// Returns the MetaSound GraphPanelPinFactory for registering custom pin types.
+		UE_API virtual TSharedRef<IMetaSoundGraphPanelPinFactory> GetGraphPanelPinFactory() const = 0;
 
-			virtual const FSlateBrush* GetIconBrush(FName InDataType, const bool bIsConstructorType) const = 0;
+		UE_DEPRECATED(5.8, "Use IMetaSoundGraphPanelPinFactory::RegisterPin instead.")
+		virtual void RegisterPinType(FName InDataTypeName, FName InPinCategory = { }, FName InPinSubCategory = { },
+			const FSlateBrush* InPinConnectedIcon = nullptr, const FSlateBrush* InPinDisconnectedIcon = nullptr) = 0;
 
-			virtual bool GetCustomPinIcons(UEdGraphPin* InPin, const FSlateBrush*& PinConnectedIcon, const FSlateBrush*& PinDisconnectedIcon) const = 0;
-			virtual bool GetCustomPinIcons(FName InDataType, const FSlateBrush*& PinConnectedIcon, const FSlateBrush*& PinDisconnectedIcon) const = 0;
-			
-			virtual void RegisterPinType(FName InDataTypeName, FName InPinCategory = { }, FName InPinSubCategory = { },
-				const FSlateBrush* InPinConnectedIcon = nullptr, const FSlateBrush* InPinDisconnectedIcon = nullptr) = 0;
-			
-			virtual void RegisterCustomPinType(FName InDataTypeName, const FGraphPinParams& Params) = 0;
+		UE_DEPRECATED(5.8, "Use IMetaSoundGraphPanelPinFactory::RegisterPin instead.")
+		virtual void RegisterCustomPinType(FName InDataTypeName, const FGraphPinParams& Params) = 0;
 
-			/** Register a custom details customization for a specific substruct of FMetaSoundFrontendNodeConfiguration 
-				InNodeConfigurationStructType - struct type name - use MyStruct::StaticStruct()->GetFName()
-				InCreateDetailsFunc - function that creates the specific details customization */
-			UE_EXPERIMENTAL(5.6, "Node configuration custom details customizations are experimental.")
-			virtual void RegisterCustomNodeConfigurationDetailsCustomization(FName InNodeConfigurationStructType, FCreateNodeConfigurationDetails InCreateDetailsFunc) = 0;
+		/** Register a custom details customization for a specific substruct of FMetaSoundFrontendNodeConfiguration
+			InNodeConfigurationStructType - struct type name - use MyStruct::StaticStruct()->GetFName()
+			InCreateDetailsFunc - function that creates the specific details customization */
+		UE_EXPERIMENTAL(5.6, "Node configuration custom details customizations are experimental.")
+		virtual void RegisterCustomNodeConfigurationDetailsCustomization(FName InNodeConfigurationStructType, FCreateNodeConfigurationDetails InCreateDetailsFunc) = 0;
 
-			// Unregister a custom details customization for a node configuration struct type, using the same struct type name as RegisterCustomNodeConfigurationDetailsCustomization
-			UE_EXPERIMENTAL(5.6, "Node configuration custom details customizations are experimental.")
-			virtual void UnregisterCustomNodeConfigurationDetailsCustomization(FName InNodeConfigurationStructType) = 0;
+		// Unregister a custom details customization for a node configuration struct type, using the same struct type name as RegisterCustomNodeConfigurationDetailsCustomization
+		UE_EXPERIMENTAL(5.6, "Node configuration custom details customizations are experimental.")
+		virtual void UnregisterCustomNodeConfigurationDetailsCustomization(FName InNodeConfigurationStructType) = 0;
 
-			// For the given node class, register a delegate that can be used for creating in-graph node visualizations.
-			virtual void RegisterGraphNodeVisualization(FName InNodeClassName, FOnCreateGraphNodeVisualizationWidget OnCreateGraphNodeVisualizationWidget) = 0;
+		// For the given node class, register a delegate that can be used for creating in-graph node visualizations.
+		virtual void RegisterGraphNodeVisualization(FName InNodeClassName, FOnCreateGraphNodeVisualizationWidget OnCreateGraphNodeVisualizationWidget) = 0;
 
-			// Queries if the MetaSound Editor is in "restricted mode" (i.e. can only make new presets and not make new assets or edit asset graphs)
-			virtual bool IsRestrictedMode() const = 0;
+		// Queries if the MetaSound Editor is in "restricted mode" (i.e. can only make new presets and not make new assets or edit asset graphs)
+		virtual bool IsRestrictedMode() const = 0;
 
-			// Sets if the MetaSound editor is in "restricted mode" (i.e. can only make new presets and not make new assets or edit asset graphs)
-			virtual void SetRestrictedMode(bool bInRestricted) = 0;
-		};
-	} // namespace Editor
-} // namespace Metasound
+		// Sets if the MetaSound editor is in "restricted mode" (i.e. can only make new presets and not make new assets or edit asset graphs)
+		virtual void SetRestrictedMode(bool bInRestricted) = 0;
+	};
+} // namespace Metasound::Editor
 
 #undef UE_API
